- Experimental Expressions for Stateless Emitters
- Allows users to do some limited operations inline in the stateless emitter stack
- This is for testing & feedback, the data may not be backwards compatabile
- Enabled using fx.NiagaraStateless.ExpressionsEnabled

#rnx

diff --git a/Engine/Plugins/FX/Niagara/Source/Niagara/Classes/NiagaraScript.h b/Engine/Plugins/FX/Niagara/Source/Niagara/Classes/NiagaraScript.h
--- a/Engine/Plugins/FX/Niagara/Source/Niagara/Classes/NiagaraScript.h
+++ b/Engine/Plugins/FX/Niagara/Source/Niagara/Classes/NiagaraScript.h
@@ -1233,12 +1233,12 @@ public:
 				}
 				else if (const int32* Offset = RapidIterationParameters.FindParameterOffset(FNiagaraVariableBase(Tag.Variable.GetType(), *Tag.StringValue)))
 				{
-					return TOptional<T>(*(T*)RapidIterationParameters.GetParameterData(*Offset));
+					return TOptional<T>(*(T*)RapidIterationParameters.GetParameterData(*Offset, Tag.Variable.GetType()));
 				}
 				else if (FallbackParameterStore)
 				{
 					if (const int32* OffsetAlternate = FallbackParameterStore->FindParameterOffset(FNiagaraVariableBase(Tag.Variable.GetType(), *Tag.StringValue)))
-						return TOptional<T>(*(T*)FallbackParameterStore->GetParameterData(*OffsetAlternate));
+						return TOptional<T>(*(T*)FallbackParameterStore->GetParameterData(*OffsetAlternate, Tag.Variable.GetType()));
 				}
 			}
 		}
diff --git a/Engine/Plugins/FX/Niagara/Source/Niagara/Internal/Stateless/Expressions/NiagaraStatelessExpressionColor.h b/Engine/Plugins/FX/Niagara/Source/Niagara/Internal/Stateless/Expressions/NiagaraStatelessExpressionColor.h
new file mode 100644
--- /dev/null
+++ b/Engine/Plugins/FX/Niagara/Source/Niagara/Internal/Stateless/Expressions/NiagaraStatelessExpressionColor.h
@@ -0,0 +1,135 @@
+// Copyright Epic Games, Inc. All Rights Reserved.
+
+#pragma once
+
+#include "Stateless/NiagaraStatelessExpression.h"
+
+#include "NiagaraStatelessExpressionColor.generated.h"
+
+////////////////////////////////////////////////////////////////////////////////////////////////////////////////////////
+
+USTRUCT()
+struct FNiagaraStatelessExpressionColor : public FNiagaraStatelessExpression
+{
+	GENERATED_BODY()
+
+	virtual FNiagaraTypeDefinition GetOutputTypeDef() const override final { return FNiagaraTypeDefinition::GetColorDef(); }
+	virtual FInstancedStruct Build(const FNiagaraStatelessEmitterDataBuildContext& BuildContext) const override final;
+	virtual void Evaluate(const FEvaluateContext& Context, void* ValueAddress) const override final { *static_cast<FLinearColor*>(ValueAddress) = EvaluateInternal(Context); }
+
+	virtual FInstancedStruct BuildInternal(const FNiagaraStatelessEmitterDataBuildContext& BuildContext) const { checkNoEntry(); return FInstancedStruct(); }
+	virtual FLinearColor EvaluateInternal(const FEvaluateContext& Context) const { checkNoEntry(); return FLinearColor::Black; }
+};
+
+////////////////////////////////////////////////////////////////////////////////////////////////////////////////////////
+
+USTRUCT(meta = (ValueExpression))
+struct FNiagaraStatelessExpressionColorConstant final : public FNiagaraStatelessExpressionColor
+{
+	GENERATED_BODY()
+
+	virtual ~FNiagaraStatelessExpressionColorConstant() = default;
+	virtual FLinearColor EvaluateInternal(const FEvaluateContext& Context) const override final;
+	virtual bool IsConstant() const override;
+
+	UPROPERTY(EditAnywhere, Category = "Values")
+	FLinearColor A = FLinearColor::Black;
+};
+
+////////////////////////////////////////////////////////////////////////////////////////////////////////////////////////
+
+USTRUCT(meta = (BindingExpression))
+struct FNiagaraStatelessExpressionColorBinding final : public FNiagaraStatelessExpressionColor
+{
+	GENERATED_BODY()
+
+	virtual ~FNiagaraStatelessExpressionColorBinding() = default;
+	virtual bool IsConstant() const { return false; }
+	virtual FInstancedStruct BuildInternal(const FNiagaraStatelessEmitterDataBuildContext& BuildContext) const override;
+	virtual FLinearColor EvaluateInternal(const FEvaluateContext& Context) const override final;
+
+	UPROPERTY(EditAnywhere, Category = "Values")
+	FName A;
+
+	int32 ParameterOffset = INDEX_NONE;
+};
+
+////////////////////////////////////////////////////////////////////////////////////////////////////////////////////////
+
+USTRUCT(meta = (OperationExpression, DisplayName = "Add Color"))
+struct FNiagaraStatelessExpressionColorAdd final : public FNiagaraStatelessExpressionColor
+{
+	GENERATED_BODY()
+
+	FNiagaraStatelessExpressionColorAdd();
+	virtual ~FNiagaraStatelessExpressionColorAdd() = default;
+	virtual FInstancedStruct BuildInternal(const FNiagaraStatelessEmitterDataBuildContext& BuildContext) const override;
+	virtual FLinearColor EvaluateInternal(const FEvaluateContext& Context) const override final;
+	virtual bool IsConstant() const override;
+
+	UPROPERTY(EditAnywhere, Category = "Values", meta = (BaseStruct = "/Script/Niagara.NiagaraStatelessExpressionColor"))
+	FInstancedStruct A;
+
+	UPROPERTY(EditAnywhere, Category = "Values", meta = (BaseStruct = "/Script/Niagara.NiagaraStatelessExpressionColor"))
+	FInstancedStruct B;
+};
+
+////////////////////////////////////////////////////////////////////////////////////////////////////////////////////////
+
+USTRUCT(meta = (OperationExpression, DisplayName = "Subtract Color"))
+struct FNiagaraStatelessExpressionColorSubtract final : public FNiagaraStatelessExpressionColor
+{
+	GENERATED_BODY()
+
+	FNiagaraStatelessExpressionColorSubtract();
+	virtual ~FNiagaraStatelessExpressionColorSubtract() = default;
+	virtual FInstancedStruct BuildInternal(const FNiagaraStatelessEmitterDataBuildContext& BuildContext) const override;
+	virtual FLinearColor EvaluateInternal(const FEvaluateContext& Context) const override final;
+	virtual bool IsConstant() const override;
+
+	UPROPERTY(EditAnywhere, Category = "Values", meta = (BaseStruct = "/Script/Niagara.NiagaraStatelessExpressionColor"))
+	FInstancedStruct A;
+
+	UPROPERTY(EditAnywhere, Category = "Values", meta = (BaseStruct = "/Script/Niagara.NiagaraStatelessExpressionColor"))
+	FInstancedStruct B;
+};
+
+////////////////////////////////////////////////////////////////////////////////////////////////////////////////////////
+
+USTRUCT(meta = (OperationExpression, DisplayName = "Multiply Color"))
+struct FNiagaraStatelessExpressionColorMultiply final : public FNiagaraStatelessExpressionColor
+{
+	GENERATED_BODY()
+
+	FNiagaraStatelessExpressionColorMultiply();
+	virtual ~FNiagaraStatelessExpressionColorMultiply() = default;
+	virtual FInstancedStruct BuildInternal(const FNiagaraStatelessEmitterDataBuildContext& BuildContext) const override;
+	virtual FLinearColor EvaluateInternal(const FEvaluateContext& Context) const override final;
+	virtual bool IsConstant() const override;
+
+	UPROPERTY(EditAnywhere, Category = "Values", meta = (BaseStruct = "/Script/Niagara.NiagaraStatelessExpressionColor"))
+	FInstancedStruct A;
+
+	UPROPERTY(EditAnywhere, Category = "Values", meta = (BaseStruct = "/Script/Niagara.NiagaraStatelessExpressionColor"))
+	FInstancedStruct B;
+};
+
+////////////////////////////////////////////////////////////////////////////////////////////////////////////////////////
+
+USTRUCT(meta = (OperationExpression, DisplayName = "Divide Color"))
+struct FNiagaraStatelessExpressionColorDivide final : public FNiagaraStatelessExpressionColor
+{
+	GENERATED_BODY()
+
+	FNiagaraStatelessExpressionColorDivide();
+	virtual ~FNiagaraStatelessExpressionColorDivide() = default;
+	virtual FInstancedStruct BuildInternal(const FNiagaraStatelessEmitterDataBuildContext& BuildContext) const override;
+	virtual FLinearColor EvaluateInternal(const FEvaluateContext& Context) const override final;
+	virtual bool IsConstant() const override;
+
+	UPROPERTY(EditAnywhere, Category = "Values", meta = (BaseStruct = "/Script/Niagara.NiagaraStatelessExpressionColor"))
+	FInstancedStruct A;
+
+	UPROPERTY(EditAnywhere, Category = "Values", meta = (BaseStruct = "/Script/Niagara.NiagaraStatelessExpressionColor"))
+	FInstancedStruct B;
+};
diff --git a/Engine/Plugins/FX/Niagara/Source/Niagara/Internal/Stateless/Expressions/NiagaraStatelessExpressionFloat.h b/Engine/Plugins/FX/Niagara/Source/Niagara/Internal/Stateless/Expressions/NiagaraStatelessExpressionFloat.h
new file mode 100644
--- /dev/null
+++ b/Engine/Plugins/FX/Niagara/Source/Niagara/Internal/Stateless/Expressions/NiagaraStatelessExpressionFloat.h
@@ -0,0 +1,135 @@
+// Copyright Epic Games, Inc. All Rights Reserved.
+
+#pragma once
+
+#include "Stateless/NiagaraStatelessExpression.h"
+
+#include "NiagaraStatelessExpressionFloat.generated.h"
+
+////////////////////////////////////////////////////////////////////////////////////////////////////////////////////////
+
+USTRUCT()
+struct FNiagaraStatelessExpressionFloat : public FNiagaraStatelessExpression
+{
+	GENERATED_BODY()
+
+	virtual FNiagaraTypeDefinition GetOutputTypeDef() const override final { return FNiagaraTypeDefinition::GetFloatDef(); }
+	virtual FInstancedStruct Build(const FNiagaraStatelessEmitterDataBuildContext& BuildContext) const override final;
+	virtual void Evaluate(const FEvaluateContext& Context, void* ValueAddress) const override final { *static_cast<float*>(ValueAddress) = EvaluateInternal(Context); }
+
+	virtual FInstancedStruct BuildInternal(const FNiagaraStatelessEmitterDataBuildContext& BuildContext) const { checkNoEntry(); return FInstancedStruct(); }
+	virtual float EvaluateInternal(const FEvaluateContext& Context) const { checkNoEntry(); return 0.0f; }
+};
+
+////////////////////////////////////////////////////////////////////////////////////////////////////////////////////////
+
+USTRUCT(meta = (ValueExpression))
+struct FNiagaraStatelessExpressionFloatConstant final : public FNiagaraStatelessExpressionFloat
+{
+	GENERATED_BODY()
+
+	virtual ~FNiagaraStatelessExpressionFloatConstant() = default;
+	virtual float EvaluateInternal(const FEvaluateContext& Context) const override final;
+	virtual bool IsConstant() const override;
+
+	UPROPERTY(EditAnywhere, Category = "Values")
+	float A = 0.0f;
+};
+
+////////////////////////////////////////////////////////////////////////////////////////////////////////////////////////
+
+USTRUCT(meta = (BindingExpression))
+struct FNiagaraStatelessExpressionFloatBinding final : public FNiagaraStatelessExpressionFloat
+{
+	GENERATED_BODY()
+
+	virtual ~FNiagaraStatelessExpressionFloatBinding() = default;
+	virtual bool IsConstant() const { return false; }
+	virtual FInstancedStruct BuildInternal(const FNiagaraStatelessEmitterDataBuildContext& BuildContext) const override;
+	virtual float EvaluateInternal(const FEvaluateContext& Context) const override final;
+
+	UPROPERTY(EditAnywhere, Category = "Values")
+	FName A;
+
+	int32 ParameterOffset = INDEX_NONE;
+};
+
+////////////////////////////////////////////////////////////////////////////////////////////////////////////////////////
+
+USTRUCT(meta = (OperationExpression, DisplayName = "Add Float"))
+struct FNiagaraStatelessExpressionFloatAdd final : public FNiagaraStatelessExpressionFloat
+{
+	GENERATED_BODY()
+
+	FNiagaraStatelessExpressionFloatAdd();
+	virtual ~FNiagaraStatelessExpressionFloatAdd() = default;
+	virtual FInstancedStruct BuildInternal(const FNiagaraStatelessEmitterDataBuildContext& BuildContext) const override;
+	virtual float EvaluateInternal(const FEvaluateContext& Context) const override final;
+	virtual bool IsConstant() const override;
+
+	UPROPERTY(EditAnywhere, Category = "Values", meta = (BaseStruct = "/Script/Niagara.NiagaraStatelessExpressionFloat"))
+	FInstancedStruct A;
+
+	UPROPERTY(EditAnywhere, Category = "Values", meta = (BaseStruct = "/Script/Niagara.NiagaraStatelessExpressionFloat"))
+	FInstancedStruct B;
+};
+
+////////////////////////////////////////////////////////////////////////////////////////////////////////////////////////
+
+USTRUCT(meta = (OperationExpression, DisplayName = "Subtract Float"))
+struct FNiagaraStatelessExpressionFloatSubtract final : public FNiagaraStatelessExpressionFloat
+{
+	GENERATED_BODY()
+
+	FNiagaraStatelessExpressionFloatSubtract();
+	virtual ~FNiagaraStatelessExpressionFloatSubtract() = default;
+	virtual FInstancedStruct BuildInternal(const FNiagaraStatelessEmitterDataBuildContext& BuildContext) const override;
+	virtual float EvaluateInternal(const FEvaluateContext& Context) const override final;
+	virtual bool IsConstant() const override;
+
+	UPROPERTY(EditAnywhere, Category = "Values", meta = (BaseStruct = "/Script/Niagara.NiagaraStatelessExpressionFloat"))
+	FInstancedStruct A;
+
+	UPROPERTY(EditAnywhere, Category = "Values", meta = (BaseStruct = "/Script/Niagara.NiagaraStatelessExpressionFloat"))
+	FInstancedStruct B;
+};
+
+////////////////////////////////////////////////////////////////////////////////////////////////////////////////////////
+
+USTRUCT(meta = (OperationExpression, DisplayName = "Multiply Float"))
+struct FNiagaraStatelessExpressionFloatMultiply final : public FNiagaraStatelessExpressionFloat
+{
+	GENERATED_BODY()
+
+	FNiagaraStatelessExpressionFloatMultiply();
+	virtual ~FNiagaraStatelessExpressionFloatMultiply() = default;
+	virtual FInstancedStruct BuildInternal(const FNiagaraStatelessEmitterDataBuildContext& BuildContext) const override;
+	virtual float EvaluateInternal(const FEvaluateContext& Context) const override final;
+	virtual bool IsConstant() const override;
+
+	UPROPERTY(EditAnywhere, Category = "Values", meta = (BaseStruct = "/Script/Niagara.NiagaraStatelessExpressionFloat"))
+	FInstancedStruct A;
+
+	UPROPERTY(EditAnywhere, Category = "Values", meta = (BaseStruct = "/Script/Niagara.NiagaraStatelessExpressionFloat"))
+	FInstancedStruct B;
+};
+
+////////////////////////////////////////////////////////////////////////////////////////////////////////////////////////
+
+USTRUCT(meta = (OperationExpression, DisplayName = "Divide Float"))
+struct FNiagaraStatelessExpressionFloatDivide final : public FNiagaraStatelessExpressionFloat
+{
+	GENERATED_BODY()
+
+	FNiagaraStatelessExpressionFloatDivide();
+	virtual ~FNiagaraStatelessExpressionFloatDivide() = default;
+	virtual FInstancedStruct BuildInternal(const FNiagaraStatelessEmitterDataBuildContext& BuildContext) const override;
+	virtual float EvaluateInternal(const FEvaluateContext& Context) const override final;
+	virtual bool IsConstant() const override;
+
+	UPROPERTY(EditAnywhere, Category = "Values", meta = (BaseStruct = "/Script/Niagara.NiagaraStatelessExpressionFloat"))
+	FInstancedStruct A;
+
+	UPROPERTY(EditAnywhere, Category = "Values", meta = (BaseStruct = "/Script/Niagara.NiagaraStatelessExpressionFloat"))
+	FInstancedStruct B;
+};
diff --git a/Engine/Plugins/FX/Niagara/Source/Niagara/Internal/Stateless/Expressions/NiagaraStatelessExpressionVec2.h b/Engine/Plugins/FX/Niagara/Source/Niagara/Internal/Stateless/Expressions/NiagaraStatelessExpressionVec2.h
new file mode 100644
--- /dev/null
+++ b/Engine/Plugins/FX/Niagara/Source/Niagara/Internal/Stateless/Expressions/NiagaraStatelessExpressionVec2.h
@@ -0,0 +1,135 @@
+// Copyright Epic Games, Inc. All Rights Reserved.
+
+#pragma once
+
+#include "Stateless/NiagaraStatelessExpression.h"
+
+#include "NiagaraStatelessExpressionVec2.generated.h"
+
+////////////////////////////////////////////////////////////////////////////////////////////////////////////////////////
+
+USTRUCT()
+struct FNiagaraStatelessExpressionVec2 : public FNiagaraStatelessExpression
+{
+	GENERATED_BODY()
+
+	virtual FNiagaraTypeDefinition GetOutputTypeDef() const override final { return FNiagaraTypeDefinition::GetVec2Def(); }
+	virtual FInstancedStruct Build(const FNiagaraStatelessEmitterDataBuildContext& BuildContext) const override final;
+	virtual void Evaluate(const FEvaluateContext& Context, void* ValueAddress) const override final { *static_cast<FVector2f*>(ValueAddress) = EvaluateInternal(Context); }
+
+	virtual FInstancedStruct BuildInternal(const FNiagaraStatelessEmitterDataBuildContext& BuildContext) const { checkNoEntry(); return FInstancedStruct(); }
+	virtual FVector2f EvaluateInternal(const FEvaluateContext& Context) const { checkNoEntry(); return FVector2f::ZeroVector; }
+};
+
+////////////////////////////////////////////////////////////////////////////////////////////////////////////////////////
+
+USTRUCT(meta = (ValueExpression))
+struct FNiagaraStatelessExpressionVec2Constant final : public FNiagaraStatelessExpressionVec2
+{
+	GENERATED_BODY()
+
+	virtual ~FNiagaraStatelessExpressionVec2Constant() = default;
+	virtual FVector2f EvaluateInternal(const FEvaluateContext& Context) const override final;
+	virtual bool IsConstant() const override;
+
+	UPROPERTY(EditAnywhere, Category = "Values")
+	FVector2f A = FVector2f::ZeroVector;
+};
+
+////////////////////////////////////////////////////////////////////////////////////////////////////////////////////////
+
+USTRUCT(meta = (BindingExpression))
+struct FNiagaraStatelessExpressionVec2Binding final : public FNiagaraStatelessExpressionVec2
+{
+	GENERATED_BODY()
+
+	virtual ~FNiagaraStatelessExpressionVec2Binding() = default;
+	virtual bool IsConstant() const { return false; }
+	virtual FInstancedStruct BuildInternal(const FNiagaraStatelessEmitterDataBuildContext& BuildContext) const override;
+	virtual FVector2f EvaluateInternal(const FEvaluateContext& Context) const override final;
+
+	UPROPERTY(EditAnywhere, Category = "Values")
+	FName A;
+
+	int32 ParameterOffset = INDEX_NONE;
+};
+
+////////////////////////////////////////////////////////////////////////////////////////////////////////////////////////
+
+USTRUCT(meta = (OperationExpression, DisplayName = "Add Vec2"))
+struct FNiagaraStatelessExpressionVec2Add final : public FNiagaraStatelessExpressionVec2
+{
+	GENERATED_BODY()
+
+	FNiagaraStatelessExpressionVec2Add();
+	virtual ~FNiagaraStatelessExpressionVec2Add() = default;
+	virtual FInstancedStruct BuildInternal(const FNiagaraStatelessEmitterDataBuildContext& BuildContext) const override;
+	virtual FVector2f EvaluateInternal(const FEvaluateContext& Context) const override final;
+	virtual bool IsConstant() const override;
+
+	UPROPERTY(EditAnywhere, Category = "Values", meta = (BaseStruct = "/Script/Niagara.NiagaraStatelessExpressionVec2"))
+	FInstancedStruct A;
+
+	UPROPERTY(EditAnywhere, Category = "Values", meta = (BaseStruct = "/Script/Niagara.NiagaraStatelessExpressionVec2"))
+	FInstancedStruct B;
+};
+
+////////////////////////////////////////////////////////////////////////////////////////////////////////////////////////
+
+USTRUCT(meta = (OperationExpression, DisplayName = "Subtract Vec2"))
+struct FNiagaraStatelessExpressionVec2Subtract final : public FNiagaraStatelessExpressionVec2
+{
+	GENERATED_BODY()
+
+	FNiagaraStatelessExpressionVec2Subtract();
+	virtual ~FNiagaraStatelessExpressionVec2Subtract() = default;
+	virtual FInstancedStruct BuildInternal(const FNiagaraStatelessEmitterDataBuildContext& BuildContext) const override;
+	virtual FVector2f EvaluateInternal(const FEvaluateContext& Context) const override final;
+	virtual bool IsConstant() const override;
+
+	UPROPERTY(EditAnywhere, Category = "Values", meta = (BaseStruct = "/Script/Niagara.NiagaraStatelessExpressionVec2"))
+	FInstancedStruct A;
+
+	UPROPERTY(EditAnywhere, Category = "Values", meta = (BaseStruct = "/Script/Niagara.NiagaraStatelessExpressionVec2"))
+	FInstancedStruct B;
+};
+
+////////////////////////////////////////////////////////////////////////////////////////////////////////////////////////
+
+USTRUCT(meta = (OperationExpression, DisplayName = "Multiply Vec2"))
+struct FNiagaraStatelessExpressionVec2Multiply final : public FNiagaraStatelessExpressionVec2
+{
+	GENERATED_BODY()
+
+	FNiagaraStatelessExpressionVec2Multiply();
+	virtual ~FNiagaraStatelessExpressionVec2Multiply() = default;
+	virtual FInstancedStruct BuildInternal(const FNiagaraStatelessEmitterDataBuildContext& BuildContext) const override;
+	virtual FVector2f EvaluateInternal(const FEvaluateContext& Context) const override final;
+	virtual bool IsConstant() const override;
+
+	UPROPERTY(EditAnywhere, Category = "Values", meta = (BaseStruct = "/Script/Niagara.NiagaraStatelessExpressionVec2"))
+	FInstancedStruct A;
+
+	UPROPERTY(EditAnywhere, Category = "Values", meta = (BaseStruct = "/Script/Niagara.NiagaraStatelessExpressionVec2"))
+	FInstancedStruct B;
+};
+
+////////////////////////////////////////////////////////////////////////////////////////////////////////////////////////
+
+USTRUCT(meta = (OperationExpression, DisplayName = "Divide Vec2"))
+struct FNiagaraStatelessExpressionVec2Divide final : public FNiagaraStatelessExpressionVec2
+{
+	GENERATED_BODY()
+
+	FNiagaraStatelessExpressionVec2Divide();
+	virtual ~FNiagaraStatelessExpressionVec2Divide() = default;
+	virtual FInstancedStruct BuildInternal(const FNiagaraStatelessEmitterDataBuildContext& BuildContext) const override;
+	virtual FVector2f EvaluateInternal(const FEvaluateContext& Context) const override final;
+	virtual bool IsConstant() const override;
+
+	UPROPERTY(EditAnywhere, Category = "Values", meta = (BaseStruct = "/Script/Niagara.NiagaraStatelessExpressionVec2"))
+	FInstancedStruct A;
+
+	UPROPERTY(EditAnywhere, Category = "Values", meta = (BaseStruct = "/Script/Niagara.NiagaraStatelessExpressionVec2"))
+	FInstancedStruct B;
+};
diff --git a/Engine/Plugins/FX/Niagara/Source/Niagara/Internal/Stateless/Expressions/NiagaraStatelessExpressionVec3.h b/Engine/Plugins/FX/Niagara/Source/Niagara/Internal/Stateless/Expressions/NiagaraStatelessExpressionVec3.h
new file mode 100644
--- /dev/null
+++ b/Engine/Plugins/FX/Niagara/Source/Niagara/Internal/Stateless/Expressions/NiagaraStatelessExpressionVec3.h
@@ -0,0 +1,135 @@
+// Copyright Epic Games, Inc. All Rights Reserved.
+
+#pragma once
+
+#include "Stateless/NiagaraStatelessExpression.h"
+
+#include "NiagaraStatelessExpressionVec3.generated.h"
+
+////////////////////////////////////////////////////////////////////////////////////////////////////////////////////////
+
+USTRUCT()
+struct FNiagaraStatelessExpressionVec3 : public FNiagaraStatelessExpression
+{
+	GENERATED_BODY()
+
+	virtual FNiagaraTypeDefinition GetOutputTypeDef() const override final { return FNiagaraTypeDefinition::GetVec3Def(); }
+	virtual FInstancedStruct Build(const FNiagaraStatelessEmitterDataBuildContext& BuildContext) const override final;
+	virtual void Evaluate(const FEvaluateContext& Context, void* ValueAddress) const override final { *static_cast<FVector3f*>(ValueAddress) = EvaluateInternal(Context); }
+
+	virtual FInstancedStruct BuildInternal(const FNiagaraStatelessEmitterDataBuildContext& BuildContext) const { checkNoEntry(); return FInstancedStruct(); }
+	virtual FVector3f EvaluateInternal(const FEvaluateContext& Context) const { checkNoEntry(); return FVector3f::ZeroVector; }
+};
+
+////////////////////////////////////////////////////////////////////////////////////////////////////////////////////////
+
+USTRUCT(meta = (ValueExpression))
+struct FNiagaraStatelessExpressionVec3Constant final : public FNiagaraStatelessExpressionVec3
+{
+	GENERATED_BODY()
+
+	virtual ~FNiagaraStatelessExpressionVec3Constant() = default;
+	virtual FVector3f EvaluateInternal(const FEvaluateContext& Context) const override final;
+	virtual bool IsConstant() const override;
+
+	UPROPERTY(EditAnywhere, Category = "Values")
+	FVector3f A = FVector3f::ZeroVector;
+};
+
+////////////////////////////////////////////////////////////////////////////////////////////////////////////////////////
+
+USTRUCT(meta = (BindingExpression))
+struct FNiagaraStatelessExpressionVec3Binding final : public FNiagaraStatelessExpressionVec3
+{
+	GENERATED_BODY()
+
+	virtual ~FNiagaraStatelessExpressionVec3Binding() = default;
+	virtual bool IsConstant() const { return false; }
+	virtual FInstancedStruct BuildInternal(const FNiagaraStatelessEmitterDataBuildContext& BuildContext) const override;
+	virtual FVector3f EvaluateInternal(const FEvaluateContext& Context) const override final;
+
+	UPROPERTY(EditAnywhere, Category = "Values")
+	FName A;
+
+	int32 ParameterOffset = INDEX_NONE;
+};
+
+////////////////////////////////////////////////////////////////////////////////////////////////////////////////////////
+
+USTRUCT(meta = (OperationExpression, DisplayName = "Add Vec3"))
+struct FNiagaraStatelessExpressionVec3Add final : public FNiagaraStatelessExpressionVec3
+{
+	GENERATED_BODY()
+
+	FNiagaraStatelessExpressionVec3Add();
+	virtual ~FNiagaraStatelessExpressionVec3Add() = default;
+	virtual FInstancedStruct BuildInternal(const FNiagaraStatelessEmitterDataBuildContext& BuildContext) const override;
+	virtual FVector3f EvaluateInternal(const FEvaluateContext& Context) const override final;
+	virtual bool IsConstant() const override;
+
+	UPROPERTY(EditAnywhere, Category = "Values", meta = (BaseStruct = "/Script/Niagara.NiagaraStatelessExpressionVec3"))
+	FInstancedStruct A;
+
+	UPROPERTY(EditAnywhere, Category = "Values", meta = (BaseStruct = "/Script/Niagara.NiagaraStatelessExpressionVec3"))
+	FInstancedStruct B;
+};
+
+////////////////////////////////////////////////////////////////////////////////////////////////////////////////////////
+
+USTRUCT(meta = (OperationExpression, DisplayName = "Subtract Vec3"))
+struct FNiagaraStatelessExpressionVec3Subtract final : public FNiagaraStatelessExpressionVec3
+{
+	GENERATED_BODY()
+
+	FNiagaraStatelessExpressionVec3Subtract();
+	virtual ~FNiagaraStatelessExpressionVec3Subtract() = default;
+	virtual FInstancedStruct BuildInternal(const FNiagaraStatelessEmitterDataBuildContext& BuildContext) const override;
+	virtual FVector3f EvaluateInternal(const FEvaluateContext& Context) const override final;
+	virtual bool IsConstant() const override;
+
+	UPROPERTY(EditAnywhere, Category = "Values", meta = (BaseStruct = "/Script/Niagara.NiagaraStatelessExpressionVec3"))
+	FInstancedStruct A;
+
+	UPROPERTY(EditAnywhere, Category = "Values", meta = (BaseStruct = "/Script/Niagara.NiagaraStatelessExpressionVec3"))
+	FInstancedStruct B;
+};
+
+////////////////////////////////////////////////////////////////////////////////////////////////////////////////////////
+
+USTRUCT(meta = (OperationExpression, DisplayName = "Multiply Vec3"))
+struct FNiagaraStatelessExpressionVec3Multiply final : public FNiagaraStatelessExpressionVec3
+{
+	GENERATED_BODY()
+
+	FNiagaraStatelessExpressionVec3Multiply();
+	virtual ~FNiagaraStatelessExpressionVec3Multiply() = default;
+	virtual FInstancedStruct BuildInternal(const FNiagaraStatelessEmitterDataBuildContext& BuildContext) const override;
+	virtual FVector3f EvaluateInternal(const FEvaluateContext& Context) const override final;
+	virtual bool IsConstant() const override;
+
+	UPROPERTY(EditAnywhere, Category = "Values", meta = (BaseStruct = "/Script/Niagara.NiagaraStatelessExpressionVec3"))
+	FInstancedStruct A;
+
+	UPROPERTY(EditAnywhere, Category = "Values", meta = (BaseStruct = "/Script/Niagara.NiagaraStatelessExpressionVec3"))
+	FInstancedStruct B;
+};
+
+////////////////////////////////////////////////////////////////////////////////////////////////////////////////////////
+
+USTRUCT(meta = (OperationExpression, DisplayName = "Divide Vec3"))
+struct FNiagaraStatelessExpressionVec3Divide final : public FNiagaraStatelessExpressionVec3
+{
+	GENERATED_BODY()
+
+	FNiagaraStatelessExpressionVec3Divide();
+	virtual ~FNiagaraStatelessExpressionVec3Divide() = default;
+	virtual FInstancedStruct BuildInternal(const FNiagaraStatelessEmitterDataBuildContext& BuildContext) const override;
+	virtual FVector3f EvaluateInternal(const FEvaluateContext& Context) const override final;
+	virtual bool IsConstant() const override;
+
+	UPROPERTY(EditAnywhere, Category = "Values", meta = (BaseStruct = "/Script/Niagara.NiagaraStatelessExpressionVec3"))
+	FInstancedStruct A;
+
+	UPROPERTY(EditAnywhere, Category = "Values", meta = (BaseStruct = "/Script/Niagara.NiagaraStatelessExpressionVec3"))
+	FInstancedStruct B;
+};
diff --git a/Engine/Plugins/FX/Niagara/Source/Niagara/Internal/Stateless/Expressions/NiagaraStatelessExpressionVec4.h b/Engine/Plugins/FX/Niagara/Source/Niagara/Internal/Stateless/Expressions/NiagaraStatelessExpressionVec4.h
new file mode 100644
--- /dev/null
+++ b/Engine/Plugins/FX/Niagara/Source/Niagara/Internal/Stateless/Expressions/NiagaraStatelessExpressionVec4.h
@@ -0,0 +1,135 @@
+// Copyright Epic Games, Inc. All Rights Reserved.
+
+#pragma once
+
+#include "Stateless/NiagaraStatelessExpression.h"
+
+#include "NiagaraStatelessExpressionVec4.generated.h"
+
+////////////////////////////////////////////////////////////////////////////////////////////////////////////////////////
+
+USTRUCT()
+struct FNiagaraStatelessExpressionVec4 : public FNiagaraStatelessExpression
+{
+	GENERATED_BODY()
+
+	virtual FNiagaraTypeDefinition GetOutputTypeDef() const override final { return FNiagaraTypeDefinition::GetVec4Def(); }
+	virtual FInstancedStruct Build(const FNiagaraStatelessEmitterDataBuildContext& BuildContext) const override final;
+	virtual void Evaluate(const FEvaluateContext& Context, void* ValueAddress) const override final { *static_cast<FVector4f*>(ValueAddress) = EvaluateInternal(Context); }
+
+	virtual FInstancedStruct BuildInternal(const FNiagaraStatelessEmitterDataBuildContext& BuildContext) const { checkNoEntry(); return FInstancedStruct(); }
+	virtual FVector4f EvaluateInternal(const FEvaluateContext& Context) const { checkNoEntry(); return FVector4f::Zero(); }
+};
+
+////////////////////////////////////////////////////////////////////////////////////////////////////////////////////////
+
+USTRUCT(meta = (ValueExpression))
+struct FNiagaraStatelessExpressionVec4Constant final : public FNiagaraStatelessExpressionVec4
+{
+	GENERATED_BODY()
+
+	virtual ~FNiagaraStatelessExpressionVec4Constant() = default;
+	virtual FVector4f EvaluateInternal(const FEvaluateContext& Context) const override final;
+	virtual bool IsConstant() const override;
+
+	UPROPERTY(EditAnywhere, Category = "Values")
+	FVector4f A = FVector4f::Zero();
+};
+
+////////////////////////////////////////////////////////////////////////////////////////////////////////////////////////
+
+USTRUCT(meta = (BindingExpression))
+struct FNiagaraStatelessExpressionVec4Binding final : public FNiagaraStatelessExpressionVec4
+{
+	GENERATED_BODY()
+
+	virtual ~FNiagaraStatelessExpressionVec4Binding() = default;
+	virtual bool IsConstant() const { return false; }
+	virtual FInstancedStruct BuildInternal(const FNiagaraStatelessEmitterDataBuildContext& BuildContext) const override;
+	virtual FVector4f EvaluateInternal(const FEvaluateContext& Context) const override final;
+
+	UPROPERTY(EditAnywhere, Category = "Values")
+	FName A;
+
+	int32 ParameterOffset = INDEX_NONE;
+};
+
+////////////////////////////////////////////////////////////////////////////////////////////////////////////////////////
+
+USTRUCT(meta = (OperationExpression, DisplayName = "Add Vec4"))
+struct FNiagaraStatelessExpressionVec4Add final : public FNiagaraStatelessExpressionVec4
+{
+	GENERATED_BODY()
+
+	FNiagaraStatelessExpressionVec4Add();
+	virtual ~FNiagaraStatelessExpressionVec4Add() = default;
+	virtual FInstancedStruct BuildInternal(const FNiagaraStatelessEmitterDataBuildContext& BuildContext) const override;
+	virtual FVector4f EvaluateInternal(const FEvaluateContext& Context) const override final;
+	virtual bool IsConstant() const override;
+
+	UPROPERTY(EditAnywhere, Category = "Values", meta = (BaseStruct = "/Script/Niagara.NiagaraStatelessExpressionVec4"))
+	FInstancedStruct A;
+
+	UPROPERTY(EditAnywhere, Category = "Values", meta = (BaseStruct = "/Script/Niagara.NiagaraStatelessExpressionVec4"))
+	FInstancedStruct B;
+};
+
+////////////////////////////////////////////////////////////////////////////////////////////////////////////////////////
+
+USTRUCT(meta = (OperationExpression, DisplayName = "Subtract Vec4"))
+struct FNiagaraStatelessExpressionVec4Subtract final : public FNiagaraStatelessExpressionVec4
+{
+	GENERATED_BODY()
+
+	FNiagaraStatelessExpressionVec4Subtract();
+	virtual ~FNiagaraStatelessExpressionVec4Subtract() = default;
+	virtual FInstancedStruct BuildInternal(const FNiagaraStatelessEmitterDataBuildContext& BuildContext) const override;
+	virtual FVector4f EvaluateInternal(const FEvaluateContext& Context) const override final;
+	virtual bool IsConstant() const override;
+
+	UPROPERTY(EditAnywhere, Category = "Values", meta = (BaseStruct = "/Script/Niagara.NiagaraStatelessExpressionVec4"))
+	FInstancedStruct A;
+
+	UPROPERTY(EditAnywhere, Category = "Values", meta = (BaseStruct = "/Script/Niagara.NiagaraStatelessExpressionVec4"))
+	FInstancedStruct B;
+};
+
+////////////////////////////////////////////////////////////////////////////////////////////////////////////////////////
+
+USTRUCT(meta = (OperationExpression, DisplayName = "Multiply Vec4"))
+struct FNiagaraStatelessExpressionVec4Multiply final : public FNiagaraStatelessExpressionVec4
+{
+	GENERATED_BODY()
+
+	FNiagaraStatelessExpressionVec4Multiply();
+	virtual ~FNiagaraStatelessExpressionVec4Multiply() = default;
+	virtual FInstancedStruct BuildInternal(const FNiagaraStatelessEmitterDataBuildContext& BuildContext) const override;
+	virtual FVector4f EvaluateInternal(const FEvaluateContext& Context) const override final;
+	virtual bool IsConstant() const override;
+
+	UPROPERTY(EditAnywhere, Category = "Values", meta = (BaseStruct = "/Script/Niagara.NiagaraStatelessExpressionVec4"))
+	FInstancedStruct A;
+
+	UPROPERTY(EditAnywhere, Category = "Values", meta = (BaseStruct = "/Script/Niagara.NiagaraStatelessExpressionVec4"))
+	FInstancedStruct B;
+};
+
+////////////////////////////////////////////////////////////////////////////////////////////////////////////////////////
+
+USTRUCT(meta = (OperationExpression, DisplayName = "Divide Vec4"))
+struct FNiagaraStatelessExpressionVec4Divide final : public FNiagaraStatelessExpressionVec4
+{
+	GENERATED_BODY()
+
+	FNiagaraStatelessExpressionVec4Divide();
+	virtual ~FNiagaraStatelessExpressionVec4Divide() = default;
+	virtual FInstancedStruct BuildInternal(const FNiagaraStatelessEmitterDataBuildContext& BuildContext) const override;
+	virtual FVector4f EvaluateInternal(const FEvaluateContext& Context) const override final;
+	virtual bool IsConstant() const override;
+
+	UPROPERTY(EditAnywhere, Category = "Values", meta = (BaseStruct = "/Script/Niagara.NiagaraStatelessExpressionVec4"))
+	FInstancedStruct A;
+
+	UPROPERTY(EditAnywhere, Category = "Values", meta = (BaseStruct = "/Script/Niagara.NiagaraStatelessExpressionVec4"))
+	FInstancedStruct B;
+};
diff --git a/Engine/Plugins/FX/Niagara/Source/Niagara/Internal/Stateless/Modules/NiagaraStatelessModule_ScaleSpriteSize.h b/Engine/Plugins/FX/Niagara/Source/Niagara/Internal/Stateless/Modules/NiagaraStatelessModule_ScaleSpriteSize.h
--- a/Engine/Plugins/FX/Niagara/Source/Niagara/Internal/Stateless/Modules/NiagaraStatelessModule_ScaleSpriteSize.h
+++ b/Engine/Plugins/FX/Niagara/Source/Niagara/Internal/Stateless/Modules/NiagaraStatelessModule_ScaleSpriteSize.h
@@ -7,6 +7,7 @@
 #include "Stateless/NiagaraStatelessEmitterDataBuildContext.h"
 #include "Stateless/NiagaraStatelessModuleShaderParameters.h"
 #include "Stateless/NiagaraStatelessParticleSimContext.h"
+#include "NiagaraParameterBinding.h"
 
 #include "NiagaraStatelessModule_ScaleSpriteSize.generated.h"
 
diff --git a/Engine/Plugins/FX/Niagara/Source/Niagara/Internal/Stateless/NiagaraStatelessEmitterData.h b/Engine/Plugins/FX/Niagara/Source/Niagara/Internal/Stateless/NiagaraStatelessEmitterData.h
--- a/Engine/Plugins/FX/Niagara/Source/Niagara/Internal/Stateless/NiagaraStatelessEmitterData.h
+++ b/Engine/Plugins/FX/Niagara/Source/Niagara/Internal/Stateless/NiagaraStatelessEmitterData.h
@@ -19,6 +19,7 @@ namespace NiagaraStateless
 	class FSpawnInfoShaderParameters;
 	class FParticleSimulationExecData;
 }
+struct FInstancedStruct;
 
 struct FNiagaraStatelessEmitterData
 {
@@ -51,6 +52,7 @@ struct FNiagaraStatelessEmitterData
 
 	bool											bModulesHaveRendererBindings = false;
 	FNiagaraParameterStore							RendererBindings;			// Contains all bindings for modules & renderers
+	TArray<TPair<int32, FInstancedStruct>>			Expressions;				// Contains a mapping of expression to paramater store entry
 
 	const UNiagaraStatelessEmitterTemplate*			EmitterTemplate = nullptr;	// Used to access shader information
 
diff --git a/Engine/Plugins/FX/Niagara/Source/Niagara/Internal/Stateless/NiagaraStatelessEmitterDataBuildContext.h b/Engine/Plugins/FX/Niagara/Source/Niagara/Internal/Stateless/NiagaraStatelessEmitterDataBuildContext.h
--- a/Engine/Plugins/FX/Niagara/Source/Niagara/Internal/Stateless/NiagaraStatelessEmitterDataBuildContext.h
+++ b/Engine/Plugins/FX/Niagara/Source/Niagara/Internal/Stateless/NiagaraStatelessEmitterDataBuildContext.h
@@ -15,15 +15,17 @@ namespace NiagaraStateless
 	class FParticleSimulationContext;
 	class FParticleSimulationExecData;
 }
+struct FInstancedStruct;
 
 class FNiagaraStatelessEmitterDataBuildContext
 {
 public:
 	UE_NONCOPYABLE(FNiagaraStatelessEmitterDataBuildContext);
 
-	FNiagaraStatelessEmitterDataBuildContext(FNiagaraDataSetCompiledData& InParticleDataSet, FNiagaraParameterStore& InRendererBindings, TArray<uint8>& InBuiltData, TArray<float>& InStaticFloatData, NiagaraStateless::FParticleSimulationExecData* InParticleExecData)
+	FNiagaraStatelessEmitterDataBuildContext(FNiagaraDataSetCompiledData& InParticleDataSet, FNiagaraParameterStore& InRendererBindings, TArray<TPair<int32, FInstancedStruct>>& InExpressions, TArray<uint8>& InBuiltData, TArray<float>& InStaticFloatData, NiagaraStateless::FParticleSimulationExecData* InParticleExecData)
 		: ParticleDataSet(InParticleDataSet)
 		, RendererBindings(InRendererBindings)
+		, Expressions(InExpressions)
 		, BuiltData(InBuiltData)
 		, StaticFloatData(InStaticFloatData)
 		, ParticleExecData(InParticleExecData)
@@ -67,6 +69,12 @@ public:
 	int32 AddRendererBinding(const FNiagaraParameterBinding& Binding) const;
 	int32 AddRendererBinding(const FNiagaraParameterBindingWithValue& Binding) const;
 
+private:
+	// Adds a binding to the renderer parameter store
+	// The returned value is INDEX_NONE is the variables is index otherwise the offset in DWORDs
+	int32 AddExpression(const FInstancedStruct& Expression) const;
+
+public:
 	// Adds an distribution into the LUT if enabled
 	template<typename TType>
 	FNiagaraStatelessBuiltDistributionType AddDistribution(ENiagaraDistributionMode Mode, TConstArrayView<TType> Values, const FVector2f& TimeRange) const
@@ -79,6 +87,7 @@ public:
 			switch (Mode)
 			{
 				case ENiagaraDistributionMode::Binding:				checkNoEntry(); break;
+				case ENiagaraDistributionMode::Expression:			checkNoEntry(); break;
 				case ENiagaraDistributionMode::UniformConstant:		FNiagaraStatelessBuiltDistribution::SetIsRandom(BuiltDistribution); FNiagaraStatelessBuiltDistribution::SetIsUniform(BuiltDistribution); break;
 				case ENiagaraDistributionMode::NonUniformConstant:	FNiagaraStatelessBuiltDistribution::SetIsRandom(BuiltDistribution); break;
 				case ENiagaraDistributionMode::UniformRange:		FNiagaraStatelessBuiltDistribution::SetIsRandom(BuiltDistribution); FNiagaraStatelessBuiltDistribution::SetIsUniform(BuiltDistribution); break;
@@ -100,6 +109,7 @@ public:
 		FNiagaraStatelessBuiltDistributionType BuiltDistribution = FNiagaraStatelessBuiltDistribution::GetDefault();
 		if (Distribution.Mode == ENiagaraDistributionMode::Binding)
 		{
+			//-OPT: If the expression is constant should we pack into the static values data rather than going into the parameter store?
 			const int32 ParameterOffset = AddRendererBinding(Distribution.ParameterBinding);
 			if (ParameterOffset >= 0)
 			{
@@ -107,6 +117,16 @@ public:
 				FNiagaraStatelessBuiltDistribution::SetLookupParameters(BuiltDistribution, ParameterOffset);
 			}
 		}
+		else if (Distribution.Mode == ENiagaraDistributionMode::Expression)
+		{
+			//-OPT: If the expression is constant should we pack into the static values data rather than going into the parameter store?
+			const int32 ParameterOffset = AddExpression(Distribution.ParameterExpression);
+			if (ParameterOffset >= 0)
+			{
+				FNiagaraStatelessBuiltDistribution::SetIsBinding(BuiltDistribution);
+				FNiagaraStatelessBuiltDistribution::SetLookupParameters(BuiltDistribution, ParameterOffset);
+			}
+		}
 		else
 		{
 			BuiltDistribution = AddDistribution(Distribution.Mode, MakeArrayView(Distribution.Values), Distribution.ValuesTimeRange);
@@ -137,8 +157,14 @@ public:
 		TRange Range(DefaultValue);
 		if (Distribution.Mode == ENiagaraDistributionMode::Binding)
 		{
+			//-OPT: If the expression is constant we can just resolve to a range rather than adding to the parameter store
 			Range.ParameterOffset = AddRendererBinding(Distribution.ParameterBinding);
 		}
+		else if (Distribution.Mode == ENiagaraDistributionMode::Expression)
+		{
+			//-OPT: If the expression is constant we can just resolve to a range rather than adding to the parameter store
+			Range.ParameterOffset = AddExpression(Distribution.ParameterExpression);
+		}
 		else
 		{
 			Range = Distribution.CalculateRange(DefaultValue);
@@ -164,6 +190,7 @@ public:
 private:
 	FNiagaraDataSetCompiledData&					ParticleDataSet;
 	FNiagaraParameterStore&							RendererBindings;
+	TArray<TPair<int32, FInstancedStruct>>&			Expressions;
 	TArray<uint8>&									BuiltData;
 	TArray<float>&									StaticFloatData;
 	NiagaraStateless::FParticleSimulationExecData*	ParticleExecData = nullptr;
diff --git a/Engine/Plugins/FX/Niagara/Source/Niagara/Internal/Stateless/NiagaraStatelessExpression.h b/Engine/Plugins/FX/Niagara/Source/Niagara/Internal/Stateless/NiagaraStatelessExpression.h
new file mode 100644
--- /dev/null
+++ b/Engine/Plugins/FX/Niagara/Source/Niagara/Internal/Stateless/NiagaraStatelessExpression.h
@@ -0,0 +1,38 @@
+// Copyright Epic Games, Inc. All Rights Reserved.
+
+#pragma once
+
+#include "NiagaraCommon.h"
+#include "StructUtils/InstancedStruct.h"
+
+#include "NiagaraStatelessExpression.generated.h"
+
+class FNiagaraStatelessEmitterDataBuildContext;
+class FNiagaraStatelessSetShaderParameterContext;
+
+USTRUCT()
+struct FNiagaraStatelessExpression
+{
+	struct FEvaluateContext
+	{
+		UE_NONCOPYABLE(FEvaluateContext)
+
+		explicit FEvaluateContext(const FNiagaraParameterStore& InParameterStore)
+			: ParameterStore(InParameterStore)
+		{
+		}
+
+		const FNiagaraParameterStore&	ParameterStore;
+	};
+
+	GENERATED_BODY()
+
+	virtual FInstancedStruct Build(const FNiagaraStatelessEmitterDataBuildContext& BuildContext) const { checkNoEntry(); return FInstancedStruct();  }
+	virtual void Evaluate(const FEvaluateContext& Context, void* ValueAddress) const { checkNoEntry(); }
+	virtual FNiagaraTypeDefinition GetOutputTypeDef() const { checkNoEntry(); return FNiagaraTypeDefinition(); }
+	virtual bool IsConstant() const { return false; }
+
+#if WITH_EDITORONLY_DATA
+	static void ForEachBinding(const FInstancedStruct& ExpressionStruct, const TFunction<void(const FNiagaraVariableBase&)>& Delegate);
+#endif
+};
diff --git a/Engine/Plugins/FX/Niagara/Source/Niagara/Private/NiagaraAsyncCompile.cpp b/Engine/Plugins/FX/Niagara/Source/Niagara/Private/NiagaraAsyncCompile.cpp
--- a/Engine/Plugins/FX/Niagara/Source/Niagara/Private/NiagaraAsyncCompile.cpp
+++ b/Engine/Plugins/FX/Niagara/Source/Niagara/Private/NiagaraAsyncCompile.cpp
@@ -766,7 +766,7 @@ bool FNiagaraActiveCompilationDefault::QueryCompileComplete(const FNiagaraQueryC
 						for (const FNiagaraVariableWithOffset& ParamWithOffset : RapidIterationParameters)
 						{
 							FNiagaraVariable& Parameter = AsyncTask->BakedRapidIterationParameters.Add_GetRef(ParamWithOffset);
-							Parameter.SetData(CompiledScript->RapidIterationParameters.GetParameterData(ParamWithOffset.Offset));
+							Parameter.SetData(CompiledScript->RapidIterationParameters.GetParameterData(ParamWithOffset));
 						}
 
 						// if preparing the rapid iteration parameters has adjusted the compile id, then we want to be sure
diff --git a/Engine/Plugins/FX/Niagara/Source/Niagara/Private/NiagaraCommon.cpp b/Engine/Plugins/FX/Niagara/Source/Niagara/Private/NiagaraCommon.cpp
--- a/Engine/Plugins/FX/Niagara/Source/Niagara/Private/NiagaraCommon.cpp
+++ b/Engine/Plugins/FX/Niagara/Source/Niagara/Private/NiagaraCommon.cpp
@@ -1593,8 +1593,8 @@ void FNiagaraUtilities::PrepareRapidIterationParameters(const TArray<UNiagaraScr
 				else
 				{
 					if (FMemory::Memcmp(
-						Script->RapidIterationParameters.GetParameterData(SourceOffset),
-						PreparedParameterStore.GetParameterData(PreparedOffset),
+						Script->RapidIterationParameters.GetParameterData(SourceOffset, SourceParameter.GetType()),
+						PreparedParameterStore.GetParameterData(PreparedOffset, SourceParameter.GetType()),
 						SourceParameter.GetSizeInBytes()) != 0)
 					{
 						bOverwriteParameters = true;
diff --git a/Engine/Plugins/FX/Niagara/Source/Niagara/Private/NiagaraComponent.cpp b/Engine/Plugins/FX/Niagara/Source/Niagara/Private/NiagaraComponent.cpp
--- a/Engine/Plugins/FX/Niagara/Source/Niagara/Private/NiagaraComponent.cpp
+++ b/Engine/Plugins/FX/Niagara/Source/Niagara/Private/NiagaraComponent.cpp
@@ -3303,7 +3303,7 @@ void UNiagaraComponent::SetUserParametersToDefaultValues()
 				}
 				else
 				{
-					OverrideParameters.SetParameterData(SourceUserParameterStore.GetParameterData(SourceIndex), DestIndex, DestParameter.GetSizeInBytes());
+					OverrideParameters.SetParameterData(SourceUserParameterStore.GetParameterData(SourceIndex, DestParameter.GetType()), DestIndex, DestParameter.GetSizeInBytes());
 				}
 			}
 			else
diff --git a/Engine/Plugins/FX/Niagara/Source/Niagara/Private/NiagaraParameterCollection.cpp b/Engine/Plugins/FX/Niagara/Source/Niagara/Private/NiagaraParameterCollection.cpp
--- a/Engine/Plugins/FX/Niagara/Source/Niagara/Private/NiagaraParameterCollection.cpp
+++ b/Engine/Plugins/FX/Niagara/Source/Niagara/Private/NiagaraParameterCollection.cpp
@@ -226,7 +226,7 @@ void UNiagaraParameterCollectionInstance::SyncWithCollection()
 			}
 			else
 			{
-				ParameterStorage.SetParameterData(OldStore.GetParameterData(Offset), ParameterStorageOffset, Param.GetSizeInBytes());
+				ParameterStorage.SetParameterData(OldStore.GetParameterData(Offset, Param.GetType()), ParameterStorageOffset, Param.GetSizeInBytes());
 			}
 		}
 		else
@@ -248,7 +248,7 @@ void UNiagaraParameterCollectionInstance::SyncWithCollection()
 			}
 			else
 			{
-				ParameterStorage.SetParameterData(DefaultStore.GetParameterData(Offset), ParameterStorageOffset, Param.GetSizeInBytes());
+				ParameterStorage.SetParameterData(DefaultStore.GetParameterData(Offset, Param.GetType()), ParameterStorageOffset, Param.GetSizeInBytes());
 			}
 		}
 	}
diff --git a/Engine/Plugins/FX/Niagara/Source/Niagara/Private/NiagaraParameterStore.cpp b/Engine/Plugins/FX/Niagara/Source/Niagara/Private/NiagaraParameterStore.cpp
--- a/Engine/Plugins/FX/Niagara/Source/Niagara/Private/NiagaraParameterStore.cpp
+++ b/Engine/Plugins/FX/Niagara/Source/Niagara/Private/NiagaraParameterStore.cpp
@@ -249,6 +249,8 @@ void FNiagaraParameterStore::Bind(FNiagaraParameterStore* DestStore, const FNiag
 	}
 }
 
+////////////////////////////////////////////////////////////////////////////////////////////////////////////////////////////////////////////////////////////////////////////////////////////////////////
+
 template <typename TVisitor>
 void FNiagaraParameterStoreBinding::MatchParameters(FNiagaraParameterStore* DestStore, FNiagaraParameterStore* SrcStore, TVisitor Visitor)
 {
@@ -322,7 +324,7 @@ void FNiagaraParameterStoreBinding::CopyParameters(FNiagaraParameterStore* DestS
 {
 	for (const FParameterBinding& Binding : ParameterBindings)
 	{
-		DestStore->SetParameterData(SrcStore->GetParameterData(Binding.SrcOffset), Binding.DestOffset, Binding.Size);
+		DestStore->SetParameterData(SrcStore->GetParameterData(Binding.SrcOffset, Binding.Size), Binding.DestOffset, Binding.Size);
 	}
 
 	for (const FInterfaceBinding& Binding : InterfaceBindings)
@@ -411,6 +413,165 @@ bool FNiagaraParameterStoreBinding::BindParameters(FNiagaraParameterStore* DestS
 	return bAnyBinding;
 }
 
+bool FNiagaraParameterStoreBinding::VerifyBinding(const FNiagaraParameterStore* DestStore, const FNiagaraParameterStore* SrcStore)const
+{
+	bool bBindingValid = true;
+#if WITH_EDITORONLY_DATA
+	TArray<FName, TInlineAllocator<32>> MissingParameterNames;
+	for (const FNiagaraVariableWithOffset& ParamWithOffset : DestStore->ReadParameterVariables())
+	{
+		const FNiagaraVariable& Parameter = ParamWithOffset;
+		int32 DestOffset = ParamWithOffset.Offset;
+		int32 SrcOffset = SrcStore->IndexOf(Parameter);
+
+		if (SrcOffset == INDEX_NONE)
+		{
+			continue;//Parameter is not present in SrcStore so shouldn't/cant be bound.
+		}
+
+		if (Parameter.IsDataInterface())
+		{
+			if (InterfaceBindings.Find(FInterfaceBinding(SrcOffset, DestOffset)) == INDEX_NONE)
+			{
+				MissingParameterNames.Add(Parameter.GetName());
+				bBindingValid = false;
+			}
+		}
+		else if (Parameter.IsUObject())
+		{
+			if (UObjectBindings.Find(FUObjectBinding(SrcOffset, DestOffset)) == INDEX_NONE)
+			{
+				MissingParameterNames.Add(Parameter.GetName());
+				bBindingValid = false;
+			}
+		}
+		else
+		{
+			if (ParameterBindings.Find(FParameterBinding(SrcOffset, DestOffset, Parameter.GetSizeInBytes())) == INDEX_NONE)
+			{
+				MissingParameterNames.Add(Parameter.GetName());
+				bBindingValid = false;
+			}
+		}
+	}
+
+	if (MissingParameterNames.Num())
+	{
+		UE_LOG(LogNiagara, Warning, TEXT("Invalid ParameterStore Binding: Parameters missing from binding between %s and %s. Stores should have been rebound when one changed!"), *SrcStore->DebugName, *DestStore->DebugName);
+		for (FName MissingParam : MissingParameterNames)
+		{
+			UE_LOG(LogNiagara, Warning, TEXT("%s"), *MissingParam.ToString());
+		}
+	}
+#endif
+	return bBindingValid;
+}
+
+void FNiagaraParameterStoreBinding::Dump(const FNiagaraParameterStore* DestStore, const FNiagaraParameterStore* SrcStore)const
+{
+#if WITH_EDITORONLY_DATA
+	UE_LOG(LogNiagara, Log, TEXT("\n\nDest Store: %s\n"), *DestStore->DebugName);
+	for (const FParameterBinding& Binding : ParameterBindings)
+	{
+		ensure(Binding.Size != 0);
+		ensure(Binding.SrcOffset != -1);
+		ensure(Binding.DestOffset != -1);
+		FNiagaraVariable Param;
+		bool bFound = false;
+		for (const FNiagaraVariableWithOffset& ParamWithOffset : DestStore->ReadParameterVariables())
+		{
+			if (ParamWithOffset.Offset == Binding.DestOffset && !ParamWithOffset.IsDataInterface())
+			{
+				bFound = true;
+				Param = ParamWithOffset;
+			}
+		}
+		if (ensure(bFound))
+		{
+			UE_LOG(LogNiagara, Log, TEXT("| Param | %s %s: Src:%d - Dest:%d - Size:%d\n"), *Param.GetType().GetName(), *Param.GetName().ToString(), Binding.SrcOffset, Binding.DestOffset, Binding.Size);
+
+			//Also ensure the param has been pushed correctly.
+			const uint8* SrcData = SrcStore->GetParameterData(Binding.SrcOffset, Binding.Size);
+			const uint8* DestData = DestStore->GetParameterData(Binding.DestOffset, Binding.Size);
+			if (!ensure(FMemory::Memcmp(SrcData, DestData, Binding.Size) == 0))
+			{
+				UE_LOG(LogNiagara, Log, TEXT("Parameter in dest store has incorrect data!\n"));
+			}
+		}
+		else
+		{
+			UE_LOG(LogNiagara, Log, TEXT("Failed to find matching param in bound store!\n"));
+		}
+	}
+
+	for (const FInterfaceBinding& Binding : InterfaceBindings)
+	{
+		ensure(Binding.SrcOffset != -1);
+		ensure(Binding.DestOffset != -1);
+		FNiagaraVariable Param;
+		bool bFound = false;
+		for (const FNiagaraVariableWithOffset& ParamWithOffset : DestStore->ReadParameterVariables())
+		{
+			if (ParamWithOffset.Offset == Binding.DestOffset && ParamWithOffset.IsDataInterface())
+			{
+				bFound = true;
+				Param = ParamWithOffset;
+			}
+		}
+		if (ensure(bFound))
+		{
+			UE_LOG(LogNiagara, Log, TEXT("| DI | %s %s: Src:%d - Dest:%d\n"), *Param.GetType().GetName(), *Param.GetName().ToString(), Binding.SrcOffset, Binding.DestOffset);
+
+			//Also ensure the param has been pushed correctly.
+			const UNiagaraDataInterface* SrcData = SrcStore->GetDataInterfaces()[Binding.SrcOffset];
+			const UNiagaraDataInterface* DestData = DestStore->GetDataInterfaces()[Binding.DestOffset];
+			if (!ensure(SrcData == DestData))
+			{
+				UE_LOG(LogNiagara, Log, TEXT("Data interface parameter in dest store is incorrect!\n"));
+			}
+		}
+		else
+		{
+			UE_LOG(LogNiagara, Log, TEXT("Failed to find matching data interface param in bound store!\n"));
+		}
+	}
+
+	for (const FUObjectBinding& Binding : UObjectBindings)
+	{
+		ensure(Binding.SrcOffset != -1);
+		ensure(Binding.DestOffset != -1);
+		FNiagaraVariable Param;
+		bool bFound = false;
+		for (const FNiagaraVariableWithOffset& ParamWithOffset : DestStore->ReadParameterVariables())
+		{
+			if (ParamWithOffset.Offset == Binding.DestOffset && ParamWithOffset.IsUObject())
+			{
+				bFound = true;
+				Param = ParamWithOffset;
+			}
+		}
+		if (ensure(bFound))
+		{
+			UE_LOG(LogNiagara, Log, TEXT("| UObject | %s %s: Src:%d - Dest:%d\n"), *Param.GetType().GetName(), *Param.GetName().ToString(), Binding.SrcOffset, Binding.DestOffset);
+
+			//Also ensure the param has been pushed correctly.
+			TObjectPtr<const UObject> SrcData = SrcStore->GetUObjects()[Binding.SrcOffset];
+			TObjectPtr<const UObject> DestData = DestStore->GetUObjects()[Binding.DestOffset];
+			if (!ensure(SrcData == DestData))
+			{
+				UE_LOG(LogNiagara, Log, TEXT("Data interface parameter in dest store is incorrect!\n"));
+			}
+		}
+		else
+		{
+			UE_LOG(LogNiagara, Log, TEXT("Failed to find matching data interface param in bound store!\n"));
+		}
+	}
+#endif
+}
+
+////////////////////////////////////////////////////////////////////////////////////////////////////////////////////////////////////////////////////////////////////////////////////////////////////////
+
 void FNiagaraParameterStore::Unbind(FNiagaraParameterStore* DestStore)
 {
 	const int32 BindingIndex = Bindings.IndexOfByPredicate([DestStore](const BindingPair& Binding)
@@ -482,28 +643,27 @@ void FNiagaraParameterStore::CheckForNaNs()const
 		bool bContainsNans = false;
 		if (Var.GetType() == FNiagaraTypeDefinition::GetFloatDef())
 		{
-			float Val = *(float*)GetParameterData(Offset);
+			const float Val = GetParameterValueFromOffset<float>(Offset);
 			bContainsNans = FMath::IsNaN(Val) || !FMath::IsFinite(Val);
 		}
 		else if (Var.GetType() == FNiagaraTypeDefinition::GetVec2Def())
 		{
-			FVector2f Val = *(FVector2f*)GetParameterData(Offset);
+			const FVector2f Val = GetParameterValueFromOffset<FVector2f>(Offset);
 			bContainsNans = Val.ContainsNaN();
 		}
 		else if (Var.GetType() == FNiagaraTypeDefinition::GetVec3Def())
 		{
-			FVector3f Val = *(FVector3f*)GetParameterData(Offset);
+			const FVector3f Val = GetParameterValueFromOffset<FVector3f>(Offset);
 			bContainsNans = Val.ContainsNaN();
 		}
 		else if (Var.GetType() == FNiagaraTypeDefinition::GetVec4Def())
 		{
-			FVector4f Val = *(FVector4f*)GetParameterData(Offset);
+			const FVector4f Val = GetParameterValueFromOffset<FVector4f>(Offset);
 			bContainsNans = Val.ContainsNaN();
 		}
 		else if (Var.GetType() == FNiagaraTypeDefinition::GetMatrix4Def())
 		{
-			FMatrix Val;
-			FMemory::Memcpy(&Val, GetParameterData(Offset), sizeof(FMatrix44f));
+			const FMatrix44f Val = GetParameterValueFromOffset<FMatrix44f>(Offset);
 			bContainsNans = Val.ContainsNaN();
 		}
 
@@ -535,7 +695,7 @@ void FNiagaraParameterStore::ConvertParameterType(const FNiagaraVariable& Existi
 	}
 	if (NewType == FNiagaraTypeDefinition::GetPositionDef())
 	{
-		FVector CurrentValue = (FVector)*reinterpret_cast<FVector3f*>(GetParameterData_Internal(Idx));
+		const FVector CurrentValue = FVector(*reinterpret_cast<const FVector3f*>(GetParameterData_Internal(Idx)));
 		SetPositionData(ExistingParam.GetName(), CurrentValue);
 		bPositionDataDirty = true;
 	}
@@ -555,7 +715,7 @@ void FNiagaraParameterStore::TickBindings() const
 		{
 			for (const FNiagaraParameterStoreBinding::FParameterBinding& ParamBinding : Binding.Value.ParameterBindings)
 			{
-				DestStore->SetParameterData(GetParameterData(ParamBinding.SrcOffset), ParamBinding.DestOffset, ParamBinding.Size);
+				DestStore->SetParameterData(GetParameterData(ParamBinding.SrcOffset, ParamBinding.Size), ParamBinding.DestOffset, ParamBinding.Size);
 			}
 		}
 		if ( bInterfacesDirty )
@@ -615,7 +775,7 @@ void FNiagaraParameterStore::DumpParameters(bool bDumpBindings) const
 	for (const FNiagaraVariableWithOffset& VariableBase : ReadParameterVariables())
 	{
 		FNiagaraVariable Var(VariableBase);
-		Var.SetData(GetParameterData(VariableBase.Offset));
+		Var.SetData(GetParameterData_Internal(VariableBase.Offset));
 		UE_LOG(LogNiagara, Log, TEXT("Param: %s Offset: %d Type : %s"), *Var.ToString(), IndexOf(Var), *Var.GetType().GetName());
 	}
 
@@ -635,7 +795,7 @@ FString FNiagaraParameterStore::ToString() const
 	for (const FNiagaraVariableWithOffset& VariableBase : ReadParameterVariables())
 	{
 		FNiagaraVariable Var(VariableBase);
-		Var.SetData(GetParameterData(VariableBase.Offset));
+		Var.SetData(GetParameterData_Internal(VariableBase.Offset));
 		Value += FString::Printf(TEXT("Param: %s Offset: %d Type : %s\n"), *Var.ToString(), IndexOf(Var), *Var.GetType().GetName());
 	}
 
@@ -724,11 +884,11 @@ bool FNiagaraParameterStore::AddParameter(const FNiagaraVariable& Param, bool bI
 			ParameterData.AddUninitialized(ParamSize);
 			if (StructConverter.IsValid())
 			{
-				StructConverter.ConvertDataToSimulation(GetParameterData_Internal(Offset), Param.GetData());
+				StructConverter.ConvertDataToSimulation(GetMutableParameterData_Internal(Offset), Param.GetData());
 			}
 			else
 			{
-				FMemory::Memcpy(GetParameterData_Internal(Offset), Param.GetData(), ParamSize);	
+				FMemory::Memcpy(GetMutableParameterData_Internal(Offset), Param.GetData(), ParamSize);	
 			}			
 		}
 		else
@@ -1052,7 +1212,7 @@ void FNiagaraParameterStore::CopyParametersTo(FNiagaraParameterStore& DestStore,
 			{
 				if (ParameterData.Num() != 0)
 				{
-					DestStore.SetParameterData(GetParameterData(SrcIndex), DestIndex, Parameter.GetSizeInBytes());
+					DestStore.SetParameterData(GetParameterData_Internal(SrcIndex), DestIndex, Parameter.GetSizeInBytes());
 				}
 			}
 		}
@@ -1228,7 +1388,7 @@ void FNiagaraParameterStore::ResolvePositions(FNiagaraLWCConverter LwcConverter)
 		if (Offset != INDEX_NONE)
 		{
 			FNiagaraPosition SimPosition = LwcConverter.ConvertWorldToSimulationPosition(Data.Value);
-			FNiagaraPosition* ParamData = reinterpret_cast<FNiagaraPosition*>( GetParameterData_Internal(Offset) );
+			FNiagaraPosition* ParamData = reinterpret_cast<FNiagaraPosition*>( GetMutableParameterData_Internal(Offset) );
 			FMemory::Memcpy(ParamData, &SimPosition, sizeof(FNiagaraPosition));
 		}
 	}
@@ -1249,7 +1409,7 @@ bool FNiagaraParameterStore::SetParameterData(const uint8* Data, FNiagaraVariabl
 	if (Offset != INDEX_NONE)
 	{
 		checkSlow(!Param.IsDataInterface());
-		uint8* Dest = GetParameterData_Internal(Offset);
+		uint8* Dest = GetMutableParameterData_Internal(Offset);
 		if (Dest != Data)
 		{
 			FNiagaraLwcStructConverter StructConverter = GetStructConverter(Param);
@@ -1273,7 +1433,7 @@ bool FNiagaraParameterStore::SetParameterData(const uint8* Data, FNiagaraVariabl
 			bool bTriggerRebind = false;
 			AddParameter(Param, bInitInterfaces, bTriggerRebind, &Offset);
 			check(Offset != INDEX_NONE);
-			uint8* Dest = GetParameterData_Internal(Offset);
+			uint8* Dest = GetMutableParameterData_Internal(Offset);
 			if (Dest != Data)
 			{
 				FNiagaraLwcStructConverter StructConverter = GetStructConverter(Param);
@@ -1384,7 +1544,7 @@ void FNiagaraParameterStore::CopyParameterData(FNiagaraParameterStore& DestStore
 		}
 		else
 		{
-			DestStore.SetParameterData(GetParameterData(SrcIndex), DestIndex, SourceParameter.GetSizeInBytes());
+			DestStore.SetParameterData(GetParameterData_Internal(SrcIndex), DestIndex, SourceParameter.GetSizeInBytes());
 		}
 	}
 }
@@ -1393,7 +1553,7 @@ bool FNiagaraParameterStore::CopyParameterData(const FNiagaraVariable& Parameter
 {
 	if (const FNiagaraVariableWithOffset* NiagaraVariableWithOffset = FindParameterVariable(Parameter))
 	{
-		const uint8* SourceData = GetParameterData(NiagaraVariableWithOffset->Offset);
+		const uint8* SourceData = GetParameterData_Internal(NiagaraVariableWithOffset->Offset);
 		if (NiagaraVariableWithOffset->StructConverter.IsValid())
 		{
 			NiagaraVariableWithOffset->StructConverter.ConvertDataFromSimulation(DestinationData, SourceData);
diff --git a/Engine/Plugins/FX/Niagara/Source/Niagara/Private/NiagaraRendererVolumes.cpp b/Engine/Plugins/FX/Niagara/Source/Niagara/Private/NiagaraRendererVolumes.cpp
--- a/Engine/Plugins/FX/Niagara/Source/Niagara/Private/NiagaraRendererVolumes.cpp
+++ b/Engine/Plugins/FX/Niagara/Source/Niagara/Private/NiagaraRendererVolumes.cpp
@@ -179,27 +179,27 @@ FNiagaraDynamicDataBase* FNiagaraRendererVolumes::GenerateDynamicData(const FNia
 
 		if (VFBoundOffsetsInParamStore[ENiagaraVolumeVFLayout::Position] != INDEX_NONE)
 		{
-			FMemory::Memcpy(&VolumeDynamicData->DefaultPosition, ParameterData.GetParameterData(VFBoundOffsetsInParamStore[ENiagaraVolumeVFLayout::Position]), sizeof(FVector3f));
+			VolumeDynamicData->DefaultPosition = ParameterData.GetParameterValueFromOffset<FVector3f>(VFBoundOffsetsInParamStore[ENiagaraVolumeVFLayout::Position]);
 		}
 		if (VFBoundOffsetsInParamStore[ENiagaraVolumeVFLayout::Rotation] != INDEX_NONE)
 		{
-			FMemory::Memcpy(&VolumeDynamicData->DefaultRotation, ParameterData.GetParameterData(VFBoundOffsetsInParamStore[ENiagaraVolumeVFLayout::Rotation]), sizeof(FQuat4f));
+			VolumeDynamicData->DefaultRotation = ParameterData.GetParameterValueFromOffset<FQuat4f>(VFBoundOffsetsInParamStore[ENiagaraVolumeVFLayout::Rotation]);
 		}
 		if (VFBoundOffsetsInParamStore[ENiagaraVolumeVFLayout::Scale] != INDEX_NONE)
 		{
-			FMemory::Memcpy(&VolumeDynamicData->DefaultScale, ParameterData.GetParameterData(VFBoundOffsetsInParamStore[ENiagaraVolumeVFLayout::Scale]), sizeof(FVector3f));
+			VolumeDynamicData->DefaultScale = ParameterData.GetParameterValueFromOffset<FVector3f>(VFBoundOffsetsInParamStore[ENiagaraVolumeVFLayout::Scale]);
 		}
 		if (VFBoundOffsetsInParamStore[ENiagaraVolumeVFLayout::RendererVisibilityTag] != INDEX_NONE)
 		{
-			FMemory::Memcpy(&VolumeDynamicData->DefaultRendererVisibilityTag, ParameterData.GetParameterData(VFBoundOffsetsInParamStore[ENiagaraVolumeVFLayout::RendererVisibilityTag]), sizeof(int32));
+			VolumeDynamicData->DefaultRendererVisibilityTag = ParameterData.GetParameterValueFromOffset<int32>(VFBoundOffsetsInParamStore[ENiagaraVolumeVFLayout::RendererVisibilityTag]);
 		}
 		if (VFBoundOffsetsInParamStore[ENiagaraVolumeVFLayout::VolumeResolutionMaxAxis] != INDEX_NONE)
 		{
-			FMemory::Memcpy(&VolumeDynamicData->DefaultVolumeResolutionMaxAxis, ParameterData.GetParameterData(VFBoundOffsetsInParamStore[ENiagaraVolumeVFLayout::VolumeResolutionMaxAxis]), sizeof(int32));
+			VolumeDynamicData->DefaultVolumeResolutionMaxAxis = ParameterData.GetParameterValueFromOffset<int32>(VFBoundOffsetsInParamStore[ENiagaraVolumeVFLayout::VolumeResolutionMaxAxis]);
 		}
 		if (VFBoundOffsetsInParamStore[ENiagaraVolumeVFLayout::VolumeWorldSpaceSize] != INDEX_NONE)
 		{
-			FMemory::Memcpy(&VolumeDynamicData->DefaultVolumeWorldSpaceSize, ParameterData.GetParameterData(VFBoundOffsetsInParamStore[ENiagaraVolumeVFLayout::VolumeWorldSpaceSize]), sizeof(FVector3f));
+			VolumeDynamicData->DefaultVolumeWorldSpaceSize = ParameterData.GetParameterValueFromOffset<FVector3f>(VFBoundOffsetsInParamStore[ENiagaraVolumeVFLayout::VolumeWorldSpaceSize]);
 		}
 	}
 
diff --git a/Engine/Plugins/FX/Niagara/Source/Niagara/Private/NiagaraScript.cpp b/Engine/Plugins/FX/Niagara/Source/Niagara/Private/NiagaraScript.cpp
--- a/Engine/Plugins/FX/Niagara/Source/Niagara/Private/NiagaraScript.cpp
+++ b/Engine/Plugins/FX/Niagara/Source/Niagara/Private/NiagaraScript.cpp
@@ -1453,7 +1453,7 @@ void UNiagaraScript::ComputeVMCompilationId(FNiagaraVMExecutableDataId& Id, cons
 					}
 				}
 				ReferencedStaticVars.Add(Var);
-				ReferencedStaticVarValues.Add(Script->RapidIterationParameters.GetParameterData(Var.Offset));
+				ReferencedStaticVarValues.Add(Script->RapidIterationParameters.GetParameterData(Var));
 			}
 		}
 	}
@@ -1484,7 +1484,7 @@ void UNiagaraScript::ComputeVMCompilationId(FNiagaraVMExecutableDataId& Id, cons
 					// Hash the name, type, and value of each parameter..
 					UpdateHashState(Var.GetName());
 					UpdateHashState(Var.GetType().GetFName());
-					HashState.Update(RapidIterationParameters.GetParameterData(Var.Offset), Var.GetSizeInBytes());
+					HashState.Update(RapidIterationParameters.GetParameterData(Var), Var.GetSizeInBytes());
 				}
 			}
 		}
@@ -4436,7 +4436,7 @@ void UNiagaraScript::GatherScriptStaticVariables(TArray<FNiagaraVariable>& Stati
 		if (Param.GetType().IsStatic() && Param.Offset != INDEX_NONE)
 		{
 			FNiagaraVariable StaticVariable = Param;
-			StaticVariable.SetData(RapidIterationParameters.GetParameterData(Param.Offset));
+			StaticVariable.SetData(RapidIterationParameters.GetParameterData(Param));
 
 			StaticVariables.AddUnique(StaticVariable);
 		}
diff --git a/Engine/Plugins/FX/Niagara/Source/Niagara/Private/NiagaraScriptExecutionParameterStore.cpp b/Engine/Plugins/FX/Niagara/Source/Niagara/Private/NiagaraScriptExecutionParameterStore.cpp
--- a/Engine/Plugins/FX/Niagara/Source/Niagara/Private/NiagaraScriptExecutionParameterStore.cpp
+++ b/Engine/Plugins/FX/Niagara/Source/Niagara/Private/NiagaraScriptExecutionParameterStore.cpp
@@ -317,7 +317,7 @@ void FNiagaraScriptInstanceParameterStore::CopyCurrToPrev()
 	if (const auto* ScriptStore = ScriptParameterStore.Get())
 	{
 		const int32 ParamStart = ScriptStore->ParameterSize;
-		FMemory::Memcpy(GetParameterData_Internal(ParamStart), GetParameterData(0), ParamStart);
+		FMemory::Memcpy(GetMutableParameterData_Internal(ParamStart), GetParameterData_Internal(0), ParamStart);
 	}
 }
 
diff --git a/Engine/Plugins/FX/Niagara/Source/Niagara/Private/Stateless/Expressions/NiagaraStatelessExpressionColor.cpp b/Engine/Plugins/FX/Niagara/Source/Niagara/Private/Stateless/Expressions/NiagaraStatelessExpressionColor.cpp
new file mode 100644
--- /dev/null
+++ b/Engine/Plugins/FX/Niagara/Source/Niagara/Private/Stateless/Expressions/NiagaraStatelessExpressionColor.cpp
@@ -0,0 +1,167 @@
+// Copyright Epic Games, Inc. All Rights Reserved.
+
+#include "Stateless/Expressions/NiagaraStatelessExpressionColor.h"
+
+#include "Stateless/NiagaraStatelessEmitterDataBuildContext.h"
+#include "Stateless/NiagaraStatelessSetShaderParameterContext.h"
+#include "NiagaraParameterStore.h"
+
+////////////////////////////////////////////////////////////////////////////////////////////////////////////////////////
+
+FInstancedStruct FNiagaraStatelessExpressionColor::Build(const FNiagaraStatelessEmitterDataBuildContext& BuildContext) const
+{
+	FInstancedStruct BuiltExpression;
+	if (IsConstant())
+	{
+		FNiagaraStatelessExpressionColorConstant ConstantExpression;
+		FNiagaraParameterStore EmptyParameterStore;
+		ConstantExpression.A = EvaluateInternal(FEvaluateContext(EmptyParameterStore));
+		return FInstancedStruct::Make(ConstantExpression);
+	}
+	return BuildInternal(BuildContext);
+}
+
+////////////////////////////////////////////////////////////////////////////////////////////////////////////////////////
+
+FLinearColor FNiagaraStatelessExpressionColorConstant::EvaluateInternal(const FEvaluateContext& Context) const
+{
+	return A;
+}
+
+bool FNiagaraStatelessExpressionColorConstant::IsConstant() const
+{
+	return true;
+}
+
+////////////////////////////////////////////////////////////////////////////////////////////////////////////////////////
+
+FInstancedStruct FNiagaraStatelessExpressionColorBinding::BuildInternal(const FNiagaraStatelessEmitterDataBuildContext& BuildContext) const
+{
+	FNiagaraStatelessExpressionColorBinding BuiltExpression;
+	BuiltExpression.A = A;
+	BuiltExpression.ParameterOffset = BuildContext.AddRendererBinding(FNiagaraVariableBase(GetOutputTypeDef(), A));
+	BuiltExpression.ParameterOffset *= sizeof(uint32);
+	return FInstancedStruct::Make(BuiltExpression);
+}
+
+FLinearColor FNiagaraStatelessExpressionColorBinding::EvaluateInternal(const FEvaluateContext& Context) const
+{
+	return Context.ParameterStore.GetParameterValueFromOffset<FLinearColor>(ParameterOffset);
+}
+
+////////////////////////////////////////////////////////////////////////////////////////////////////////////////////////
+
+FNiagaraStatelessExpressionColorAdd::FNiagaraStatelessExpressionColorAdd()
+{
+	A.InitializeAs<FNiagaraStatelessExpressionColorConstant>();
+	B.InitializeAs<FNiagaraStatelessExpressionColorConstant>();
+}
+
+FInstancedStruct FNiagaraStatelessExpressionColorAdd::BuildInternal(const FNiagaraStatelessEmitterDataBuildContext& BuildContext) const
+{
+	FNiagaraStatelessExpressionColorAdd BuiltExpression;
+	BuiltExpression.A = A.Get<FNiagaraStatelessExpression>().Build(BuildContext);
+	BuiltExpression.B = B.Get<FNiagaraStatelessExpression>().Build(BuildContext);
+	return FInstancedStruct::Make(BuiltExpression);
+}
+
+FLinearColor FNiagaraStatelessExpressionColorAdd::EvaluateInternal(const FEvaluateContext& Context) const
+{
+	const FLinearColor AValue = A.Get<FNiagaraStatelessExpressionColor>().EvaluateInternal(Context);
+	const FLinearColor BValue = B.Get<FNiagaraStatelessExpressionColor>().EvaluateInternal(Context);
+	return AValue + BValue;
+}
+
+bool FNiagaraStatelessExpressionColorAdd::IsConstant() const
+{
+	return A.Get<FNiagaraStatelessExpressionColor>().IsConstant() && B.Get<FNiagaraStatelessExpressionColor>().IsConstant();
+}
+
+////////////////////////////////////////////////////////////////////////////////////////////////////////////////////////
+
+FNiagaraStatelessExpressionColorSubtract::FNiagaraStatelessExpressionColorSubtract()
+{
+	A.InitializeAs<FNiagaraStatelessExpressionColorConstant>();
+	B.InitializeAs<FNiagaraStatelessExpressionColorConstant>();
+}
+
+FInstancedStruct FNiagaraStatelessExpressionColorSubtract::BuildInternal(const FNiagaraStatelessEmitterDataBuildContext& BuildContext) const
+{
+	FNiagaraStatelessExpressionColorSubtract BuiltExpression;
+	BuiltExpression.A = A.Get<FNiagaraStatelessExpression>().Build(BuildContext);
+	BuiltExpression.B = B.Get<FNiagaraStatelessExpression>().Build(BuildContext);
+	return FInstancedStruct::Make(BuiltExpression);
+}
+
+FLinearColor FNiagaraStatelessExpressionColorSubtract::EvaluateInternal(const FEvaluateContext& Context) const
+{
+	const FLinearColor AValue = A.Get<FNiagaraStatelessExpressionColor>().EvaluateInternal(Context);
+	const FLinearColor BValue = B.Get<FNiagaraStatelessExpressionColor>().EvaluateInternal(Context);
+	return AValue - BValue;
+}
+
+bool FNiagaraStatelessExpressionColorSubtract::IsConstant() const
+{
+	return A.Get<FNiagaraStatelessExpressionColor>().IsConstant() && B.Get<FNiagaraStatelessExpressionColor>().IsConstant();
+}
+
+////////////////////////////////////////////////////////////////////////////////////////////////////////////////////////
+
+FNiagaraStatelessExpressionColorMultiply::FNiagaraStatelessExpressionColorMultiply()
+{
+	A.InitializeAs<FNiagaraStatelessExpressionColorConstant>();
+	B.InitializeAs<FNiagaraStatelessExpressionColorConstant>();
+}
+
+FInstancedStruct FNiagaraStatelessExpressionColorMultiply::BuildInternal(const FNiagaraStatelessEmitterDataBuildContext& BuildContext) const
+{
+	FNiagaraStatelessExpressionColorMultiply BuiltExpression;
+	BuiltExpression.A = A.Get<FNiagaraStatelessExpression>().Build(BuildContext);
+	BuiltExpression.B = B.Get<FNiagaraStatelessExpression>().Build(BuildContext);
+	return FInstancedStruct::Make(BuiltExpression);
+}
+
+FLinearColor FNiagaraStatelessExpressionColorMultiply::EvaluateInternal(const FEvaluateContext& Context) const
+{
+	const FLinearColor AValue = A.Get<FNiagaraStatelessExpressionColor>().EvaluateInternal(Context);
+	const FLinearColor BValue = B.Get<FNiagaraStatelessExpressionColor>().EvaluateInternal(Context);
+	return AValue * BValue;
+}
+
+bool FNiagaraStatelessExpressionColorMultiply::IsConstant() const
+{
+	return A.Get<FNiagaraStatelessExpressionColor>().IsConstant() && B.Get<FNiagaraStatelessExpressionColor>().IsConstant();
+}
+
+////////////////////////////////////////////////////////////////////////////////////////////////////////////////////////
+
+FNiagaraStatelessExpressionColorDivide::FNiagaraStatelessExpressionColorDivide()
+{
+	A.InitializeAs<FNiagaraStatelessExpressionColorConstant>();
+	B.InitializeAs<FNiagaraStatelessExpressionColorConstant>();
+}
+
+FInstancedStruct FNiagaraStatelessExpressionColorDivide::BuildInternal(const FNiagaraStatelessEmitterDataBuildContext& BuildContext) const
+{
+	FNiagaraStatelessExpressionColorDivide BuiltExpression;
+	BuiltExpression.A = A.Get<FNiagaraStatelessExpression>().Build(BuildContext);
+	BuiltExpression.B = B.Get<FNiagaraStatelessExpression>().Build(BuildContext);
+	return FInstancedStruct::Make(BuiltExpression);
+}
+
+FLinearColor FNiagaraStatelessExpressionColorDivide::EvaluateInternal(const FEvaluateContext& Context) const
+{
+	const FLinearColor AValue = A.Get<FNiagaraStatelessExpressionColor>().EvaluateInternal(Context);
+	const FLinearColor BValue = B.Get<FNiagaraStatelessExpressionColor>().EvaluateInternal(Context);
+	return FLinearColor(
+		FMath::Abs(BValue.R) > UE_SMALL_NUMBER ? AValue.R / BValue.R : 0.0f,
+		FMath::Abs(BValue.G) > UE_SMALL_NUMBER ? AValue.G / BValue.G : 0.0f,
+		FMath::Abs(BValue.B) > UE_SMALL_NUMBER ? AValue.B / BValue.B : 0.0f,
+		FMath::Abs(BValue.A) > UE_SMALL_NUMBER ? AValue.A / BValue.A : 0.0f
+	);
+}
+
+bool FNiagaraStatelessExpressionColorDivide::IsConstant() const
+{
+	return A.Get<FNiagaraStatelessExpressionColor>().IsConstant() && B.Get<FNiagaraStatelessExpressionColor>().IsConstant();
+}
diff --git a/Engine/Plugins/FX/Niagara/Source/Niagara/Private/Stateless/Expressions/NiagaraStatelessExpressionFloat.cpp b/Engine/Plugins/FX/Niagara/Source/Niagara/Private/Stateless/Expressions/NiagaraStatelessExpressionFloat.cpp
new file mode 100644
--- /dev/null
+++ b/Engine/Plugins/FX/Niagara/Source/Niagara/Private/Stateless/Expressions/NiagaraStatelessExpressionFloat.cpp
@@ -0,0 +1,162 @@
+// Copyright Epic Games, Inc. All Rights Reserved.
+
+#include "Stateless/Expressions/NiagaraStatelessExpressionFloat.h"
+
+#include "Stateless/NiagaraStatelessEmitterDataBuildContext.h"
+#include "Stateless/NiagaraStatelessSetShaderParameterContext.h"
+#include "NiagaraParameterStore.h"
+
+////////////////////////////////////////////////////////////////////////////////////////////////////////////////////////
+
+FInstancedStruct FNiagaraStatelessExpressionFloat::Build(const FNiagaraStatelessEmitterDataBuildContext& BuildContext) const
+{
+	FInstancedStruct BuiltExpression;
+	if (IsConstant())
+	{
+		FNiagaraStatelessExpressionFloatConstant ConstantExpression;
+		FNiagaraParameterStore EmptyParameterStore;
+		ConstantExpression.A = EvaluateInternal(FEvaluateContext(EmptyParameterStore));
+		return FInstancedStruct::Make(ConstantExpression);
+	}
+	return BuildInternal(BuildContext);
+}
+
+////////////////////////////////////////////////////////////////////////////////////////////////////////////////////////
+
+float FNiagaraStatelessExpressionFloatConstant::EvaluateInternal(const FEvaluateContext& Context) const
+{
+	return A;
+}
+
+bool FNiagaraStatelessExpressionFloatConstant::IsConstant() const
+{
+	return true;
+}
+
+////////////////////////////////////////////////////////////////////////////////////////////////////////////////////////
+
+FInstancedStruct FNiagaraStatelessExpressionFloatBinding::BuildInternal(const FNiagaraStatelessEmitterDataBuildContext& BuildContext) const
+{
+	FNiagaraStatelessExpressionFloatBinding BuiltExpression;
+	BuiltExpression.A = A;
+	BuiltExpression.ParameterOffset = BuildContext.AddRendererBinding(FNiagaraVariableBase(GetOutputTypeDef(), A));
+	BuiltExpression.ParameterOffset *= sizeof(uint32);
+	return FInstancedStruct::Make(BuiltExpression);
+}
+
+float FNiagaraStatelessExpressionFloatBinding::EvaluateInternal(const FEvaluateContext& Context) const
+{
+	return Context.ParameterStore.GetParameterValueFromOffset<float>(ParameterOffset);
+}
+
+////////////////////////////////////////////////////////////////////////////////////////////////////////////////////////
+
+FNiagaraStatelessExpressionFloatAdd::FNiagaraStatelessExpressionFloatAdd()
+{
+	A.InitializeAs<FNiagaraStatelessExpressionFloatConstant>();
+	B.InitializeAs<FNiagaraStatelessExpressionFloatConstant>();
+}
+
+FInstancedStruct FNiagaraStatelessExpressionFloatAdd::BuildInternal(const FNiagaraStatelessEmitterDataBuildContext& BuildContext) const
+{
+	FNiagaraStatelessExpressionFloatAdd BuiltExpression;
+	BuiltExpression.A = A.Get<FNiagaraStatelessExpression>().Build(BuildContext);
+	BuiltExpression.B = B.Get<FNiagaraStatelessExpression>().Build(BuildContext);
+	return FInstancedStruct::Make(BuiltExpression);
+}
+
+float FNiagaraStatelessExpressionFloatAdd::EvaluateInternal(const FEvaluateContext& Context) const
+{
+	const float AValue = A.Get<FNiagaraStatelessExpressionFloat>().EvaluateInternal(Context);
+	const float BValue = B.Get<FNiagaraStatelessExpressionFloat>().EvaluateInternal(Context);
+	return AValue + BValue;
+}
+
+bool FNiagaraStatelessExpressionFloatAdd::IsConstant() const
+{
+	return A.Get<FNiagaraStatelessExpressionFloat>().IsConstant() && B.Get<FNiagaraStatelessExpressionFloat>().IsConstant();
+}
+
+////////////////////////////////////////////////////////////////////////////////////////////////////////////////////////
+
+FNiagaraStatelessExpressionFloatSubtract::FNiagaraStatelessExpressionFloatSubtract()
+{
+	A.InitializeAs<FNiagaraStatelessExpressionFloatConstant>();
+	B.InitializeAs<FNiagaraStatelessExpressionFloatConstant>();
+}
+
+FInstancedStruct FNiagaraStatelessExpressionFloatSubtract::BuildInternal(const FNiagaraStatelessEmitterDataBuildContext& BuildContext) const
+{
+	FNiagaraStatelessExpressionFloatSubtract BuiltExpression;
+	BuiltExpression.A = A.Get<FNiagaraStatelessExpression>().Build(BuildContext);
+	BuiltExpression.B = B.Get<FNiagaraStatelessExpression>().Build(BuildContext);
+	return FInstancedStruct::Make(BuiltExpression);
+}
+
+float FNiagaraStatelessExpressionFloatSubtract::EvaluateInternal(const FEvaluateContext& Context) const
+{
+	const float AValue = A.Get<FNiagaraStatelessExpressionFloat>().EvaluateInternal(Context);
+	const float BValue = B.Get<FNiagaraStatelessExpressionFloat>().EvaluateInternal(Context);
+	return AValue - BValue;
+}
+
+bool FNiagaraStatelessExpressionFloatSubtract::IsConstant() const
+{
+	return A.Get<FNiagaraStatelessExpressionFloat>().IsConstant() && B.Get<FNiagaraStatelessExpressionFloat>().IsConstant();
+}
+
+////////////////////////////////////////////////////////////////////////////////////////////////////////////////////////
+
+FNiagaraStatelessExpressionFloatMultiply::FNiagaraStatelessExpressionFloatMultiply()
+{
+	A.InitializeAs<FNiagaraStatelessExpressionFloatConstant>();
+	B.InitializeAs<FNiagaraStatelessExpressionFloatConstant>();
+}
+
+FInstancedStruct FNiagaraStatelessExpressionFloatMultiply::BuildInternal(const FNiagaraStatelessEmitterDataBuildContext& BuildContext) const
+{
+	FNiagaraStatelessExpressionFloatMultiply BuiltExpression;
+	BuiltExpression.A = A.Get<FNiagaraStatelessExpression>().Build(BuildContext);
+	BuiltExpression.B = B.Get<FNiagaraStatelessExpression>().Build(BuildContext);
+	return FInstancedStruct::Make(BuiltExpression);
+}
+
+float FNiagaraStatelessExpressionFloatMultiply::EvaluateInternal(const FEvaluateContext& Context) const
+{
+	const float AValue = A.Get<FNiagaraStatelessExpressionFloat>().EvaluateInternal(Context);
+	const float BValue = B.Get<FNiagaraStatelessExpressionFloat>().EvaluateInternal(Context);
+	return AValue * BValue;
+}
+
+bool FNiagaraStatelessExpressionFloatMultiply::IsConstant() const
+{
+	return A.Get<FNiagaraStatelessExpressionFloat>().IsConstant() && B.Get<FNiagaraStatelessExpressionFloat>().IsConstant();
+}
+
+////////////////////////////////////////////////////////////////////////////////////////////////////////////////////////
+
+FNiagaraStatelessExpressionFloatDivide::FNiagaraStatelessExpressionFloatDivide()
+{
+	A.InitializeAs<FNiagaraStatelessExpressionFloatConstant>();
+	B.InitializeAs<FNiagaraStatelessExpressionFloatConstant>();
+}
+
+FInstancedStruct FNiagaraStatelessExpressionFloatDivide::BuildInternal(const FNiagaraStatelessEmitterDataBuildContext& BuildContext) const
+{
+	FNiagaraStatelessExpressionFloatDivide BuiltExpression;
+	BuiltExpression.A = A.Get<FNiagaraStatelessExpression>().Build(BuildContext);
+	BuiltExpression.B = B.Get<FNiagaraStatelessExpression>().Build(BuildContext);
+	return FInstancedStruct::Make(BuiltExpression);
+}
+
+float FNiagaraStatelessExpressionFloatDivide::EvaluateInternal(const FEvaluateContext& Context) const
+{
+	const float AValue = A.Get<FNiagaraStatelessExpressionFloat>().EvaluateInternal(Context);
+	const float BValue = B.Get<FNiagaraStatelessExpressionFloat>().EvaluateInternal(Context);
+	return FMath::Abs(BValue) > UE_SMALL_NUMBER ? AValue / BValue : 0.0f;
+}
+
+bool FNiagaraStatelessExpressionFloatDivide::IsConstant() const
+{
+	return A.Get<FNiagaraStatelessExpressionFloat>().IsConstant() && B.Get<FNiagaraStatelessExpressionFloat>().IsConstant();
+}
diff --git a/Engine/Plugins/FX/Niagara/Source/Niagara/Private/Stateless/Expressions/NiagaraStatelessExpressionVec2.cpp b/Engine/Plugins/FX/Niagara/Source/Niagara/Private/Stateless/Expressions/NiagaraStatelessExpressionVec2.cpp
new file mode 100644
--- /dev/null
+++ b/Engine/Plugins/FX/Niagara/Source/Niagara/Private/Stateless/Expressions/NiagaraStatelessExpressionVec2.cpp
@@ -0,0 +1,165 @@
+// Copyright Epic Games, Inc. All Rights Reserved.
+
+#include "Stateless/Expressions/NiagaraStatelessExpressionVec2.h"
+
+#include "Stateless/NiagaraStatelessEmitterDataBuildContext.h"
+#include "Stateless/NiagaraStatelessSetShaderParameterContext.h"
+#include "NiagaraParameterStore.h"
+
+////////////////////////////////////////////////////////////////////////////////////////////////////////////////////////
+
+FInstancedStruct FNiagaraStatelessExpressionVec2::Build(const FNiagaraStatelessEmitterDataBuildContext& BuildContext) const
+{
+	FInstancedStruct BuiltExpression;
+	if (IsConstant())
+	{
+		FNiagaraStatelessExpressionVec2Constant ConstantExpression;
+		FNiagaraParameterStore EmptyParameterStore;
+		ConstantExpression.A = EvaluateInternal(FEvaluateContext(EmptyParameterStore));
+		return FInstancedStruct::Make(ConstantExpression);
+	}
+	return BuildInternal(BuildContext);
+}
+
+////////////////////////////////////////////////////////////////////////////////////////////////////////////////////////
+
+FVector2f FNiagaraStatelessExpressionVec2Constant::EvaluateInternal(const FEvaluateContext& Context) const
+{
+	return A;
+}
+
+bool FNiagaraStatelessExpressionVec2Constant::IsConstant() const
+{
+	return true;
+}
+
+////////////////////////////////////////////////////////////////////////////////////////////////////////////////////////
+
+FInstancedStruct FNiagaraStatelessExpressionVec2Binding::BuildInternal(const FNiagaraStatelessEmitterDataBuildContext& BuildContext) const
+{
+	FNiagaraStatelessExpressionVec2Binding BuiltExpression;
+	BuiltExpression.A = A;
+	BuiltExpression.ParameterOffset = BuildContext.AddRendererBinding(FNiagaraVariableBase(GetOutputTypeDef(), A));
+	BuiltExpression.ParameterOffset *= sizeof(uint32);
+	return FInstancedStruct::Make(BuiltExpression);
+}
+
+FVector2f FNiagaraStatelessExpressionVec2Binding::EvaluateInternal(const FEvaluateContext& Context) const
+{
+	return Context.ParameterStore.GetParameterValueFromOffset<FVector2f>(ParameterOffset);
+}
+
+////////////////////////////////////////////////////////////////////////////////////////////////////////////////////////
+
+FNiagaraStatelessExpressionVec2Add::FNiagaraStatelessExpressionVec2Add()
+{
+	A.InitializeAs<FNiagaraStatelessExpressionVec2Constant>();
+	B.InitializeAs<FNiagaraStatelessExpressionVec2Constant>();
+}
+
+FInstancedStruct FNiagaraStatelessExpressionVec2Add::BuildInternal(const FNiagaraStatelessEmitterDataBuildContext& BuildContext) const
+{
+	FNiagaraStatelessExpressionVec2Add BuiltExpression;
+	BuiltExpression.A = A.Get<FNiagaraStatelessExpression>().Build(BuildContext);
+	BuiltExpression.B = B.Get<FNiagaraStatelessExpression>().Build(BuildContext);
+	return FInstancedStruct::Make(BuiltExpression);
+}
+
+FVector2f FNiagaraStatelessExpressionVec2Add::EvaluateInternal(const FEvaluateContext& Context) const
+{
+	const FVector2f AValue = A.Get<FNiagaraStatelessExpressionVec2>().EvaluateInternal(Context);
+	const FVector2f BValue = B.Get<FNiagaraStatelessExpressionVec2>().EvaluateInternal(Context);
+	return AValue + BValue;
+}
+
+bool FNiagaraStatelessExpressionVec2Add::IsConstant() const
+{
+	return A.Get<FNiagaraStatelessExpressionVec2>().IsConstant() && B.Get<FNiagaraStatelessExpressionVec2>().IsConstant();
+}
+
+////////////////////////////////////////////////////////////////////////////////////////////////////////////////////////
+
+FNiagaraStatelessExpressionVec2Subtract::FNiagaraStatelessExpressionVec2Subtract()
+{
+	A.InitializeAs<FNiagaraStatelessExpressionVec2Constant>();
+	B.InitializeAs<FNiagaraStatelessExpressionVec2Constant>();
+}
+
+FInstancedStruct FNiagaraStatelessExpressionVec2Subtract::BuildInternal(const FNiagaraStatelessEmitterDataBuildContext& BuildContext) const
+{
+	FNiagaraStatelessExpressionVec2Subtract BuiltExpression;
+	BuiltExpression.A = A.Get<FNiagaraStatelessExpression>().Build(BuildContext);
+	BuiltExpression.B = B.Get<FNiagaraStatelessExpression>().Build(BuildContext);
+	return FInstancedStruct::Make(BuiltExpression);
+}
+
+FVector2f FNiagaraStatelessExpressionVec2Subtract::EvaluateInternal(const FEvaluateContext& Context) const
+{
+	const FVector2f AValue = A.Get<FNiagaraStatelessExpressionVec2>().EvaluateInternal(Context);
+	const FVector2f BValue = B.Get<FNiagaraStatelessExpressionVec2>().EvaluateInternal(Context);
+	return AValue - BValue;
+}
+
+bool FNiagaraStatelessExpressionVec2Subtract::IsConstant() const
+{
+	return A.Get<FNiagaraStatelessExpressionVec2>().IsConstant() && B.Get<FNiagaraStatelessExpressionVec2>().IsConstant();
+}
+
+////////////////////////////////////////////////////////////////////////////////////////////////////////////////////////
+
+FNiagaraStatelessExpressionVec2Multiply::FNiagaraStatelessExpressionVec2Multiply()
+{
+	A.InitializeAs<FNiagaraStatelessExpressionVec2Constant>();
+	B.InitializeAs<FNiagaraStatelessExpressionVec2Constant>();
+}
+
+FInstancedStruct FNiagaraStatelessExpressionVec2Multiply::BuildInternal(const FNiagaraStatelessEmitterDataBuildContext& BuildContext) const
+{
+	FNiagaraStatelessExpressionVec2Multiply BuiltExpression;
+	BuiltExpression.A = A.Get<FNiagaraStatelessExpression>().Build(BuildContext);
+	BuiltExpression.B = B.Get<FNiagaraStatelessExpression>().Build(BuildContext);
+	return FInstancedStruct::Make(BuiltExpression);
+}
+
+FVector2f FNiagaraStatelessExpressionVec2Multiply::EvaluateInternal(const FEvaluateContext& Context) const
+{
+	const FVector2f AValue = A.Get<FNiagaraStatelessExpressionVec2>().EvaluateInternal(Context);
+	const FVector2f BValue = B.Get<FNiagaraStatelessExpressionVec2>().EvaluateInternal(Context);
+	return AValue * BValue;
+}
+
+bool FNiagaraStatelessExpressionVec2Multiply::IsConstant() const
+{
+	return A.Get<FNiagaraStatelessExpressionVec2>().IsConstant() && B.Get<FNiagaraStatelessExpressionVec2>().IsConstant();
+}
+
+////////////////////////////////////////////////////////////////////////////////////////////////////////////////////////
+
+FNiagaraStatelessExpressionVec2Divide::FNiagaraStatelessExpressionVec2Divide()
+{
+	A.InitializeAs<FNiagaraStatelessExpressionVec2Constant>();
+	B.InitializeAs<FNiagaraStatelessExpressionVec2Constant>();
+}
+
+FInstancedStruct FNiagaraStatelessExpressionVec2Divide::BuildInternal(const FNiagaraStatelessEmitterDataBuildContext& BuildContext) const
+{
+	FNiagaraStatelessExpressionVec2Divide BuiltExpression;
+	BuiltExpression.A = A.Get<FNiagaraStatelessExpression>().Build(BuildContext);
+	BuiltExpression.B = B.Get<FNiagaraStatelessExpression>().Build(BuildContext);
+	return FInstancedStruct::Make(BuiltExpression);
+}
+
+FVector2f FNiagaraStatelessExpressionVec2Divide::EvaluateInternal(const FEvaluateContext& Context) const
+{
+	const FVector2f AValue = A.Get<FNiagaraStatelessExpressionVec2>().EvaluateInternal(Context);
+	const FVector2f BValue = B.Get<FNiagaraStatelessExpressionVec2>().EvaluateInternal(Context);
+	return FVector2f(
+		FMath::Abs(BValue.X) > UE_SMALL_NUMBER ? AValue.X / BValue.X : 0.0f,
+		FMath::Abs(BValue.Y) > UE_SMALL_NUMBER ? AValue.Y / BValue.Y : 0.0f
+	);
+}
+
+bool FNiagaraStatelessExpressionVec2Divide::IsConstant() const
+{
+	return A.Get<FNiagaraStatelessExpressionVec2>().IsConstant() && B.Get<FNiagaraStatelessExpressionVec2>().IsConstant();
+}
diff --git a/Engine/Plugins/FX/Niagara/Source/Niagara/Private/Stateless/Expressions/NiagaraStatelessExpressionVec3.cpp b/Engine/Plugins/FX/Niagara/Source/Niagara/Private/Stateless/Expressions/NiagaraStatelessExpressionVec3.cpp
new file mode 100644
--- /dev/null
+++ b/Engine/Plugins/FX/Niagara/Source/Niagara/Private/Stateless/Expressions/NiagaraStatelessExpressionVec3.cpp
@@ -0,0 +1,166 @@
+// Copyright Epic Games, Inc. All Rights Reserved.
+
+#include "Stateless/Expressions/NiagaraStatelessExpressionVec3.h"
+
+#include "Stateless/NiagaraStatelessEmitterDataBuildContext.h"
+#include "Stateless/NiagaraStatelessSetShaderParameterContext.h"
+#include "NiagaraParameterStore.h"
+
+////////////////////////////////////////////////////////////////////////////////////////////////////////////////////////
+
+FInstancedStruct FNiagaraStatelessExpressionVec3::Build(const FNiagaraStatelessEmitterDataBuildContext& BuildContext) const
+{
+	FInstancedStruct BuiltExpression;
+	if (IsConstant())
+	{
+		FNiagaraStatelessExpressionVec3Constant ConstantExpression;
+		FNiagaraParameterStore EmptyParameterStore;
+		ConstantExpression.A = EvaluateInternal(FEvaluateContext(EmptyParameterStore));
+		return FInstancedStruct::Make(ConstantExpression);
+	}
+	return BuildInternal(BuildContext);
+}
+
+////////////////////////////////////////////////////////////////////////////////////////////////////////////////////////
+
+FVector3f FNiagaraStatelessExpressionVec3Constant::EvaluateInternal(const FEvaluateContext& Context) const
+{
+	return A;
+}
+
+bool FNiagaraStatelessExpressionVec3Constant::IsConstant() const
+{
+	return true;
+}
+
+////////////////////////////////////////////////////////////////////////////////////////////////////////////////////////
+
+FInstancedStruct FNiagaraStatelessExpressionVec3Binding::BuildInternal(const FNiagaraStatelessEmitterDataBuildContext& BuildContext) const
+{
+	FNiagaraStatelessExpressionVec3Binding BuiltExpression;
+	BuiltExpression.A = A;
+	BuiltExpression.ParameterOffset = BuildContext.AddRendererBinding(FNiagaraVariableBase(GetOutputTypeDef(), A));
+	BuiltExpression.ParameterOffset *= sizeof(uint32);
+	return FInstancedStruct::Make(BuiltExpression);
+}
+
+FVector3f FNiagaraStatelessExpressionVec3Binding::EvaluateInternal(const FEvaluateContext& Context) const
+{
+	return Context.ParameterStore.GetParameterValueFromOffset<FVector3f>(ParameterOffset);
+}
+
+////////////////////////////////////////////////////////////////////////////////////////////////////////////////////////
+
+FNiagaraStatelessExpressionVec3Add::FNiagaraStatelessExpressionVec3Add()
+{
+	A.InitializeAs<FNiagaraStatelessExpressionVec3Constant>();
+	B.InitializeAs<FNiagaraStatelessExpressionVec3Constant>();
+}
+
+FInstancedStruct FNiagaraStatelessExpressionVec3Add::BuildInternal(const FNiagaraStatelessEmitterDataBuildContext& BuildContext) const
+{
+	FNiagaraStatelessExpressionVec3Add BuiltExpression;
+	BuiltExpression.A = A.Get<FNiagaraStatelessExpression>().Build(BuildContext);
+	BuiltExpression.B = B.Get<FNiagaraStatelessExpression>().Build(BuildContext);
+	return FInstancedStruct::Make(BuiltExpression);
+}
+
+FVector3f FNiagaraStatelessExpressionVec3Add::EvaluateInternal(const FEvaluateContext& Context) const
+{
+	const FVector3f AValue = A.Get<FNiagaraStatelessExpressionVec3>().EvaluateInternal(Context);
+	const FVector3f BValue = B.Get<FNiagaraStatelessExpressionVec3>().EvaluateInternal(Context);
+	return AValue + BValue;
+}
+
+bool FNiagaraStatelessExpressionVec3Add::IsConstant() const
+{
+	return A.Get<FNiagaraStatelessExpressionVec3>().IsConstant() && B.Get<FNiagaraStatelessExpressionVec3>().IsConstant();
+}
+
+////////////////////////////////////////////////////////////////////////////////////////////////////////////////////////
+
+FNiagaraStatelessExpressionVec3Subtract::FNiagaraStatelessExpressionVec3Subtract()
+{
+	A.InitializeAs<FNiagaraStatelessExpressionVec3Constant>();
+	B.InitializeAs<FNiagaraStatelessExpressionVec3Constant>();
+}
+
+FInstancedStruct FNiagaraStatelessExpressionVec3Subtract::BuildInternal(const FNiagaraStatelessEmitterDataBuildContext& BuildContext) const
+{
+	FNiagaraStatelessExpressionVec3Subtract BuiltExpression;
+	BuiltExpression.A = A.Get<FNiagaraStatelessExpression>().Build(BuildContext);
+	BuiltExpression.B = B.Get<FNiagaraStatelessExpression>().Build(BuildContext);
+	return FInstancedStruct::Make(BuiltExpression);
+}
+
+FVector3f FNiagaraStatelessExpressionVec3Subtract::EvaluateInternal(const FEvaluateContext& Context) const
+{
+	const FVector3f AValue = A.Get<FNiagaraStatelessExpressionVec3>().EvaluateInternal(Context);
+	const FVector3f BValue = B.Get<FNiagaraStatelessExpressionVec3>().EvaluateInternal(Context);
+	return AValue - BValue;
+}
+
+bool FNiagaraStatelessExpressionVec3Subtract::IsConstant() const
+{
+	return A.Get<FNiagaraStatelessExpressionVec3>().IsConstant() && B.Get<FNiagaraStatelessExpressionVec3>().IsConstant();
+}
+
+////////////////////////////////////////////////////////////////////////////////////////////////////////////////////////
+
+FNiagaraStatelessExpressionVec3Multiply::FNiagaraStatelessExpressionVec3Multiply()
+{
+	A.InitializeAs<FNiagaraStatelessExpressionVec3Constant>();
+	B.InitializeAs<FNiagaraStatelessExpressionVec3Constant>();
+}
+
+FInstancedStruct FNiagaraStatelessExpressionVec3Multiply::BuildInternal(const FNiagaraStatelessEmitterDataBuildContext& BuildContext) const
+{
+	FNiagaraStatelessExpressionVec3Multiply BuiltExpression;
+	BuiltExpression.A = A.Get<FNiagaraStatelessExpression>().Build(BuildContext);
+	BuiltExpression.B = B.Get<FNiagaraStatelessExpression>().Build(BuildContext);
+	return FInstancedStruct::Make(BuiltExpression);
+}
+
+FVector3f FNiagaraStatelessExpressionVec3Multiply::EvaluateInternal(const FEvaluateContext& Context) const
+{
+	const FVector3f AValue = A.Get<FNiagaraStatelessExpressionVec3>().EvaluateInternal(Context);
+	const FVector3f BValue = B.Get<FNiagaraStatelessExpressionVec3>().EvaluateInternal(Context);
+	return AValue * BValue;
+}
+
+bool FNiagaraStatelessExpressionVec3Multiply::IsConstant() const
+{
+	return A.Get<FNiagaraStatelessExpressionVec3>().IsConstant() && B.Get<FNiagaraStatelessExpressionVec3>().IsConstant();
+}
+
+////////////////////////////////////////////////////////////////////////////////////////////////////////////////////////
+
+FNiagaraStatelessExpressionVec3Divide::FNiagaraStatelessExpressionVec3Divide()
+{
+	A.InitializeAs<FNiagaraStatelessExpressionVec3Constant>();
+	B.InitializeAs<FNiagaraStatelessExpressionVec3Constant>();
+}
+
+FInstancedStruct FNiagaraStatelessExpressionVec3Divide::BuildInternal(const FNiagaraStatelessEmitterDataBuildContext& BuildContext) const
+{
+	FNiagaraStatelessExpressionVec3Divide BuiltExpression;
+	BuiltExpression.A = A.Get<FNiagaraStatelessExpression>().Build(BuildContext);
+	BuiltExpression.B = B.Get<FNiagaraStatelessExpression>().Build(BuildContext);
+	return FInstancedStruct::Make(BuiltExpression);
+}
+
+FVector3f FNiagaraStatelessExpressionVec3Divide::EvaluateInternal(const FEvaluateContext& Context) const
+{
+	const FVector3f AValue = A.Get<FNiagaraStatelessExpressionVec3>().EvaluateInternal(Context);
+	const FVector3f BValue = B.Get<FNiagaraStatelessExpressionVec3>().EvaluateInternal(Context);
+	return FVector3f(
+		FMath::Abs(BValue.X) > UE_SMALL_NUMBER ? AValue.X / BValue.X : 0.0f,
+		FMath::Abs(BValue.Y) > UE_SMALL_NUMBER ? AValue.Y / BValue.Y : 0.0f,
+		FMath::Abs(BValue.Z) > UE_SMALL_NUMBER ? AValue.Z / BValue.Z : 0.0f
+	);
+}
+
+bool FNiagaraStatelessExpressionVec3Divide::IsConstant() const
+{
+	return A.Get<FNiagaraStatelessExpressionVec3>().IsConstant() && B.Get<FNiagaraStatelessExpressionVec3>().IsConstant();
+}
diff --git a/Engine/Plugins/FX/Niagara/Source/Niagara/Private/Stateless/Expressions/NiagaraStatelessExpressionVec4.cpp b/Engine/Plugins/FX/Niagara/Source/Niagara/Private/Stateless/Expressions/NiagaraStatelessExpressionVec4.cpp
new file mode 100644
--- /dev/null
+++ b/Engine/Plugins/FX/Niagara/Source/Niagara/Private/Stateless/Expressions/NiagaraStatelessExpressionVec4.cpp
@@ -0,0 +1,167 @@
+// Copyright Epic Games, Inc. All Rights Reserved.
+
+#include "Stateless/Expressions/NiagaraStatelessExpressionVec4.h"
+
+#include "Stateless/NiagaraStatelessEmitterDataBuildContext.h"
+#include "Stateless/NiagaraStatelessSetShaderParameterContext.h"
+#include "NiagaraParameterStore.h"
+
+////////////////////////////////////////////////////////////////////////////////////////////////////////////////////////
+
+FInstancedStruct FNiagaraStatelessExpressionVec4::Build(const FNiagaraStatelessEmitterDataBuildContext& BuildContext) const
+{
+	FInstancedStruct BuiltExpression;
+	if (IsConstant())
+	{
+		FNiagaraStatelessExpressionVec4Constant ConstantExpression;
+		FNiagaraParameterStore EmptyParameterStore;
+		ConstantExpression.A = EvaluateInternal(FEvaluateContext(EmptyParameterStore));
+		return FInstancedStruct::Make(ConstantExpression);
+	}
+	return BuildInternal(BuildContext);
+}
+
+////////////////////////////////////////////////////////////////////////////////////////////////////////////////////////
+
+FVector4f FNiagaraStatelessExpressionVec4Constant::EvaluateInternal(const FEvaluateContext& Context) const
+{
+	return A;
+}
+
+bool FNiagaraStatelessExpressionVec4Constant::IsConstant() const
+{
+	return true;
+}
+
+////////////////////////////////////////////////////////////////////////////////////////////////////////////////////////
+
+FInstancedStruct FNiagaraStatelessExpressionVec4Binding::BuildInternal(const FNiagaraStatelessEmitterDataBuildContext& BuildContext) const
+{
+	FNiagaraStatelessExpressionVec4Binding BuiltExpression;
+	BuiltExpression.A = A;
+	BuiltExpression.ParameterOffset = BuildContext.AddRendererBinding(FNiagaraVariableBase(GetOutputTypeDef(), A));
+	BuiltExpression.ParameterOffset *= sizeof(uint32);
+	return FInstancedStruct::Make(BuiltExpression);
+}
+
+FVector4f FNiagaraStatelessExpressionVec4Binding::EvaluateInternal(const FEvaluateContext& Context) const
+{
+	return Context.ParameterStore.GetParameterValueFromOffset<FVector4f>(ParameterOffset);
+}
+
+////////////////////////////////////////////////////////////////////////////////////////////////////////////////////////
+
+FNiagaraStatelessExpressionVec4Add::FNiagaraStatelessExpressionVec4Add()
+{
+	A.InitializeAs<FNiagaraStatelessExpressionVec4Constant>();
+	B.InitializeAs<FNiagaraStatelessExpressionVec4Constant>();
+}
+
+FInstancedStruct FNiagaraStatelessExpressionVec4Add::BuildInternal(const FNiagaraStatelessEmitterDataBuildContext& BuildContext) const
+{
+	FNiagaraStatelessExpressionVec4Add BuiltExpression;
+	BuiltExpression.A = A.Get<FNiagaraStatelessExpression>().Build(BuildContext);
+	BuiltExpression.B = B.Get<FNiagaraStatelessExpression>().Build(BuildContext);
+	return FInstancedStruct::Make(BuiltExpression);
+}
+
+FVector4f FNiagaraStatelessExpressionVec4Add::EvaluateInternal(const FEvaluateContext& Context) const
+{
+	const FVector4f AValue = A.Get<FNiagaraStatelessExpressionVec4>().EvaluateInternal(Context);
+	const FVector4f BValue = B.Get<FNiagaraStatelessExpressionVec4>().EvaluateInternal(Context);
+	return AValue + BValue;
+}
+
+bool FNiagaraStatelessExpressionVec4Add::IsConstant() const
+{
+	return A.Get<FNiagaraStatelessExpressionVec4>().IsConstant() && B.Get<FNiagaraStatelessExpressionVec4>().IsConstant();
+}
+
+////////////////////////////////////////////////////////////////////////////////////////////////////////////////////////
+
+FNiagaraStatelessExpressionVec4Subtract::FNiagaraStatelessExpressionVec4Subtract()
+{
+	A.InitializeAs<FNiagaraStatelessExpressionVec4Constant>();
+	B.InitializeAs<FNiagaraStatelessExpressionVec4Constant>();
+}
+
+FInstancedStruct FNiagaraStatelessExpressionVec4Subtract::BuildInternal(const FNiagaraStatelessEmitterDataBuildContext& BuildContext) const
+{
+	FNiagaraStatelessExpressionVec4Subtract BuiltExpression;
+	BuiltExpression.A = A.Get<FNiagaraStatelessExpression>().Build(BuildContext);
+	BuiltExpression.B = B.Get<FNiagaraStatelessExpression>().Build(BuildContext);
+	return FInstancedStruct::Make(BuiltExpression);
+}
+
+FVector4f FNiagaraStatelessExpressionVec4Subtract::EvaluateInternal(const FEvaluateContext& Context) const
+{
+	const FVector4f AValue = A.Get<FNiagaraStatelessExpressionVec4>().EvaluateInternal(Context);
+	const FVector4f BValue = B.Get<FNiagaraStatelessExpressionVec4>().EvaluateInternal(Context);
+	return AValue - BValue;
+}
+
+bool FNiagaraStatelessExpressionVec4Subtract::IsConstant() const
+{
+	return A.Get<FNiagaraStatelessExpressionVec4>().IsConstant() && B.Get<FNiagaraStatelessExpressionVec4>().IsConstant();
+}
+
+////////////////////////////////////////////////////////////////////////////////////////////////////////////////////////
+
+FNiagaraStatelessExpressionVec4Multiply::FNiagaraStatelessExpressionVec4Multiply()
+{
+	A.InitializeAs<FNiagaraStatelessExpressionVec4Constant>();
+	B.InitializeAs<FNiagaraStatelessExpressionVec4Constant>();
+}
+
+FInstancedStruct FNiagaraStatelessExpressionVec4Multiply::BuildInternal(const FNiagaraStatelessEmitterDataBuildContext& BuildContext) const
+{
+	FNiagaraStatelessExpressionVec4Multiply BuiltExpression;
+	BuiltExpression.A = A.Get<FNiagaraStatelessExpression>().Build(BuildContext);
+	BuiltExpression.B = B.Get<FNiagaraStatelessExpression>().Build(BuildContext);
+	return FInstancedStruct::Make(BuiltExpression);
+}
+
+FVector4f FNiagaraStatelessExpressionVec4Multiply::EvaluateInternal(const FEvaluateContext& Context) const
+{
+	const FVector4f AValue = A.Get<FNiagaraStatelessExpressionVec4>().EvaluateInternal(Context);
+	const FVector4f BValue = B.Get<FNiagaraStatelessExpressionVec4>().EvaluateInternal(Context);
+	return AValue * BValue;
+}
+
+bool FNiagaraStatelessExpressionVec4Multiply::IsConstant() const
+{
+	return A.Get<FNiagaraStatelessExpressionVec4>().IsConstant() && B.Get<FNiagaraStatelessExpressionVec4>().IsConstant();
+}
+
+////////////////////////////////////////////////////////////////////////////////////////////////////////////////////////
+
+FNiagaraStatelessExpressionVec4Divide::FNiagaraStatelessExpressionVec4Divide()
+{
+	A.InitializeAs<FNiagaraStatelessExpressionVec4Constant>();
+	B.InitializeAs<FNiagaraStatelessExpressionVec4Constant>();
+}
+
+FInstancedStruct FNiagaraStatelessExpressionVec4Divide::BuildInternal(const FNiagaraStatelessEmitterDataBuildContext& BuildContext) const
+{
+	FNiagaraStatelessExpressionVec4Divide BuiltExpression;
+	BuiltExpression.A = A.Get<FNiagaraStatelessExpression>().Build(BuildContext);
+	BuiltExpression.B = B.Get<FNiagaraStatelessExpression>().Build(BuildContext);
+	return FInstancedStruct::Make(BuiltExpression);
+}
+
+FVector4f FNiagaraStatelessExpressionVec4Divide::EvaluateInternal(const FEvaluateContext& Context) const
+{
+	const FVector4f AValue = A.Get<FNiagaraStatelessExpressionVec4>().EvaluateInternal(Context);
+	const FVector4f BValue = B.Get<FNiagaraStatelessExpressionVec4>().EvaluateInternal(Context);
+	return FVector4f(
+		FMath::Abs(BValue.X) > UE_SMALL_NUMBER ? AValue.X / BValue.X : 0.0f,
+		FMath::Abs(BValue.Y) > UE_SMALL_NUMBER ? AValue.Y / BValue.Y : 0.0f,
+		FMath::Abs(BValue.Z) > UE_SMALL_NUMBER ? AValue.Z / BValue.Z : 0.0f,
+		FMath::Abs(BValue.W) > UE_SMALL_NUMBER ? AValue.W / BValue.W : 0.0f
+	);
+}
+
+bool FNiagaraStatelessExpressionVec4Divide::IsConstant() const
+{
+	return A.Get<FNiagaraStatelessExpressionVec4>().IsConstant() && B.Get<FNiagaraStatelessExpressionVec4>().IsConstant();
+}
diff --git a/Engine/Plugins/FX/Niagara/Source/Niagara/Private/Stateless/NiagaraStatelessDistribution.cpp b/Engine/Plugins/FX/Niagara/Source/Niagara/Private/Stateless/NiagaraStatelessDistribution.cpp
--- a/Engine/Plugins/FX/Niagara/Source/Niagara/Private/Stateless/NiagaraStatelessDistribution.cpp
+++ b/Engine/Plugins/FX/Niagara/Source/Niagara/Private/Stateless/NiagaraStatelessDistribution.cpp
@@ -1,6 +1,7 @@
 // Copyright Epic Games, Inc. All Rights Reserved.
 
 #include "Stateless/NiagaraStatelessDistribution.h"
+#include "Stateless/NiagaraStatelessExpression.h"
 #include "HAL/IConsoleManager.h"
 
 namespace NiagaraStatelessDistributionPrivate
@@ -643,6 +644,29 @@ FNiagaraDistributionCurveVector3::FNiagaraDistributionCurveVector3(ENiagaraDistr
 }
 
 #if WITH_EDITORONLY_DATA
+bool FNiagaraDistributionBase::operator==(const FNiagaraDistributionBase& Other) const
+{
+	return
+		Mode == Other.Mode &&
+		ParameterBinding == Other.ParameterBinding &&
+		ParameterExpression.Identical(&Other.ParameterExpression, PPF_DeepComparison) &&
+		ChannelConstantsAndRanges == Other.ChannelConstantsAndRanges &&
+		ChannelCurves == Other.ChannelCurves &&
+		MaxLutSampleCount == Other.MaxLutSampleCount;
+}
+
+void FNiagaraDistributionBase::ForEachParameterBinding(TFunction<void(const FNiagaraVariableBase&)> Delegate) const
+{
+	if (IsBinding())
+	{
+		Delegate(ParameterBinding);
+	}
+	else if (IsExpression())
+	{
+		FNiagaraStatelessExpression::ForEachBinding(ParameterExpression, Delegate);
+	}
+}
+
 void FNiagaraDistributionBase::PostEditChangeProperty(UObject* OwnerObject, FPropertyChangedEvent& PropertyChangedEvent)
 {
 	FStructProperty* StructProperty = CastField<FStructProperty>(PropertyChangedEvent.Property);
diff --git a/Engine/Plugins/FX/Niagara/Source/Niagara/Private/Stateless/NiagaraStatelessEmitter.cpp b/Engine/Plugins/FX/Niagara/Source/Niagara/Private/Stateless/NiagaraStatelessEmitter.cpp
--- a/Engine/Plugins/FX/Niagara/Source/Niagara/Private/Stateless/NiagaraStatelessEmitter.cpp
+++ b/Engine/Plugins/FX/Niagara/Source/Niagara/Private/Stateless/NiagaraStatelessEmitter.cpp
@@ -255,26 +255,23 @@ void UNiagaraStatelessEmitter::CacheParameterCollectionReferences()
 			}
 
 			const FNiagaraDistributionBase* DistributionBase = StructProp->ContainerPtrToValuePtr<const FNiagaraDistributionBase>(Module);
-			if (!DistributionBase->IsBinding())
-			{
-				continue;
-			}
-
-
-			if (!DistributionBase->ParameterBinding.IsInNameSpace(FNiagaraConstants::ParameterCollectionNamespaceString))
-			{
-				continue;
-			}
-
-			for (TObjectIterator<UNiagaraParameterCollection> NPCIt; NPCIt; ++NPCIt)
-			{
-				UNiagaraParameterCollection* NPC = *NPCIt;
-				if (NPC->GetParameters().Contains(DistributionBase->ParameterBinding))
+			DistributionBase->ForEachParameterBinding(
+				[this](const FNiagaraVariableBase& Variable)
 				{
-					CachedParameterCollectionReferences.AddUnique(NPC);
-					break;
+					if (Variable.IsInNameSpace(FNiagaraConstants::ParameterCollectionNamespaceString))
+					{
+						for (TObjectIterator<UNiagaraParameterCollection> NPCIt; NPCIt; ++NPCIt)
+						{
+							UNiagaraParameterCollection* NPC = *NPCIt;
+							if (NPC->GetParameters().Contains(Variable))
+							{
+								CachedParameterCollectionReferences.AddUnique(NPC);
+								break;
+							}
+						}
+					}
 				}
-			}
+			);
 		}
 	}
 	
@@ -379,6 +376,7 @@ void UNiagaraStatelessEmitter::CacheFromCompiledData()
 		FNiagaraStatelessEmitterDataBuildContext EmitterBuildContext(
 			*StatelessEmitterData->ParticleDataSetCompiledData.Get(),
 			StatelessEmitterData->RendererBindings,
+			StatelessEmitterData->Expressions,
 			StatelessEmitterData->BuiltData,
 			StatelessEmitterData->StaticFloatData,
 			StatelessEmitterData->ParticleSimExecData
diff --git a/Engine/Plugins/FX/Niagara/Source/Niagara/Private/Stateless/NiagaraStatelessEmitterDataBuildContext.cpp b/Engine/Plugins/FX/Niagara/Source/Niagara/Private/Stateless/NiagaraStatelessEmitterDataBuildContext.cpp
--- a/Engine/Plugins/FX/Niagara/Source/Niagara/Private/Stateless/NiagaraStatelessEmitterDataBuildContext.cpp
+++ b/Engine/Plugins/FX/Niagara/Source/Niagara/Private/Stateless/NiagaraStatelessEmitterDataBuildContext.cpp
@@ -1,6 +1,7 @@
 // Copyright Epic Games, Inc. All Rights Reserved.
 
 #include "Stateless/NiagaraStatelessEmitterDataBuildContext.h"
+#include "Stateless/NiagaraStatelessExpression.h"
 #include "Stateless/NiagaraStatelessParticleSimExecData.h"
 
 #include "NiagaraDataSetCompiledData.h"
@@ -99,6 +100,21 @@ int32 FNiagaraStatelessEmitterDataBuildContext::AddRendererBinding(const FNiagar
 	return DataOffset;
 }
 
+
+int32 FNiagaraStatelessEmitterDataBuildContext::AddExpression(const FInstancedStruct& ExpressionStruct) const
+{
+	int32 DataOffset = INDEX_NONE;
+	if (const FNiagaraStatelessExpression* Expression = ExpressionStruct.GetPtr<FNiagaraStatelessExpression>())
+	{
+		const FName ExpressionName("__StatelessInternal__.Expression", Expressions.Num());
+		const FNiagaraVariable ExpressionVariable(Expression->GetOutputTypeDef(), ExpressionName);
+		RendererBindings.AddParameter(ExpressionVariable, false, false, &DataOffset);
+
+		Expressions.Emplace(DataOffset, Expression->Build(*this));
+	}
+	return DataOffset;
+}
+
 void FNiagaraStatelessEmitterDataBuildContext::AddParticleSimulationExecSimulate(TFunction<void(const NiagaraStateless::FParticleSimulationContext&)> Func) const
 {
 	if (!ParticleExecData)
diff --git a/Engine/Plugins/FX/Niagara/Source/Niagara/Private/Stateless/NiagaraStatelessEmitterInstance.cpp b/Engine/Plugins/FX/Niagara/Source/Niagara/Private/Stateless/NiagaraStatelessEmitterInstance.cpp
--- a/Engine/Plugins/FX/Niagara/Source/Niagara/Private/Stateless/NiagaraStatelessEmitterInstance.cpp
+++ b/Engine/Plugins/FX/Niagara/Source/Niagara/Private/Stateless/NiagaraStatelessEmitterInstance.cpp
@@ -1,9 +1,10 @@
 // Copyright Epic Games, Inc. All Rights Reserved.
 
 #include "Stateless/NiagaraStatelessEmitterInstance.h"
+#include "Stateless/NiagaraStatelessComputeManager.h"
 #include "Stateless/NiagaraStatelessEmitter.h"
 #include "Stateless/NiagaraStatelessEmitterData.h"
-#include "Stateless/NiagaraStatelessComputeManager.h"
+#include "Stateless/NiagaraStatelessExpression.h"
 #include "NiagaraGpuComputeDispatchInterface.h"
 #include "NiagaraSystem.h"
 #include "NiagaraSystemInstance.h"
@@ -34,6 +35,13 @@ namespace NiagaraStatelessEmitterInstancePrivate
 		{
 			return Distribution.ParameterBinding.IsValid() ? ParameterStore.GetParameterValueOrDefault(Distribution.ParameterBinding, DefaultValue) : DefaultValue;
 		}
+		else if (Distribution.IsExpression())
+		{
+			check(Distribution.ParameterExpression.Get<FNiagaraStatelessExpression>().GetOutputTypeDef() == FNiagaraTypeDefinition::GetFloatDef());
+			float ExpressionValue = DefaultValue;
+			Distribution.ParameterExpression.Get<FNiagaraStatelessExpression>().Evaluate(FNiagaraStatelessExpression::FEvaluateContext(ParameterStore), &ExpressionValue);
+			return ExpressionValue;
+		}
 		const float Fraction = RandomStream.GetFraction();
 		return ((Distribution.Max - Distribution.Min) * Fraction) + Distribution.Min;
 	}
@@ -44,6 +52,13 @@ namespace NiagaraStatelessEmitterInstancePrivate
 		{
 			return Distribution.ParameterBinding.IsValid() ? ParameterStore.GetParameterValueOrDefault(Distribution.ParameterBinding, DefaultValue) : DefaultValue;
 		}
+		else if (Distribution.IsExpression())
+		{
+			check(Distribution.ParameterExpression.Get<FNiagaraStatelessExpression>().GetOutputTypeDef() == FNiagaraTypeDefinition::GetIntDef());
+			int32 ExpressionValue = DefaultValue;
+			Distribution.ParameterExpression.Get<FNiagaraStatelessExpression>().Evaluate(FNiagaraStatelessExpression::FEvaluateContext(ParameterStore), &ExpressionValue);
+			return ExpressionValue;
+		}
 		return RandomStream.RandRange(Distribution.Min, Distribution.Max);
 	}
 
@@ -51,7 +66,7 @@ namespace NiagaraStatelessEmitterInstancePrivate
 	{
 		if (Distribution.ParameterOffset != INDEX_NONE)
 		{
-			FMemory::Memcpy(&OutMin, ParameterStore.GetParameterData(Distribution.ParameterOffset), sizeof(float));
+			OutMin = ParameterStore.GetParameterValueFromOffset<float>(Distribution.ParameterOffset);
 			OutMax = OutMin;
 		}
 		else
@@ -533,6 +548,12 @@ void FNiagaraStatelessEmitterInstance::SendRenderData()
 	bool bNeedsShaderParametersUpdate = false;
 	if (RendererBindings.GetParametersDirty())
 	{
+		for (const TPair<int32, FInstancedStruct>& ExpressionMapping : EmitterData->Expressions)
+		{
+			const FNiagaraStatelessExpression& Expression = ExpressionMapping.Value.Get<FNiagaraStatelessExpression>();
+			uint8* ValueDestination = RendererBindings.GetMutableParameterData(ExpressionMapping.Key, Expression.GetOutputTypeDef());
+			Expression.Evaluate(FNiagaraStatelessExpression::FEvaluateContext(RendererBindings), ValueDestination);
+		}
 		RendererBindings.Tick();
 
 		if (EmitterData->bModulesHaveRendererBindings)
diff --git a/Engine/Plugins/FX/Niagara/Source/Niagara/Private/Stateless/NiagaraStatelessExpression.cpp b/Engine/Plugins/FX/Niagara/Source/Niagara/Private/Stateless/NiagaraStatelessExpression.cpp
new file mode 100644
--- /dev/null
+++ b/Engine/Plugins/FX/Niagara/Source/Niagara/Private/Stateless/NiagaraStatelessExpression.cpp
@@ -0,0 +1,37 @@
+// Copyright Epic Games, Inc. All Rights Reserved.
+
+#include "Stateless/NiagaraStatelessExpression.h"
+
+#if WITH_EDITORONLY_DATA
+void FNiagaraStatelessExpression::ForEachBinding(const FInstancedStruct& ExpressionStruct, const TFunction<void(const FNiagaraVariableBase&)>& Delegate)
+{
+	const FNiagaraStatelessExpression* Expression = ExpressionStruct.GetPtr<const FNiagaraStatelessExpression>();
+	if ( Expression == nullptr)
+	{
+		return;
+	}
+
+	// This isn't ideal, needs more thought
+	static const FName NAME_BindingExpression("BindingExpression");
+	const UScriptStruct* ExpressionScriptStruct = ExpressionStruct.GetScriptStruct();
+	if (ExpressionScriptStruct->HasMetaData(NAME_BindingExpression))
+	{
+		FNameProperty* NameProperty = CastField<FNameProperty>(ExpressionScriptStruct->ChildProperties);
+		if (ensure(NameProperty && NameProperty->Next == nullptr))
+		{
+			const FName BindingName = NameProperty->GetPropertyValue_InContainer(Expression);
+			Delegate(FNiagaraVariableBase(Expression->GetOutputTypeDef(), BindingName));
+		}
+	}
+
+	for (TFieldIterator<FProperty> PropertyIt(ExpressionStruct.GetScriptStruct()); PropertyIt; ++PropertyIt)
+	{
+		const FStructProperty* StructProperty = CastField<const FStructProperty>(*PropertyIt);
+		if (StructProperty && StructProperty->Struct == FInstancedStruct::StaticStruct())
+		{
+			const FInstancedStruct* InnerExpressionStruct = StructProperty->ContainerPtrToValuePtr<const FInstancedStruct>(ExpressionStruct.GetMemory());
+			ForEachBinding(*InnerExpressionStruct, Delegate);
+		}
+	}
+}
+#endif
diff --git a/Engine/Plugins/FX/Niagara/Source/Niagara/Public/NiagaraParameterStore.h b/Engine/Plugins/FX/Niagara/Source/Niagara/Public/NiagaraParameterStore.h
--- a/Engine/Plugins/FX/Niagara/Source/Niagara/Public/NiagaraParameterStore.h
+++ b/Engine/Plugins/FX/Niagara/Source/Niagara/Public/NiagaraParameterStore.h
@@ -107,9 +107,9 @@ struct FNiagaraParameterStoreBinding
 
 	FORCEINLINE_DEBUGGABLE void Empty(FNiagaraParameterStore* DestStore, FNiagaraParameterStore* SrcStore);
 	FORCEINLINE_DEBUGGABLE bool Initialize(FNiagaraParameterStore* DestStore, FNiagaraParameterStore* SrcStore, const FNiagaraBoundParameterArray* BoundParameters = nullptr);
-	FORCEINLINE_DEBUGGABLE bool VerifyBinding(const FNiagaraParameterStore* DestStore, const FNiagaraParameterStore* SrcStore)const;
+	NIAGARA_API bool VerifyBinding(const FNiagaraParameterStore* DestStore, const FNiagaraParameterStore* SrcStore)const;
 	void CopyParameters(FNiagaraParameterStore* DestStore, const FNiagaraParameterStore* SrcStore) const;
-	FORCEINLINE_DEBUGGABLE void Dump(const FNiagaraParameterStore* DestStore, const FNiagaraParameterStore* SrcStore)const;
+	NIAGARA_API void Dump(const FNiagaraParameterStore* DestStore, const FNiagaraParameterStore* SrcStore)const;
 
 	static void GetBindingData(FNiagaraParameterStore* DestStore, FNiagaraParameterStore* SrcStore, FNiagaraBoundParameterArray& OutBoundParameters);
 
@@ -377,16 +377,16 @@ public:
 	
 	NIAGARA_API virtual const FNiagaraVariableWithOffset* FindParameterVariable(const FNiagaraVariable& Parameter, bool IgnoreType = false) const;
 
-private:
 	template<typename T>
 	T GetParameterValueFromOffset(int32 Offset) const
 	{
+		check((Offset >= 0) && (Offset + sizeof(T) <= ParameterData.Num()));
+
 		T OutValue;
-		FMemory::Memcpy(&OutValue, GetParameterData(Offset), sizeof(T));
+		FMemory::Memcpy(&OutValue, ParameterData.GetData() + Offset, sizeof(T));
 		return OutValue;
 	}
 
-public:
 	/** Gets the typed parameter data. */
 	template<typename T>
 	FORCEINLINE_DEBUGGABLE void GetParameterValue(T& OutValue, const FNiagaraVariableBase& Parameter)const
@@ -432,22 +432,48 @@ public:
 		return Offset != INDEX_NONE ? GetParameterValueFromOffset<T>(Offset) : DefaultValue;
 	}
 
-	FORCEINLINE const uint8* GetParameterData(int32 Offset)const
+	UE_DEPRECATED(5.5, "GetParameterData without a type side is deprecated as it's unsafe")
+	FORCEINLINE const uint8* GetParameterData(int32 Offset) const
+	{
+		return ParameterData.GetData() + Offset;
+	}
+
+	FORCEINLINE const uint8* GetParameterData(int32 Offset, int32 SizeInBytes) const
 	{
+		check(Offset >= 0 && (Offset + SizeInBytes) <= ParameterData.Num());
+		return ParameterData.GetData() + Offset;
+	}
+
+	FORCEINLINE const uint8* GetParameterData(int32 Offset, const FNiagaraTypeDefinition& TypeDef) const
+	{
+		check(Offset >= 0 && (Offset + TypeDef.GetSize()) <= ParameterData.Num());
+		return ParameterData.GetData() + Offset;
+	}
+
+	/** Returns a none const pointer to the parameter data, any modifications to the data will not be recorded and automatically pushed to other parameter stores, use with caution. */
+	FORCEINLINE uint8* GetMutableParameterData(int32 Offset, const FNiagaraTypeDefinition& TypeDef)
+	{
+		check(Offset >= 0 && (Offset + TypeDef.GetSize()) <= ParameterData.Num());
 		return ParameterData.GetData() + Offset;
 	}
 
 	/** Returns the parameter data for the passed parameter if it exists in this store. Null if not. */
-	FORCEINLINE_DEBUGGABLE const uint8* GetParameterData(const FNiagaraVariableBase& Parameter)const
+	FORCEINLINE_DEBUGGABLE const uint8* GetParameterData(const FNiagaraVariableBase& Parameter) const
 	{
 		int32 Offset = IndexOf(Parameter);
 		if (Offset != INDEX_NONE)
 		{
-			return GetParameterData(Offset);
+			return GetParameterData_Internal(Offset);
 		}
 		return nullptr;
 	}
 
+	FORCEINLINE_DEBUGGABLE const uint8* GetParameterData(const FNiagaraVariableWithOffset& Parameter) const
+	{
+		check(Parameter.Offset != INDEX_NONE);
+		return GetParameterData_Internal(Parameter.Offset);
+	}
+
 	/** Copies the data stored for the given variable to the target pointer. This method automatically converts custom struct data back to lwc types.
 	 * Just using the raw parameter store pointer via IndexOf() will be wrong if the target struct contains a lwc type like FVector.
 	 * Returns true if the data was copied, false if the parameter could not be found in the store.
@@ -544,10 +570,8 @@ public:
 		int32 Offset = IndexOf(Param);
 		if (Offset != INDEX_NONE)
 		{
-			//*(T*)(GetParameterData_Internal(Offset)) = InValue;
 			//Until we solve our alignment issues, temporarily just doing a memcpy here.
-			T* ParamData = reinterpret_cast<T*>( GetParameterData_Internal(Offset) );
-			FMemory::Memcpy(ParamData, &InValue, sizeof(T));
+			FMemory::Memcpy(GetMutableParameterData_Internal(Offset), &InValue, sizeof(T));
 			OnParameterChange();
 			return true;
 		}
@@ -560,9 +584,7 @@ public:
 				AddParameter(Param, bInitInterfaces, bTriggerRebind, &Offset);
 				check(Offset != INDEX_NONE);
 				//Until we solve our alignment issues, temporarily just doing a memcpy here.
-				//*(T*)(GetParameterData_Internal(Offset)) = InValue;
-				T* ParamData = reinterpret_cast<T*>(GetParameterData_Internal(Offset));
-				FMemory::Memcpy(ParamData, &InValue, sizeof(T));
+				FMemory::Memcpy(GetMutableParameterData_Internal(Offset), &InValue, sizeof(T));
 				OnLayoutChange();
 				return true;
 			}
@@ -572,9 +594,9 @@ public:
 
 	FORCEINLINE_DEBUGGABLE void SetParameterData(const uint8* Data, int32 Offset, int32 Size)
 	{
-		checkSlow(Data != nullptr);
-		checkSlow((Offset + Size) <= ParameterData.Num());
-		uint8* Dest = GetParameterData_Internal(Offset);
+		check(Data != nullptr);
+		check(Offset >= 0 && (Offset + Size) <= ParameterData.Num());
+		uint8* Dest = GetMutableParameterData_Internal(Offset);
 		if (Dest != Data)
 		{
 			FMemory::Memcpy(Dest, Data, Size);
@@ -680,11 +702,8 @@ protected:
 	NIAGARA_API void AssignParameterData(TConstArrayView<uint8> SourceParameterData);
 	static NIAGARA_API int32 PaddedParameterSize(int32 ParameterSize);
 
-	/** Returns the parameter data at the passed offset. */
-	FORCEINLINE uint8* GetParameterData_Internal(int32 Offset) 
-	{
-		return ParameterData.GetData() + Offset;
-	}
+	const uint8* GetParameterData_Internal(int32 Offset) const { return ParameterData.GetData() + Offset; }
+	uint8* GetMutableParameterData_Internal(int32 Offset) { return ParameterData.GetData() + Offset; }
 
 	template<typename ParamType>
 	void SetParameterByOffset(uint32 ParamOffset, const ParamType& Param)
@@ -750,163 +769,6 @@ FORCEINLINE_DEBUGGABLE bool FNiagaraParameterStoreBinding::Initialize(FNiagaraPa
 	}
 }
 
-FORCEINLINE_DEBUGGABLE bool FNiagaraParameterStoreBinding::VerifyBinding(const FNiagaraParameterStore* DestStore, const FNiagaraParameterStore* SrcStore)const
-{
-	bool bBindingValid = true;
-#if WITH_EDITORONLY_DATA
-	TArray<FName, TInlineAllocator<32>> MissingParameterNames;
-	for (const FNiagaraVariableWithOffset& ParamWithOffset : DestStore->ReadParameterVariables())
-	{
-		const FNiagaraVariable& Parameter = ParamWithOffset;
-		int32 DestOffset = ParamWithOffset.Offset;
-		int32 SrcOffset = SrcStore->IndexOf(Parameter);
-
-		if (SrcOffset == INDEX_NONE)
-		{
-			continue;//Parameter is not present in SrcStore so shouldn't/cant be bound.
-		}
-
-		if (Parameter.IsDataInterface())
-		{
-			if (InterfaceBindings.Find(FInterfaceBinding(SrcOffset, DestOffset)) == INDEX_NONE)
-			{
-				MissingParameterNames.Add(Parameter.GetName());
-				bBindingValid = false;
-			}
-		}
-		else if (Parameter.IsUObject())
-		{
-			if (UObjectBindings.Find(FUObjectBinding(SrcOffset, DestOffset)) == INDEX_NONE)
-			{
-				MissingParameterNames.Add(Parameter.GetName());
-				bBindingValid = false;
-			}
-		}
-		else
-		{
-			if (ParameterBindings.Find(FParameterBinding(SrcOffset, DestOffset, Parameter.GetSizeInBytes())) == INDEX_NONE)
-			{
-				MissingParameterNames.Add(Parameter.GetName());
-				bBindingValid = false;
-			}
-		}
-	}
-
-	if(MissingParameterNames.Num())
-	{
-		UE_LOG(LogNiagara, Warning, TEXT("Invalid ParameterStore Binding: Parameters missing from binding between %s and %s. Stores should have been rebound when one changed!"), *SrcStore->DebugName, *DestStore->DebugName);
-		for (FName MissingParam : MissingParameterNames)
-		{
-			UE_LOG(LogNiagara, Warning, TEXT("%s"), *MissingParam.ToString());
-		}
-	}
-#endif
-	return bBindingValid;
-}
-
-FORCEINLINE_DEBUGGABLE void FNiagaraParameterStoreBinding::Dump(const FNiagaraParameterStore* DestStore, const FNiagaraParameterStore* SrcStore)const
-{
-#if WITH_EDITORONLY_DATA
-	UE_LOG(LogNiagara, Log, TEXT("\n\nDest Store: %s\n"), *DestStore->DebugName);
-	for (const FParameterBinding& Binding : ParameterBindings)
-	{
-		ensure(Binding.Size != 0);
-		ensure(Binding.SrcOffset != -1);
-		ensure(Binding.DestOffset != -1);
-		FNiagaraVariable Param;
-		bool bFound = false;
-		for (const FNiagaraVariableWithOffset& ParamWithOffset : DestStore->ReadParameterVariables())
-		{
-			if (ParamWithOffset.Offset == Binding.DestOffset && !ParamWithOffset.IsDataInterface())
-			{
-				bFound = true;
-				Param = ParamWithOffset;
-			}
-		}
-		if (ensure(bFound))
-		{
-			UE_LOG(LogNiagara, Log, TEXT("| Param | %s %s: Src:%d - Dest:%d - Size:%d\n"), *Param.GetType().GetName(), *Param.GetName().ToString(), Binding.SrcOffset, Binding.DestOffset, Binding.Size);
-
-			//Also ensure the param has been pushed correctly.
-			const uint8* SrcData = SrcStore->GetParameterData(Binding.SrcOffset);
-			const uint8* DestData = DestStore->GetParameterData(Binding.DestOffset);
-			if (!ensure(FMemory::Memcmp(SrcData, DestData, Binding.Size) == 0))
-			{
-				UE_LOG(LogNiagara, Log, TEXT("Parameter in dest store has incorrect data!\n"));
-			}
-		}
-		else
-		{
-			UE_LOG(LogNiagara, Log, TEXT("Failed to find matching param in bound store!\n"));
-		}
-	}
-
-	for (const FInterfaceBinding& Binding : InterfaceBindings)
-	{
-		ensure(Binding.SrcOffset != -1);
-		ensure(Binding.DestOffset != -1);
-		FNiagaraVariable Param;
-		bool bFound = false;
-		for (const FNiagaraVariableWithOffset& ParamWithOffset : DestStore->ReadParameterVariables())
-		{
-			if (ParamWithOffset.Offset == Binding.DestOffset && ParamWithOffset.IsDataInterface())
-			{
-				bFound = true;
-				Param = ParamWithOffset;
-			}
-		}
-		if (ensure(bFound))
-		{
-			UE_LOG(LogNiagara, Log, TEXT("| DI | %s %s: Src:%d - Dest:%d\n"), *Param.GetType().GetName(), *Param.GetName().ToString(), Binding.SrcOffset, Binding.DestOffset);
-
-			//Also ensure the param has been pushed correctly.
-			const UNiagaraDataInterface* SrcData = SrcStore->GetDataInterfaces()[Binding.SrcOffset];
-			const UNiagaraDataInterface* DestData = DestStore->GetDataInterfaces()[Binding.DestOffset];
-			if (!ensure(SrcData == DestData))
-			{
-				UE_LOG(LogNiagara, Log, TEXT("Data interface parameter in dest store is incorrect!\n"));
-			}
-		}
-		else
-		{
-			UE_LOG(LogNiagara, Log, TEXT("Failed to find matching data interface param in bound store!\n"));
-		}
-	}
-
-	for (const FUObjectBinding& Binding : UObjectBindings)
-	{
-		ensure(Binding.SrcOffset != -1);
-		ensure(Binding.DestOffset != -1);
-		FNiagaraVariable Param;
-		bool bFound = false;
-		for (const FNiagaraVariableWithOffset& ParamWithOffset : DestStore->ReadParameterVariables())
-		{
-			if (ParamWithOffset.Offset == Binding.DestOffset && ParamWithOffset.IsUObject())
-			{
-				bFound = true;
-				Param = ParamWithOffset;
-			}
-		}
-		if (ensure(bFound))
-		{
-			UE_LOG(LogNiagara, Log, TEXT("| UObject | %s %s: Src:%d - Dest:%d\n"), *Param.GetType().GetName(), *Param.GetName().ToString(), Binding.SrcOffset, Binding.DestOffset);
-
-			//Also ensure the param has been pushed correctly.
-			TObjectPtr<const UObject> SrcData = SrcStore->GetUObjects()[Binding.SrcOffset];
-			TObjectPtr<const UObject> DestData = DestStore->GetUObjects()[Binding.DestOffset];
-			if (!ensure(SrcData == DestData))
-			{
-				UE_LOG(LogNiagara, Log, TEXT("Data interface parameter in dest store is incorrect!\n"));
-			}
-		}
-		else
-		{
-			UE_LOG(LogNiagara, Log, TEXT("Failed to find matching data interface param in bound store!\n"));
-		}
-	}
-#endif
-}
-
 //////////////////////////////////////////////////////////////////////////
 
 #define NIAGARA_VALIDATE_DIRECT_BINDINGS	DO_CHECK
diff --git a/Engine/Plugins/FX/Niagara/Source/Niagara/Public/Stateless/NiagaraStatelessDistribution.h b/Engine/Plugins/FX/Niagara/Source/Niagara/Public/Stateless/NiagaraStatelessDistribution.h
--- a/Engine/Plugins/FX/Niagara/Source/Niagara/Public/Stateless/NiagaraStatelessDistribution.h
+++ b/Engine/Plugins/FX/Niagara/Source/Niagara/Public/Stateless/NiagaraStatelessDistribution.h
@@ -4,7 +4,8 @@
 
 #include "Curves/RichCurve.h"
 #include "NiagaraStatelessRange.h"
-#include "NiagaraParameterBinding.h"
+#include "NiagaraCommon.h"
+#include "StructUtils/InstancedStruct.h"
 
 #include "NiagaraStatelessDistribution.generated.h"
 
@@ -12,6 +13,7 @@ UENUM()
 enum class ENiagaraDistributionMode
 {
 	Binding,
+	Expression,
 	UniformConstant,
 	NonUniformConstant,
 	UniformRange,
@@ -39,7 +41,11 @@ struct FNiagaraDistributionBase
 	UPROPERTY(EditAnywhere, Category = "Parameters")
 	FNiagaraVariableBase ParameterBinding;
 
+	UPROPERTY(EditAnywhere, Category = "Parameters")
+	FInstancedStruct ParameterExpression;
+
 	bool IsBinding() const { return Mode == ENiagaraDistributionMode::Binding; }
+	bool IsExpression() const { return Mode == ENiagaraDistributionMode::Expression; }
 	bool IsConstant() const { return Mode == ENiagaraDistributionMode::UniformConstant || Mode == ENiagaraDistributionMode::NonUniformConstant; }
 	bool IsUniform() const { return Mode == ENiagaraDistributionMode::UniformConstant || Mode == ENiagaraDistributionMode::UniformRange; }
 	bool IsCurve() const { return Mode == ENiagaraDistributionMode::UniformCurve || Mode == ENiagaraDistributionMode::NonUniformCurve; }
@@ -55,15 +61,9 @@ struct FNiagaraDistributionBase
 	UPROPERTY(EditAnywhere, Category = "Parameters")
 	int32 MaxLutSampleCount = 128;
 
-	bool operator==(const FNiagaraDistributionBase& Other) const
-	{
-		return
-			Mode == Other.Mode &&
-			ParameterBinding == Other.ParameterBinding &&
-			ChannelConstantsAndRanges == Other.ChannelConstantsAndRanges &&
-			ChannelCurves == Other.ChannelCurves &&
-			MaxLutSampleCount == Other.MaxLutSampleCount;
-	}
+	NIAGARA_API bool operator==(const FNiagaraDistributionBase& Other) const;
+
+	void ForEachParameterBinding(TFunction<void(const FNiagaraVariableBase&)> Delegate) const;
 
 	virtual bool AllowBinding() const { return true; }
 	virtual bool AllowConstant() const { return true; }
@@ -92,6 +92,9 @@ struct FNiagaraDistributionRangeInt
 	UPROPERTY(EditAnywhere, Category = "Parameters")
 	FNiagaraVariableBase ParameterBinding;
 
+	UPROPERTY(EditAnywhere, Category = "Parameters")
+	FInstancedStruct ParameterExpression;
+
 	UPROPERTY(EditAnywhere, Category = "Parameters")
 	int32 Min = 0;
 
@@ -102,6 +105,7 @@ struct FNiagaraDistributionRangeInt
 	NIAGARA_API FNiagaraStatelessRangeInt CalculateRange(const int32 Default = 0) const;
 
 	bool IsBinding() const { return Mode == ENiagaraDistributionMode::Binding; }
+	bool IsExpression() const { return Mode == ENiagaraDistributionMode::Expression; }
 	bool IsConstant() const { return Mode == ENiagaraDistributionMode::UniformConstant || Mode == ENiagaraDistributionMode::NonUniformConstant; }
 	bool IsUniform() const { return Mode == ENiagaraDistributionMode::UniformConstant || Mode == ENiagaraDistributionMode::UniformRange; }
 	bool IsCurve() const { return Mode == ENiagaraDistributionMode::UniformCurve || Mode == ENiagaraDistributionMode::NonUniformCurve; }
diff --git a/Engine/Plugins/FX/Niagara/Source/NiagaraEditor/Private/NiagaraCompilationCopy.cpp b/Engine/Plugins/FX/Niagara/Source/NiagaraEditor/Private/NiagaraCompilationCopy.cpp
--- a/Engine/Plugins/FX/Niagara/Source/NiagaraEditor/Private/NiagaraCompilationCopy.cpp
+++ b/Engine/Plugins/FX/Niagara/Source/NiagaraEditor/Private/NiagaraCompilationCopy.cpp
@@ -479,7 +479,7 @@ void FNiagaraPrecompileData::CollectBakedRapidIterationParameters(const FNiagara
 				if (!AlreadyAdded)
 				{
 					FNiagaraVariable& BakedParameter = BakedRapidIterationParameters.Add_GetRef(Param);
-					BakedParameter.SetData(ScriptInfo->RapidIterationParameters.GetParameterData(Param.Offset));
+					BakedParameter.SetData(ScriptInfo->RapidIterationParameters.GetParameterData(Param));
 				}
 			}
 		}
diff --git a/Engine/Plugins/FX/Niagara/Source/NiagaraEditor/Private/NiagaraCompilationTasks.cpp b/Engine/Plugins/FX/Niagara/Source/NiagaraEditor/Private/NiagaraCompilationTasks.cpp
--- a/Engine/Plugins/FX/Niagara/Source/NiagaraEditor/Private/NiagaraCompilationTasks.cpp
+++ b/Engine/Plugins/FX/Niagara/Source/NiagaraEditor/Private/NiagaraCompilationTasks.cpp
@@ -1909,7 +1909,7 @@ bool FNiagaraSystemCompilationTask::Poll(FNiagaraSystemAsyncCompileResults& Resu
 						for (const FNiagaraVariableWithOffset& ParamWithOffset : RapidIterationParameters)
 						{
 							FNiagaraVariable& Parameter = CompileData.RapidIterationParameters.Add_GetRef(ParamWithOffset);
-							Parameter.SetData(ScriptInfo->RapidIterationParameters.GetParameterData(ParamWithOffset.Offset));
+							Parameter.SetData(ScriptInfo->RapidIterationParameters.GetParameterData(ParamWithOffset));
 						}
 					}
 				}
@@ -2576,7 +2576,7 @@ UE::Tasks::FTask FNiagaraSystemCompilationTask::BuildRapidIterationParametersAsy
 					if (ScriptRI.GetType().IsStatic())
 					{
 						FNiagaraVariable StaticVariable = FNiagaraVariable(ScriptRI);
-						StaticVariable.SetData(SrcScriptInfo.RapidIterationParameters.GetParameterData(ScriptRI.Offset));
+						StaticVariable.SetData(SrcScriptInfo.RapidIterationParameters.GetParameterData(ScriptRI));
 
 						if (SrcScriptInfo.SourceEmitterIndex == INDEX_NONE)
 						{
@@ -2618,7 +2618,7 @@ void FNiagaraSystemCompilationTask::BuildAndApplyRapidIterationParameters()
 			for (const FNiagaraVariableWithOffset& ParamWithOffset : SrcRapidIterationParameters)
 			{
 				FNiagaraVariable& Parameter = DstRapidIterationParameters.Add_GetRef(ParamWithOffset);
-				Parameter.SetData(ScriptInfo.RapidIterationParameters.GetParameterData(ParamWithOffset.Offset));
+				Parameter.SetData(ScriptInfo.RapidIterationParameters.GetParameterData(ParamWithOffset));
 			}
 
 			Script->ApplyRapidIterationParameters(DstRapidIterationParameters, bAllowParameterRemoval);
diff --git a/Engine/Plugins/FX/Niagara/Source/NiagaraEditor/Private/NiagaraCompiler.cpp b/Engine/Plugins/FX/Niagara/Source/NiagaraEditor/Private/NiagaraCompiler.cpp
--- a/Engine/Plugins/FX/Niagara/Source/NiagaraEditor/Private/NiagaraCompiler.cpp
+++ b/Engine/Plugins/FX/Niagara/Source/NiagaraEditor/Private/NiagaraCompiler.cpp
@@ -488,7 +488,7 @@ void FNiagaraCompileRequestData::AddRapidIterationParameters(const FNiagaraParam
 			const int32* Index = InParamStore.FindParameterOffset(StoreParams[i]);
 			if (Index != nullptr)
 			{
-				StoreParams[i].SetData(InParamStore.GetParameterData(*Index)); // This will memcopy the data in.
+				StoreParams[i].SetData(InParamStore.GetParameterData(*Index, StoreParams[i].GetType())); // This will memcopy the data in.
 				RapidIterationParams.Add(StoreParams[i]);
 			}
 		}
@@ -499,7 +499,7 @@ void FNiagaraCompileRequestData::AddRapidIterationParameters(const FNiagaraParam
 			const int32* Index = InParamStore.FindParameterOffset(StoreParams[i]);
 			if (Index != nullptr)
 			{
-				StoreParams[i].SetData(InParamStore.GetParameterData(*Index)); // This will memcopy the data in.
+				StoreParams[i].SetData(InParamStore.GetParameterData(*Index, StoreParams[i].GetType())); // This will memcopy the data in.
 
 				if (StoreParams[i] != ExistingVar)
 				{
diff --git a/Engine/Plugins/FX/Niagara/Source/NiagaraEditor/Private/Stateless/NiagaraDistributionIntPropertyCustomization.cpp b/Engine/Plugins/FX/Niagara/Source/NiagaraEditor/Private/Stateless/NiagaraDistributionIntPropertyCustomization.cpp
--- a/Engine/Plugins/FX/Niagara/Source/NiagaraEditor/Private/Stateless/NiagaraDistributionIntPropertyCustomization.cpp
+++ b/Engine/Plugins/FX/Niagara/Source/NiagaraEditor/Private/Stateless/NiagaraDistributionIntPropertyCustomization.cpp
@@ -27,6 +27,9 @@
 #include "Widgets/Input/SNumericEntryBox.h"
 #include "Widgets/Layout/SBox.h"
 
+//#include "Stateless/SNiagaraStatelessExpressionWidget.h"
+//#include "NiagaraStatelessExpressionTypeData.h"
+
 #define LOCTEXT_NAMESPACE "NiagaraDistributionIntPropertyCustomization"
 
 class SNiagaraDistributionIntPropertyWidget : public SCompoundWidget
@@ -90,6 +93,7 @@ class SNiagaraDistributionIntPropertyWidget : public SCompoundWidget
 			switch (Distribution->Mode)
 			{
 				case ENiagaraDistributionMode::Binding:			return ENiagaraDistributionEditorMode::Binding;
+				//case ENiagaraDistributionMode::Expression:		return ENiagaraDistributionEditorMode::Expression;
 				case ENiagaraDistributionMode::UniformConstant: return ENiagaraDistributionEditorMode::UniformConstant;
 				case ENiagaraDistributionMode::UniformRange:	return ENiagaraDistributionEditorMode::UniformRange;
 				default:										break;
@@ -111,6 +115,7 @@ class SNiagaraDistributionIntPropertyWidget : public SCompoundWidget
 			switch (NewMode)
 			{
 				case ENiagaraDistributionEditorMode::Binding:			Distribution->Mode = ENiagaraDistributionMode::Binding; break;
+				//case ENiagaraDistributionEditorMode::Expression:		Distribution->Mode = ENiagaraDistributionMode::Expression; break;
 				case ENiagaraDistributionEditorMode::UniformConstant:	Distribution->Mode = ENiagaraDistributionMode::UniformConstant; break;
 				case ENiagaraDistributionEditorMode::UniformRange:		Distribution->Mode = ENiagaraDistributionMode::UniformRange; break;
 				default:												break;
@@ -139,6 +144,11 @@ class SNiagaraDistributionIntPropertyWidget : public SCompoundWidget
 			if (!PropertyHandle->HasMetaData(DisableBindingDistributionName))
 			{
 				SupportedModes.Add(ENiagaraDistributionEditorMode::Binding);
+				//const FNiagaraStatelessExpressionTypeData& TypeData = FNiagaraStatelessExpressionTypeData::GetTypeData(Distribution->GetBindingTypeDef());
+				//if (TypeData.IsValid())
+				//{
+				//	SupportedModes.Add(ENiagaraDistributionEditorMode::Expression);
+				//}
 			}
 			SupportedModes.Add(ENiagaraDistributionEditorMode::UniformConstant);
 			SupportedModes.Add(ENiagaraDistributionEditorMode::UniformRange);
@@ -184,6 +194,14 @@ class SNiagaraDistributionIntPropertyWidget : public SCompoundWidget
 					ConstructBindingWidget()
 				];
 		}
+		//else if (Mode == ENiagaraDistributionEditorMode::Expression)
+		//{
+		//	return SNew(SHorizontalBox)
+		//		+ SHorizontalBox::Slot()
+		//		[
+		//			ConstructExpressionWidget()
+		//		];
+		//}
 		else if (Mode == ENiagaraDistributionEditorMode::UniformConstant)
 		{
 			return SNew(SHorizontalBox)
@@ -301,6 +319,11 @@ class SNiagaraDistributionIntPropertyWidget : public SCompoundWidget
 	{
 		bContinuousChangeActive = false;
 	}
+	
+	//TSharedRef<SWidget> ConstructExpressionWidget()
+	//{
+	//	//-TODO: Not supported currently
+	//}
 
 	TSharedRef<SWidget> ConstructBindingWidget()
 	{
diff --git a/Engine/Plugins/FX/Niagara/Source/NiagaraEditor/Private/Stateless/NiagaraDistributionPropertyCustomization.cpp b/Engine/Plugins/FX/Niagara/Source/NiagaraEditor/Private/Stateless/NiagaraDistributionPropertyCustomization.cpp
--- a/Engine/Plugins/FX/Niagara/Source/NiagaraEditor/Private/Stateless/NiagaraDistributionPropertyCustomization.cpp
+++ b/Engine/Plugins/FX/Niagara/Source/NiagaraEditor/Private/Stateless/NiagaraDistributionPropertyCustomization.cpp
@@ -20,6 +20,7 @@
 #include "NiagaraNodeParameterMapBase.h"
 #include "NiagaraParameterMapHistory.h"
 #include "NiagaraScriptSource.h"
+#include "NiagaraStatelessExpressionTypeData.h"
 #include "NiagaraSystem.h"
 
 #define LOCTEXT_NAMESPACE "NiagaraDistributionPropertyCustomization"
@@ -152,6 +153,12 @@ public:
 		if (bAllowBinding)
 		{
 			OutSupportedModes.Add(ENiagaraDistributionEditorMode::Binding);
+
+			const FNiagaraStatelessExpressionTypeData& TypeData = FNiagaraStatelessExpressionTypeData::GetTypeData(SourceDistribution->GetBindingTypeDef());
+			if (TypeData.IsValid())
+			{
+				OutSupportedModes.Add(ENiagaraDistributionEditorMode::Expression);
+			}
 		}
 		if (bDisplayAsColor)
 		{
@@ -212,6 +219,8 @@ public:
 			{
 			case ENiagaraDistributionMode::Binding:
 				return ENiagaraDistributionEditorMode::Binding;
+			case ENiagaraDistributionMode::Expression:
+				return ENiagaraDistributionEditorMode::Expression;
 			case ENiagaraDistributionMode::NonUniformConstant:
 				return ENiagaraDistributionEditorMode::ColorConstant;
 			case ENiagaraDistributionMode::NonUniformRange:
@@ -228,6 +237,8 @@ public:
 			{
 			case ENiagaraDistributionMode::Binding:
 				return ENiagaraDistributionEditorMode::Binding;
+			case ENiagaraDistributionMode::Expression:
+				return ENiagaraDistributionEditorMode::Expression;
 			case ENiagaraDistributionMode::UniformConstant:
 				return ENiagaraDistributionEditorMode::Constant;
 			case ENiagaraDistributionMode::UniformRange:
@@ -244,6 +255,8 @@ public:
 			{
 			case ENiagaraDistributionMode::Binding:
 				return ENiagaraDistributionEditorMode::Binding;
+			case ENiagaraDistributionMode::Expression:
+				return ENiagaraDistributionEditorMode::Expression;
 			case ENiagaraDistributionMode::UniformConstant:
 				return ENiagaraDistributionEditorMode::UniformConstant;
 			case ENiagaraDistributionMode::NonUniformConstant:
@@ -268,6 +281,8 @@ public:
 		{
 		case ENiagaraDistributionEditorMode::Binding:
 			return ENiagaraDistributionMode::Binding;
+		case ENiagaraDistributionEditorMode::Expression:
+			return ENiagaraDistributionMode::Expression;
 		case ENiagaraDistributionEditorMode::Constant:
 		case ENiagaraDistributionEditorMode::UniformConstant:
 			return ENiagaraDistributionMode::UniformConstant;
@@ -458,6 +473,36 @@ public:
 		PropertyHandle->NotifyPostChange(bContinuousChangeActive ? EPropertyChangeType::Interactive : EPropertyChangeType::ValueSet);
 	}
 
+	virtual FNiagaraTypeDefinition GetExpressionTypeDef() const override
+	{
+		return SourceDistribution->GetBindingTypeDef();
+	}
+
+	virtual FInstancedStruct& GetExpressionRoot() const
+	{
+		return SourceDistribution->ParameterExpression;
+	}
+
+	virtual void ExecuteTransaction(FText TransactionText, TFunction<void()> TransactionFunc) override
+	{
+		if (UObject* OwnerObject = OwnerObjectWeak.Get())
+		{
+			if (bContinuousTransactionPending)
+			{
+				bContinuousTransactionPending = false;
+				ContinuousTransactionIndex = GEditor->BeginTransaction(TransactionText);
+			}
+
+			const FScopedTransaction Transaction(TransactionText, ContinuousTransactionIndex.IsSet() == false);
+			OwnerObject->Modify();
+			PropertyHandle->NotifyPreChange();
+
+			TransactionFunc();
+
+			PropertyHandle->NotifyPostChange(bContinuousChangeActive ? EPropertyChangeType::Interactive : EPropertyChangeType::ValueSet);
+		}
+	}
+
 	virtual FNiagaraVariableBase GetBindingValue() const
 	{
 		UObject* OwnerObject = OwnerObjectWeak.Get();
@@ -470,27 +515,13 @@ public:
 
 	virtual void SetBindingValue(FNiagaraVariableBase Binding)
 	{
-		UObject* OwnerObject = OwnerObjectWeak.Get();
-		if (OwnerObject == nullptr)
-		{
-			return;
-		}
-
-		FText TransactionText = LOCTEXT("SetBinding", "Set binding value");
-
-		if (bContinuousTransactionPending)
-		{
-			bContinuousTransactionPending = false;
-			ContinuousTransactionIndex = GEditor->BeginTransaction(TransactionText);
-		}
-
-		const FScopedTransaction Transaction(TransactionText, ContinuousTransactionIndex.IsSet() == false);
-		OwnerObject->Modify();
-		PropertyHandle->NotifyPreChange();
-
-		SourceDistribution->ParameterBinding = Binding;
-
-		PropertyHandle->NotifyPostChange(bContinuousChangeActive ? EPropertyChangeType::Interactive : EPropertyChangeType::ValueSet);
+		ExecuteTransaction(
+			LOCTEXT("SetBinding", "Set binding value"),
+			[&]()
+			{
+				SourceDistribution->ParameterBinding = Binding;
+			}
+		);
 	}
 
 	virtual TArray<FNiagaraVariableBase> GetAvailableBindings() const
diff --git a/Engine/Plugins/FX/Niagara/Source/NiagaraEditor/Private/Stateless/NiagaraStatelessExpressionTypeData.cpp b/Engine/Plugins/FX/Niagara/Source/NiagaraEditor/Private/Stateless/NiagaraStatelessExpressionTypeData.cpp
new file mode 100644
--- /dev/null
+++ b/Engine/Plugins/FX/Niagara/Source/NiagaraEditor/Private/Stateless/NiagaraStatelessExpressionTypeData.cpp
@@ -0,0 +1,129 @@
+// Copyright Epic Games, Inc. All Rights Reserved.
+
+#include "NiagaraStatelessExpressionTypeData.h"
+#include "Stateless/NiagaraStatelessExpression.h"
+
+#include "HAL/IConsoleManager.h"
+#include "UObject/UObjectIterator.h"
+
+static bool GNiagaraStatelessExpressionsEnabled = false;
+static FAutoConsoleVariableRef CVarNiagaraStatelessExpressionsEnabled(
+	TEXT("fx.NiagaraStateless.ExpressionsEnabled"),
+	GNiagaraStatelessExpressionsEnabled,
+	TEXT("CVar to enable stateless expressions"),
+	ECVF_Default
+);
+
+bool FNiagaraStatelessExpressionTypeData::IsValid() const
+{
+	return GNiagaraStatelessExpressionsEnabled && ValueExpression.IsValid();
+}
+
+bool FNiagaraStatelessExpressionTypeData::ContainsExpression(const FInstancedStruct* Expresssion) const
+{
+	const UScriptStruct* ExpressionType = Expresssion ? Expresssion->GetScriptStruct() : nullptr;
+	return
+		ExpressionType != nullptr &&
+		(
+			ExpressionType == ValueExpression.Get() ||
+			ExpressionType == BindingExpression.Get() ||
+			OperationExpressions.Contains(ExpressionType)
+		);
+}
+
+bool FNiagaraStatelessExpressionTypeData::IsBindingExpression(const FInstancedStruct* Expression) const
+{
+	return Expression && Expression->GetScriptStruct() == BindingExpression;
+}
+
+bool FNiagaraStatelessExpressionTypeData::IsValueExpression(const FInstancedStruct* Expression) const
+{
+	return Expression && Expression->GetScriptStruct() == ValueExpression;
+}
+
+bool FNiagaraStatelessExpressionTypeData::IsOperationExpression(const FInstancedStruct* Expression) const
+{
+	return Expression && OperationExpressions.Contains(Expression->GetScriptStruct());
+}
+
+FName FNiagaraStatelessExpressionTypeData::GetBindingName(const FInstancedStruct* Expression) const
+{
+	check(Expression && Expression->GetScriptStruct() == BindingExpression);
+	return BindingNameField->GetPropertyValue_InContainer(Expression->GetMemory());
+}
+
+FInstancedStruct FNiagaraStatelessExpressionTypeData::MakeBindingStruct(FName BindingName) const
+{
+	FInstancedStruct NewExpression(BindingExpression.Get());
+	BindingNameField->SetPropertyValue_InContainer(NewExpression.GetMutableMemory(), BindingName);
+	return NewExpression;
+}
+
+const FNiagaraStatelessExpressionTypeData& FNiagaraStatelessExpressionTypeData::GetTypeData(const FNiagaraTypeDefinition& TypeDef)
+{
+	static TMap<FNiagaraTypeDefinition, FNiagaraStatelessExpressionTypeData> TypeDefToExpressionData;
+	if (TypeDefToExpressionData.IsEmpty())
+	{
+		static const FName NAME_BindingExpression("BindingExpression");
+		static const FName NAME_ValueExpression("ValueExpression");
+		static const FName NAME_OperationExpression("OperationExpression");
+
+		for (TObjectIterator<UScriptStruct> StructIt; StructIt; ++StructIt)
+		{
+			UScriptStruct* ScriptStruct = *StructIt;
+			if (!ScriptStruct || !ScriptStruct->IsChildOf<FNiagaraStatelessExpression>())
+			{
+				continue;
+			}
+
+			const bool bIsBinding = ScriptStruct->HasMetaData(NAME_BindingExpression);
+			const bool bIsValue = ScriptStruct->HasMetaData(NAME_ValueExpression);
+			const bool bIsOperation = ScriptStruct->HasMetaData(NAME_OperationExpression);
+			const bool bIsAny = bIsBinding || bIsValue || bIsOperation;
+			if (!bIsAny)
+			{
+				continue;
+			}
+
+			FInstancedStruct TempStruct;
+			TempStruct.InitializeAs(ScriptStruct, nullptr);
+			const FNiagaraStatelessExpression& Expression = TempStruct.Get<FNiagaraStatelessExpression>();
+			FNiagaraStatelessExpressionTypeData& ExpressionTypeData = TypeDefToExpressionData.FindOrAdd(Expression.GetOutputTypeDef());
+			ExpressionTypeData.TypeDef = Expression.GetOutputTypeDef();
+
+			if (bIsBinding)
+			{
+				FNameProperty* NameProperty = CastField<FNameProperty>(ScriptStruct->ChildProperties);
+				if (ensure(NameProperty && NameProperty->Next == nullptr))
+				{
+					ExpressionTypeData.BindingExpression = ScriptStruct;
+					ExpressionTypeData.BindingNameField = NameProperty;
+				}
+			}
+			else if (bIsValue)
+			{
+				FProperty* ValueProperty = CastField<FProperty>(ScriptStruct->ChildProperties);
+				if (ensure(ValueProperty && ValueProperty->Next == nullptr))
+				{
+					if (ensure(FNiagaraTypeDefinition(ValueProperty, FNiagaraTypeDefinition::EAllowUnfriendlyStruct::Deny) == Expression.GetOutputTypeDef()))
+					{
+						ExpressionTypeData.ValueExpression = ScriptStruct;
+						ExpressionTypeData.ValueProperty = ValueProperty;
+					}
+				}
+			}
+			else if (bIsOperation)
+			{
+				ExpressionTypeData.OperationExpressions.Add(ScriptStruct);
+			}
+		}
+	}
+
+	return TypeDefToExpressionData.FindOrAdd(TypeDef);
+}
+
+const FNiagaraStatelessExpressionTypeData& FNiagaraStatelessExpressionTypeData::GetTypeData(const FInstancedStruct* Expression)
+{
+	check(Expression);
+	return GetTypeData(Expression->Get<const FNiagaraStatelessExpression>().GetOutputTypeDef());
+}
diff --git a/Engine/Plugins/FX/Niagara/Source/NiagaraEditor/Private/Stateless/NiagaraStatelessExpressionTypeData.h b/Engine/Plugins/FX/Niagara/Source/NiagaraEditor/Private/Stateless/NiagaraStatelessExpressionTypeData.h
new file mode 100644
--- /dev/null
+++ b/Engine/Plugins/FX/Niagara/Source/NiagaraEditor/Private/Stateless/NiagaraStatelessExpressionTypeData.h
@@ -0,0 +1,33 @@
+// Copyright Epic Games, Inc. All Rights Reserved.
+
+#pragma once
+
+#include "NiagaraCommon.h"
+
+struct FInstancedStruct;
+
+struct FNiagaraStatelessExpressionTypeData
+{
+	FNiagaraTypeDefinition					TypeDef;
+
+	TWeakObjectPtr<UScriptStruct>			ValueExpression;
+	FProperty*								ValueProperty = nullptr;
+
+	TWeakObjectPtr<UScriptStruct>			BindingExpression;
+	FNameProperty*							BindingNameField = nullptr;
+
+	TArray<TWeakObjectPtr<UScriptStruct>>	OperationExpressions;
+
+	bool IsValid() const;
+	bool ContainsExpression(const FInstancedStruct* Expresssion) const;
+
+	bool IsBindingExpression(const FInstancedStruct* Expression) const;
+	bool IsValueExpression(const FInstancedStruct* Expression) const;
+	bool IsOperationExpression(const FInstancedStruct* Expression) const;
+
+	FName GetBindingName(const FInstancedStruct* Expression) const;
+	FInstancedStruct MakeBindingStruct(FName BindingName) const;
+
+	static const FNiagaraStatelessExpressionTypeData& GetTypeData(const FNiagaraTypeDefinition& TypeDef);
+	static const FNiagaraStatelessExpressionTypeData& GetTypeData(const FInstancedStruct* Expression);
+};
diff --git a/Engine/Plugins/FX/Niagara/Source/NiagaraEditor/Private/Stateless/SNiagaraStatelessExpressionWidget.cpp b/Engine/Plugins/FX/Niagara/Source/NiagaraEditor/Private/Stateless/SNiagaraStatelessExpressionWidget.cpp
new file mode 100644
--- /dev/null
+++ b/Engine/Plugins/FX/Niagara/Source/NiagaraEditor/Private/Stateless/SNiagaraStatelessExpressionWidget.cpp
@@ -0,0 +1,572 @@
+// Copyright Epic Games, Inc. All Rights Reserved.
+
+#include "SNiagaraStatelessExpressionWidget.h"
+
+#include "Stateless/NiagaraStatelessExpression.h"
+#include "NiagaraStatelessExpressionTypeData.h"
+
+#include "NiagaraEditorModule.h"
+#include "INiagaraEditorTypeUtilities.h"
+#include "NiagaraEditorStyle.h"
+
+#include "Curves/CurveOwnerInterface.h"
+#include "EdGraph/EdGraphSchema.h"
+#include "Framework/MultiBox/MultiBoxBuilder.h"
+#include "NiagaraEditorStyle.h"
+#include "SColorGradientEditor.h"
+#include "SNiagaraParameterEditor.h"
+#include "TypeEditorUtilities/NiagaraFloatTypeEditorUtilities.h"
+#include "Widgets/Input/SComboButton.h"
+#include "Widgets/Input/SNumericEntryBox.h"
+#include "Widgets/Input/SVectorInputBox.h"
+#include "Widgets/Layout/SBox.h"
+#include "Widgets/Layout/SGridPanel.h"
+#include "Widgets/NiagaraDistributionEditorUtilities.h"
+#include "Widgets/SNiagaraColorEditor.h"
+#include "Widgets/SNiagaraDistributionCurveEditor.h"
+#include "Widgets/SNiagaraParameterName.h"
+
+#include "Modules/ModuleManager.h"
+#include "IPropertyRowGenerator.h"
+#include "PropertyEditorModule.h"
+#include "IDetailTreeNode.h"
+#include "IStructureDetailsView.h"
+
+#include "SGraphActionMenu.h"
+
+#define LOCTEXT_NAMESPACE "NiagaraStatelessExpression"
+
+/////////////////////////////////////////////////////////////////////////////////////// 
+namespace SNiagaraStatelessExpressionPrivate
+{
+	DECLARE_DELEGATE_TwoParams(FOnExecuteTransaction, FText, TFunction<void()>)
+	DECLARE_DELEGATE_RetVal(TArray<FNiagaraVariableBase>, FOnGetAvailableBindings)
+
+	/////////////////////////////////////////////////////////////////////////////////////////////////// 
+
+	struct FSelectExpressionTypeAction : public FEdGraphSchemaAction
+	{
+		FSelectExpressionTypeAction() : FEdGraphSchemaAction() {}
+
+		FSelectExpressionTypeAction(FInstancedStruct InExpressionStruct, FText InCategory, FText InDisplayName, FText InToolTipText)
+			: FEdGraphSchemaAction(
+				MoveTemp(InCategory),
+				MoveTemp(InDisplayName),
+				MoveTemp(InToolTipText),
+				0,
+				FText()
+			)
+			, ExpressionStruct(MoveTemp(InExpressionStruct))
+		{
+		}
+
+		//~ Begin FEdGraphSchemaAction Interface
+		virtual UEdGraphNode* PerformAction(class UEdGraph* ParentGraph, UEdGraphPin* FromPin, const FVector2D Location, bool bSelectNewNode = true) override
+		{
+			return nullptr;
+		}
+		//~ End FEdGraphSchemaAction Interface
+
+		FInstancedStruct ExpressionStruct;
+	};
+
+	class SExpressionTypeSelector : public SCompoundWidget
+	{
+	public:
+		DECLARE_DELEGATE_OneParam(FOnActionSelected, const FInstancedStruct&)
+
+		SLATE_BEGIN_ARGS(SExpressionTypeSelector) { }
+			SLATE_ARGUMENT(FNiagaraTypeDefinition,	ExpressionTypeDef)
+			SLATE_EVENT(FOnActionSelected,			OnActionSelected)
+			SLATE_EVENT(FOnGetAvailableBindings,	OnGetAvailableBindings)
+		SLATE_END_ARGS();
+
+		void Construct(const FArguments& InArgs)
+		{
+			ExpressionTypeDef		= InArgs._ExpressionTypeDef;
+			ActionSelected			= InArgs._OnActionSelected;
+			GetAvailableBindings	= InArgs._OnGetAvailableBindings;
+
+			ChildSlot
+			[
+				SNew(SComboButton)					
+				.ButtonStyle(FAppStyle::Get(), "HoverHintOnly")
+				.ForegroundColor(FSlateColor::UseForeground())
+				.OnGetMenuContent(this, &SExpressionTypeSelector::OnGetMenuContent)
+				.ContentPadding(2)
+				.MenuPlacement(MenuPlacement_BelowRightAnchor)
+				//.ToolTipText(this, &SExpressionTypeSelector::GetTooltipText)
+			];
+		}
+
+		TSharedRef<SWidget> OnGetMenuContent() const
+		{
+			FGraphActionMenuBuilder MenuBuilder;
+
+			return SNew(SBorder)
+				.BorderImage(FAppStyle::GetBrush("Menu.Background"))
+				.Padding(5)
+				[
+					SNew(SBox)
+					[
+						SNew(SGraphActionMenu)
+						.OnActionSelected(this, &SExpressionTypeSelector::OnActionSelected)
+						.OnCreateWidgetForAction(this, &SExpressionTypeSelector::OnCreateWidgetForAction)
+						.OnCollectAllActions(this, &SExpressionTypeSelector::CollectAllActions)
+						.AutoExpandActionMenu(true)
+						.ShowFilterTextBox(true)
+					]
+				];
+		}
+
+		void OnActionSelected(const TArray<TSharedPtr<FEdGraphSchemaAction>>& SelectedActions, ESelectInfo::Type InSelectionType) const
+		{
+			if ((InSelectionType != ESelectInfo::OnMouseClick && InSelectionType != ESelectInfo::OnKeyPress))
+			{
+				return;
+			}
+
+			FSlateApplication::Get().DismissAllMenus();
+
+			for (const TSharedPtr<FEdGraphSchemaAction>& SelectedAction : SelectedActions)
+			{
+				if (SelectedAction.IsValid() == false)
+				{
+					continue;
+				}
+
+				ActionSelected.ExecuteIfBound(static_cast<FSelectExpressionTypeAction*>(SelectedAction.Get())->ExpressionStruct);
+			}
+		}
+
+		TSharedRef<SWidget> OnCreateWidgetForAction(FCreateWidgetForActionData* const InCreateData) const
+		{
+			const FNiagaraStatelessExpressionTypeData& TypeData = FNiagaraStatelessExpressionTypeData::GetTypeData(ExpressionTypeDef);
+			if (TypeData.IsValid() == false)
+			{
+				return SNullWidget::NullWidget;
+			}
+
+			FSelectExpressionTypeAction* ExpressionAction = static_cast<FSelectExpressionTypeAction*>(InCreateData->Action.Get());
+			if (ExpressionAction->ExpressionStruct.GetScriptStruct() == TypeData.BindingExpression)
+			{
+				const FName BindingName = TypeData.GetBindingName(&ExpressionAction->ExpressionStruct);
+
+				return
+					SNew(SVerticalBox)
+					+ SVerticalBox::Slot()
+					.AutoHeight()
+					[
+						SNew(SNiagaraParameterName)
+						.ParameterName(BindingName)
+						.IsReadOnly(true)
+					];
+			}
+			else
+			{
+				return
+					SNew(SVerticalBox)
+					+ SVerticalBox::Slot()
+					.AutoHeight()
+					[
+						SNew(STextBlock)
+						.Text(InCreateData->Action->GetMenuDescription())
+						.ToolTipText(InCreateData->Action->GetTooltipDescription())
+					];
+			}
+		}
+
+		void CollectAllActions(FGraphActionListBuilderBase& OutAllActions) const
+		{
+			const FNiagaraStatelessExpressionTypeData& TypeData = FNiagaraStatelessExpressionTypeData::GetTypeData(ExpressionTypeDef);
+			if (TypeData.IsValid() == false)
+			{
+				return;
+			}
+
+			// Add Value expression
+			if (UScriptStruct* ValueExpression = TypeData.ValueExpression.Get())
+			{
+				OutAllActions.AddAction(
+					MakeShared<FSelectExpressionTypeAction>(
+						FInstancedStruct(ValueExpression),
+						FText(),
+						LOCTEXT("NewLocalValue", "New Local Value"),
+						LOCTEXT("NewLocalValueTooltip", "Set to a constant value")
+					)
+				);
+			}
+
+			// Add Binding expressions
+			if (UScriptStruct* BindingExpression = TypeData.BindingExpression.Get())
+			{
+				if (GetAvailableBindings.IsBound())
+				{
+					const TArray<FNiagaraVariableBase> AvailableBindings = GetAvailableBindings.Execute();
+					for (const FNiagaraVariableBase& Binding : AvailableBindings)
+					{
+						if (Binding.GetType() != ExpressionTypeDef)
+						{
+							continue;
+						}
+
+						OutAllActions.AddAction(
+							MakeShared<FSelectExpressionTypeAction>(
+								TypeData.MakeBindingStruct(Binding.GetName()),
+								LOCTEXT("LinkInput", "Link Input"),
+								FText::FromName(Binding.GetName()),
+								LOCTEXT("LinkInputTooltip", "Set the parameter")	//-TOOD:
+							)
+						);
+					}
+				}
+			}
+
+			// Add Operation expressions
+			for (TWeakObjectPtr<UScriptStruct> WeakExpressionStruct : TypeData.OperationExpressions)
+			{
+				if (UScriptStruct* ExpressionStruct = WeakExpressionStruct.Get())
+				{
+					OutAllActions.AddAction(
+						MakeShared<FSelectExpressionTypeAction>(
+							FInstancedStruct(ExpressionStruct),
+							LOCTEXT("Operation", "Operation"),
+							ExpressionStruct->GetDisplayNameText(),
+							ExpressionStruct->GetToolTipText()
+						)
+					);
+				}
+			}
+		}
+
+		FNiagaraTypeDefinition	ExpressionTypeDef;
+		FOnActionSelected		ActionSelected;
+		FOnGetAvailableBindings	GetAvailableBindings;
+	};
+
+	/////////////////////////////////////////////////////////////////////////////////////////////////// 
+
+	class SExpressionWidget : public SCompoundWidget
+	{
+	public:
+		SLATE_BEGIN_ARGS(SExpressionWidget) { }
+			SLATE_EVENT(FOnGetAvailableBindings,	OnGetAvailableBindings)
+			SLATE_EVENT(FOnExecuteTransaction,		OnExecuteTransaction)
+		SLATE_END_ARGS();
+
+		void Construct(const FArguments& InArgs, FInstancedStruct* InExpressionStruct, int32 InDepth, FText InDisplayName)
+		{
+			Depth					= InDepth;
+			DisplayName				= InDisplayName;
+			ExpressionStruct		= InExpressionStruct;
+			WidgetContainer			= SNew(SVerticalBox);
+			GetAvailableBindings	= InArgs._OnGetAvailableBindings;
+			ExecuteTransaction		= InArgs._OnExecuteTransaction;
+
+			ChildSlot
+			[
+				SNew(SHorizontalBox)
+				+SHorizontalBox::Slot()
+				.Padding(10 * Depth, 0, 0, 0)
+				[
+					WidgetContainer.ToSharedRef()
+				]
+			];
+
+			RebuildChildren();
+		}
+
+		void Construct(const FArguments& InArgs, FInstancedStruct* InExpressionStruct)
+		{
+			Construct(InArgs, InExpressionStruct, 0, FText());
+		}
+
+	private:
+		TSharedRef<SWidget> GetExpressionWidget() const
+		{
+			const FNiagaraStatelessExpressionTypeData& TypeData = FNiagaraStatelessExpressionTypeData::GetTypeData(ExpressionStruct);
+
+			// Show binding information
+			if (TypeData.IsBindingExpression(ExpressionStruct))
+			{
+				return
+					SNew(SNiagaraParameterName)
+					.ParameterName(TypeData.GetBindingName(ExpressionStruct))
+					.IsReadOnly(true);
+			}
+
+			// Show editable value widget		
+			if (TypeData.IsValueExpression(ExpressionStruct))
+			{
+				FNiagaraEditorModule& NiagaraEditorModule = FModuleManager::GetModuleChecked<FNiagaraEditorModule>("NiagaraEditor");
+				TSharedPtr<INiagaraEditorTypeUtilities, ESPMode::ThreadSafe> TypeEditorUtilities = NiagaraEditorModule.GetTypeUtilities(TypeData.TypeDef);
+				if (TypeEditorUtilities.IsValid())
+				{
+					TSharedRef<FStructOnScope> ValueStructOnScope = MakeShared<FStructOnScope>(
+						TypeData.TypeDef.GetStruct(),
+						const_cast<uint8*>(TypeData.ValueProperty->ContainerPtrToValuePtr<uint8>(ExpressionStruct->GetMemory()))	//-TODO: Remove?
+					);
+
+					FNiagaraInputParameterCustomization CustomizationOptions;
+					CustomizationOptions.bBroadcastValueChangesOnCommitOnly = true;
+					TSharedPtr<SNiagaraParameterEditor> ValueParameterEditor = TypeEditorUtilities->CreateParameterEditor(TypeData.TypeDef, EUnit::Unspecified, CustomizationOptions);
+					ValueParameterEditor->UpdateInternalValueFromStruct(ValueStructOnScope);
+					ValueParameterEditor->SetOnValueChanged(
+						SNiagaraParameterEditor::FOnValueChange::CreateSP(this, &SExpressionWidget::OnValueChanged, ValueParameterEditor, ValueStructOnScope)
+					);
+
+					return
+						ValueParameterEditor.ToSharedRef();
+				}
+				return SNullWidget::NullWidget;
+			}
+
+			// Show selected expression
+			check(TypeData.IsOperationExpression(ExpressionStruct));
+
+			return
+				SNew(STextBlock)
+				.Text(ExpressionStruct->GetScriptStruct()->GetDisplayNameText())
+				.TextStyle(FNiagaraEditorStyle::Get(), "NiagaraEditor.ParameterText");
+		}
+
+		TSharedRef<SWidget> CreateTypeSelector()
+		{
+			const FNiagaraStatelessExpressionTypeData& TypeData = FNiagaraStatelessExpressionTypeData::GetTypeData(ExpressionStruct);
+			return
+				SNew(SExpressionTypeSelector)
+				.ExpressionTypeDef(TypeData.TypeDef)
+				.OnActionSelected(this, &SExpressionWidget::OnChangeExpression)
+				.OnGetAvailableBindings(GetAvailableBindings);
+		}
+
+		void RebuildChildren()
+		{
+			// Clear existing
+			WidgetContainer->ClearChildren();
+
+			// Build row
+			{
+				TSharedRef<SHorizontalBox> HorizontalBox = SNew(SHorizontalBox);
+
+				if (DisplayName.IsEmpty() == false)
+				{
+					HorizontalBox->AddSlot()
+					.Padding(0, 0, 5, 0)
+					.HAlign(HAlign_Left)
+					.AutoWidth()
+					[
+						SNew(STextBlock)
+						.Text(DisplayName)
+					];
+				}
+
+				HorizontalBox->AddSlot()
+				.HAlign(HAlign_Left)
+				.AutoWidth()
+				[
+					GetExpressionWidget()
+				];
+
+				HorizontalBox->AddSlot()
+				.HAlign(HAlign_Right)
+				.AutoWidth()
+				[
+					CreateTypeSelector()
+				];
+
+				WidgetContainer->AddSlot()
+				.VAlign(VAlign_Top)
+				.AutoHeight()
+				[
+					HorizontalBox
+				];
+			}
+
+			// For operations recurse into child structures
+			const FNiagaraStatelessExpressionTypeData& TypeData = FNiagaraStatelessExpressionTypeData::GetTypeData(ExpressionStruct);
+			if (TypeData.IsOperationExpression(ExpressionStruct))
+			{
+				for (TFieldIterator<FProperty> PropertyIt(ExpressionStruct->GetScriptStruct()); PropertyIt; ++PropertyIt)
+				{
+					const FStructProperty* StructProperty = CastField<const FStructProperty>(*PropertyIt);
+					if (StructProperty == nullptr || StructProperty->Struct != FInstancedStruct::StaticStruct())
+					{
+						continue;
+					}
+
+					FInstancedStruct* InnerExpressionStruct = StructProperty->ContainerPtrToValuePtr<FInstancedStruct>(ExpressionStruct->GetMutableMemory());
+
+					WidgetContainer->AddSlot()
+					.VAlign(VAlign_Top)
+					.AutoHeight()
+					[
+						SNew(SExpressionWidget, InnerExpressionStruct, Depth + 1, StructProperty->GetDisplayNameText())
+						.OnGetAvailableBindings(GetAvailableBindings)
+						.OnExecuteTransaction(ExecuteTransaction)
+					];
+				}
+			}
+		}
+
+		void OnChangeExpression(const FInstancedStruct& NewExpression)
+		{
+			WidgetContainer->ClearChildren();
+
+			ExecuteTransaction.ExecuteIfBound(
+				LOCTEXT("ChangeExpression", " Change Expression"),		//-TODO: Make more verbose
+				[&]()
+				{
+					ExpressionStruct->InitializeAs(NewExpression.GetScriptStruct(), NewExpression.GetMemory());
+				}
+			);
+
+			RebuildChildren();
+		}
+
+		void OnValueChanged(TSharedPtr<SNiagaraParameterEditor> ParameterEditor, TSharedRef<FStructOnScope> ValueStructOnScope) const
+		{
+			ExecuteTransaction.ExecuteIfBound(
+				LOCTEXT("ChangeValue", " Change Expression Value"),		//-TODO: Make more verbose
+				[&]()
+				{
+					ParameterEditor->UpdateStructFromInternalValue(ValueStructOnScope);
+				}
+			);
+		}
+
+	private:
+		int32						Depth = 0;
+		FText						DisplayName;
+		FInstancedStruct*			ExpressionStruct = nullptr;
+		TSharedPtr<SVerticalBox>	WidgetContainer;
+		FOnGetAvailableBindings		GetAvailableBindings;
+		FOnExecuteTransaction		ExecuteTransaction;
+	};
+}
+
+/////////////////////////////////////////////////////////////////////////////////////////////////// 
+
+void SNiagaraStatelessExpressionWidget::Construct(const FArguments& InArgs, TSharedRef<INiagaraDistributionAdapter> InDistributionAdapter)
+{
+	DistributionAdapter = InDistributionAdapter;
+
+	// Get the type data
+	FInstancedStruct& ExpressionStruct = DistributionAdapter->GetExpressionRoot();
+	const FNiagaraStatelessExpressionTypeData& TypeData = FNiagaraStatelessExpressionTypeData::GetTypeData(DistributionAdapter->GetExpressionTypeDef());
+	if (TypeData.IsValid() == false)
+	{
+		checkNoEntry();		//-TODO:
+		return;
+	}
+
+	// Make sure root is initialized, all other expression should self initialize and provide the type requirements
+	if (TypeData.ContainsExpression(&ExpressionStruct) == false)
+	{
+		ExpressionStruct.InitializeAs(TypeData.ValueExpression.Get());
+	}
+
+	ChildSlot
+	[
+		SNew(SNiagaraStatelessExpressionPrivate::SExpressionWidget, &ExpressionStruct)
+		.OnGetAvailableBindings(this, &SNiagaraStatelessExpressionWidget::GetAvailableBindings)
+		.OnExecuteTransaction(this, &SNiagaraStatelessExpressionWidget::ExecuteTransaction)
+	];
+
+#if 0
+
+
+	//const FNiagaraTypeDefinition TypeDef = DistributionAdapter->GetExpressionTypeDef();
+
+
+	//FStructIterat
+
+	//-TODO: Check to see if the expression is valid or not, if not create a default entry
+	//FInstancedStruct* RootExpression = DistributionAdapter->GetRootExpression()
+
+	//FInstancedStruct* RootExpression = DistributionAdapter->GetRootExpression();
+	//DistributionAdapter->GetExpressionTypeDef()
+	// Constant
+	// Operation	:dropdown:
+	// - Constant	:dropdown: (inline value)
+	// - Operation	:dropdown:
+	// -- Constant	:dropdown: (inline value)
+	// -- Binding	:dropdown: (inline value)
+
+	// For Each Property
+	// - FInstanceStruct -> Drop Down Selection for the correct type def
+	// - FName -> Show binding information
+
+	FPropertyEditorModule& PropertyEditorModule = FModuleManager::LoadModuleChecked<FPropertyEditorModule>("PropertyEditor");
+#if 0
+	FPropertyRowGeneratorArgs Args;
+	//Args.NotifyHook = this;
+	PropertyRowGenerator = PropertyEditorModule.CreatePropertyRowGenerator(Args);
+	PropertyRowGenerator->SetStructure(InDistributionAdapter->MakeStructOnScope());
+
+	TArray<TSharedRef<IDetailTreeNode>> RootNodes = PropertyRowGenerator->GetRootTreeNodes();
+	if (RootNodes.Num() == 1)
+	{
+		static TArray<TSharedRef<IDetailTreeNode>> ChildNodes;
+		RootNodes[0]->GetChildren(ChildNodes);
+
+		FNodeWidgets NodeWidgets = ChildNodes[1]->CreateNodeWidgets();
+		ChildSlot
+		[
+			NodeWidgets.ValueWidget.ToSharedRef()
+		];
+	}
+#else
+	FDetailsViewArgs DetailsViewArgs;
+	DetailsViewArgs.bHideSelectionTip = true;
+	DetailsViewArgs.bAllowSearch = false;
+	//DetailsViewArgs.NotifyHook = DebugHudSettings;
+
+	FStructureDetailsViewArgs StructureViewArgs;
+	StructureViewArgs.bShowObjects = true;
+	StructureViewArgs.bShowAssets = true;
+	StructureViewArgs.bShowClasses = true;
+	StructureViewArgs.bShowInterfaces = true;
+
+	TSharedRef<IStructureDetailsView> StructureDetailsView = PropertyEditorModule.CreateStructureDetailView(
+		DetailsViewArgs, StructureViewArgs, DistributionAdapter->MakeStructOnScope()
+	);
+
+	ChildSlot
+	[
+		StructureDetailsView->GetWidget().ToSharedRef()
+	];
+#endif
+	//FNodeWidgets NodeWidgets = InItem->CreateNodeWidgets();
+	//if (NodeWidgets.WholeRowWidget.IsValid())
+	//{
+	//	TableRow->SetContent(NodeWidgets.WholeRowWidget.ToSharedRef());
+	//}
+	//else if (NodeWidgets.NameWidget.IsValid() && NodeWidgets.ValueWidget.IsValid())
+	//{
+	//	TableRow->SetContent(
+	//		SNew(SHorizontalBox)
+	//		+ SHorizontalBox::Slot()
+	//		[
+	//			NodeWidgets.NameWidget.ToSharedRef()
+	//		]
+	//		+ SHorizontalBox::Slot()
+	//		[
+	//			NodeWidgets.ValueWidget.ToSharedRef()
+	//		]
+	//	);
+	//}
+#endif
+}
+
+TArray<FNiagaraVariableBase> SNiagaraStatelessExpressionWidget::GetAvailableBindings() const
+{
+	return DistributionAdapter->GetAvailableBindings();
+}
+
+void SNiagaraStatelessExpressionWidget::ExecuteTransaction(FText TransactionText, TFunction<void()> TransactionFunc)
+{
+	DistributionAdapter->ExecuteTransaction(MoveTemp(TransactionText), MoveTemp(TransactionFunc));
+}
+
+#undef LOCTEXT_NAMESPACE
diff --git a/Engine/Plugins/FX/Niagara/Source/NiagaraEditor/Private/Stateless/SNiagaraStatelessExpressionWidget.h b/Engine/Plugins/FX/Niagara/Source/NiagaraEditor/Private/Stateless/SNiagaraStatelessExpressionWidget.h
new file mode 100644
--- /dev/null
+++ b/Engine/Plugins/FX/Niagara/Source/NiagaraEditor/Private/Stateless/SNiagaraStatelessExpressionWidget.h
@@ -0,0 +1,27 @@
+// Copyright Epic Games, Inc. All Rights Reserved.
+
+#pragma once
+
+#include "Widgets/SCompoundWidget.h"
+
+#include "Widgets/INiagaraDistributionAdapter.h"
+
+class INiagaraDistributionAdapter;
+class IPropertyRowGenerator;
+
+class SNiagaraStatelessExpressionWidget : public SCompoundWidget
+{
+public:
+	SLATE_BEGIN_ARGS(SNiagaraStatelessExpressionWidget) { }
+	SLATE_END_ARGS();
+
+	void Construct(const FArguments& InArgs, TSharedRef<INiagaraDistributionAdapter> InDistributionAdapter);
+
+private:
+	TArray<FNiagaraVariableBase> GetAvailableBindings() const;
+	void ExecuteTransaction(FText TransactionText, TFunction<void()> TransactionFunc);
+
+private:
+	TSharedPtr<INiagaraDistributionAdapter>	DistributionAdapter;
+	TSharedPtr<IPropertyRowGenerator>		PropertyRowGenerator;
+};
diff --git a/Engine/Plugins/FX/Niagara/Source/NiagaraEditor/Private/TraversalCache/TraversalCache.cpp b/Engine/Plugins/FX/Niagara/Source/NiagaraEditor/Private/TraversalCache/TraversalCache.cpp
--- a/Engine/Plugins/FX/Niagara/Source/NiagaraEditor/Private/TraversalCache/TraversalCache.cpp
+++ b/Engine/Plugins/FX/Niagara/Source/NiagaraEditor/Private/TraversalCache/TraversalCache.cpp
@@ -395,7 +395,7 @@ void FTraversalCache::GetTopLevelScriptFunctionInputContextsInternal(const UNiag
 			if (ensureMsgf(ParsedStaticFunctionInput.DataOffset + ParsedStaticFunctionInput.InputVariable.GetSizeInBytes() <= Script.RapidIterationParameters.GetParameterDataArray().Num(), TEXT("Cached parameter offset not valid for parameter store.")))
 			{
 				FNiagaraVariable Temp = ParsedStaticFunctionInput.InputVariable;
-				Temp.SetData(Script.RapidIterationParameters.GetParameterData(ParsedStaticFunctionInput.DataOffset));
+				Temp.SetData(Script.RapidIterationParameters.GetParameterData(ParsedStaticFunctionInput.DataOffset, ParsedStaticFunctionInput.InputVariable.GetType()));
 				FSelectValue SelectValue = FTraversalBuilder::CreateSelectValue(Temp);
 				FunctionInputContext.Add(SelectKey, SelectValue);
 			}
diff --git a/Engine/Plugins/FX/Niagara/Source/NiagaraEditor/Private/Widgets/INiagaraDistributionAdapter.h b/Engine/Plugins/FX/Niagara/Source/NiagaraEditor/Private/Widgets/INiagaraDistributionAdapter.h
--- a/Engine/Plugins/FX/Niagara/Source/NiagaraEditor/Private/Widgets/INiagaraDistributionAdapter.h
+++ b/Engine/Plugins/FX/Niagara/Source/NiagaraEditor/Private/Widgets/INiagaraDistributionAdapter.h
@@ -16,6 +16,9 @@ enum class ENiagaraDistributionEditorMode
 {
 	/* Bound to an script attribute. */
 	Binding,
+	/** Bound to an expression chain */
+	Expression,
+	/** Constant value */
 	Constant,
 	/** A multi-channel value edited with a single constant value. */
 	UniformConstant,
@@ -99,6 +102,13 @@ public:
 	*/
 	virtual void SetConstantOrRangeValues(int32 ValueIndex, const TArray<float>& InValues) = 0;
 
+	/** Get the type used when generating an expression to run */
+	virtual FNiagaraTypeDefinition GetExpressionTypeDef() const = 0;
+	/** Get the root expression struct. */
+	virtual struct FInstancedStruct& GetExpressionRoot() const = 0;
+	/** Run the provided function inside a transaction, calling pre / post modify on the owning UObject. */
+	virtual void ExecuteTransaction(FText TransactionText, TFunction<void()> TransactionFunc) = 0;
+
 	/** Get the name of the variable we are currently bound to. */
 	virtual FNiagaraVariableBase GetBindingValue() const = 0;
 	/** Set the name of the variable we want to bind to. */
diff --git a/Engine/Plugins/FX/Niagara/Source/NiagaraEditor/Private/Widgets/NiagaraDistributionEditorUtilities.cpp b/Engine/Plugins/FX/Niagara/Source/NiagaraEditor/Private/Widgets/NiagaraDistributionEditorUtilities.cpp
--- a/Engine/Plugins/FX/Niagara/Source/NiagaraEditor/Private/Widgets/NiagaraDistributionEditorUtilities.cpp
+++ b/Engine/Plugins/FX/Niagara/Source/NiagaraEditor/Private/Widgets/NiagaraDistributionEditorUtilities.cpp
@@ -14,6 +14,8 @@ FText FNiagaraDistributionEditorUtilities::DistributionModeToDisplayName(ENiagar
 	{
 	case ENiagaraDistributionEditorMode::Binding:
 		return LOCTEXT("BindingDisplayName", "Binding");
+	case ENiagaraDistributionEditorMode::Expression:
+		return LOCTEXT("ExpressionDisplayName", "Expression");
 	case ENiagaraDistributionEditorMode::Constant:
 		return LOCTEXT("ConstantDisplayName", "Constant");
 	case ENiagaraDistributionEditorMode::UniformConstant:
@@ -52,6 +54,8 @@ FText FNiagaraDistributionEditorUtilities::DistributionModeToToolTipText(ENiagar
 	{
 	case ENiagaraDistributionEditorMode::Binding:
 		return LOCTEXT("BindingToolTip", "Value bound to an attribute.");
+	case ENiagaraDistributionEditorMode::Expression:
+		return LOCTEXT("ExpressionToolTip", "Value bound to a expression chain.");
 	case ENiagaraDistributionEditorMode::Constant:
 		return LOCTEXT("ConstantToolTip", "A constant single value.");
 	case ENiagaraDistributionEditorMode::UniformConstant:
@@ -90,6 +94,8 @@ FName FNiagaraDistributionEditorUtilities::DistributionModeToIconBrushName(ENiag
 	{
 	case ENiagaraDistributionEditorMode::Binding:
 		return "NiagaraEditor.DistributionEditor.Binding";
+	case ENiagaraDistributionEditorMode::Expression:
+		return "NiagaraEditor.DistributionEditor.Binding";
 	case ENiagaraDistributionEditorMode::Constant:
 	case ENiagaraDistributionEditorMode::UniformConstant:
 		return "NiagaraEditor.DistributionEditor.UniformConstant";
@@ -135,6 +141,12 @@ bool FNiagaraDistributionEditorUtilities::IsBinding(ENiagaraDistributionEditorMo
 		InMode == ENiagaraDistributionEditorMode::Binding;
 }
 
+bool FNiagaraDistributionEditorUtilities::IsExpression(ENiagaraDistributionEditorMode InMode)
+{
+	return
+		InMode == ENiagaraDistributionEditorMode::Expression;
+}
+
 bool FNiagaraDistributionEditorUtilities::IsUniform(ENiagaraDistributionEditorMode InMode)
 {
 	return
diff --git a/Engine/Plugins/FX/Niagara/Source/NiagaraEditor/Private/Widgets/NiagaraDistributionEditorUtilities.h b/Engine/Plugins/FX/Niagara/Source/NiagaraEditor/Private/Widgets/NiagaraDistributionEditorUtilities.h
--- a/Engine/Plugins/FX/Niagara/Source/NiagaraEditor/Private/Widgets/NiagaraDistributionEditorUtilities.h
+++ b/Engine/Plugins/FX/Niagara/Source/NiagaraEditor/Private/Widgets/NiagaraDistributionEditorUtilities.h
@@ -20,6 +20,7 @@ namespace FNiagaraDistributionEditorUtilities
 	FSlateIcon DistributionModeToIcon(ENiagaraDistributionEditorMode InMode);
 
 	bool IsBinding(ENiagaraDistributionEditorMode InMode);
+	bool IsExpression(ENiagaraDistributionEditorMode InMode);
 	bool IsUniform(ENiagaraDistributionEditorMode InMode);
 	bool IsColor(ENiagaraDistributionEditorMode InMode);
 
diff --git a/Engine/Plugins/FX/Niagara/Source/NiagaraEditor/Private/Widgets/SNiagaraDistributionEditor.cpp b/Engine/Plugins/FX/Niagara/Source/NiagaraEditor/Private/Widgets/SNiagaraDistributionEditor.cpp
--- a/Engine/Plugins/FX/Niagara/Source/NiagaraEditor/Private/Widgets/SNiagaraDistributionEditor.cpp
+++ b/Engine/Plugins/FX/Niagara/Source/NiagaraEditor/Private/Widgets/SNiagaraDistributionEditor.cpp
@@ -19,6 +19,8 @@
 #include "Widgets/SNiagaraDistributionCurveEditor.h"
 #include "Widgets/SNiagaraParameterName.h"
 
+#include "Stateless/SNiagaraStatelessExpressionWidget.h"
+
 #define LOCTEXT_NAMESPACE "NiagaraDistributionEditor"
 
 class SNiagaraDistributionModeSelector : public SComboButton
@@ -672,6 +674,10 @@ TSharedRef<SWidget> SNiagaraDistributionEditor::ConstructContentForMode()
 	{
 		return SNew(SNiagaraDistributionBindingEditor, DistributionAdapter.ToSharedRef());
 	}
+	else if (FNiagaraDistributionEditorUtilities::IsExpression(Mode))
+	{
+		return SNew(SNiagaraStatelessExpressionWidget, DistributionAdapter.ToSharedRef());
+	}
 	else if (FNiagaraDistributionEditorUtilities::IsConstant(Mode) || FNiagaraDistributionEditorUtilities::IsRange(Mode))
 	{
 		return SNew(SNiagaraDistributionValueEditor, DistributionAdapter.ToSharedRef());
