Dedupe Niagara simplex noise impl for curl noise into a single copy.

diff --git a/Engine/Plugins/FX/Niagara/Source/Niagara/Private/NiagaraDataInterfaceCurlNoise.cpp b/Engine/Plugins/FX/Niagara/Source/Niagara/Private/NiagaraDataInterfaceCurlNoise.cpp
--- a/Engine/Plugins/FX/Niagara/Source/Niagara/Private/NiagaraDataInterfaceCurlNoise.cpp
+++ b/Engine/Plugins/FX/Niagara/Source/Niagara/Private/NiagaraDataInterfaceCurlNoise.cpp
@@ -3,6 +3,7 @@
 #include "NiagaraDataInterfaceCurlNoise.h"
 #include "NiagaraCompileHashVisitor.h"
 #include "NiagaraShaderParametersBuilder.h"
+#include "NiagaraSimplexNoise.h"
 #include "Math/IntVector.h"
 #include "RenderingThread.h"
 
@@ -10,315 +11,11 @@
 
 static const FName SampleNoiseFieldName(TEXT("SampleNoiseField"));
 
-// ----------------------------------------------------- start noise helpers -----------------------------------------------------
-// Fairly straightforward CPU implementation of the equivalent HLSL code found in Engine/Shaders/Private/Random.ush
-// Also contains a partial implementation of some math types and functions found in HLSL but not in UE
-
-// TODO(mv): Move this to a separate NiagaraMath unit.  
-struct FNiagaraUIntVector
-{
-	uint32 X;
-	uint32 Y;
-	uint32 Z;
-
-	FNiagaraUIntVector() : X(0), Y(0), Z(0) {}
-	FNiagaraUIntVector(uint32 X, uint32 Y, uint32 Z) : X(X), Y(Y), Z(Z) {}
-	FNiagaraUIntVector(uint32 Value) : X(Value), Y(Value), Z(Value) {}
-	FNiagaraUIntVector(FIntVector Values) : X(Values.X), Y(Values.Y), Z(Values.Z) {}
-
-	FNiagaraUIntVector operator+(FNiagaraUIntVector Rhs)
-	{
-		return FNiagaraUIntVector{ X + Rhs.X, Y + Rhs.Y, Z + Rhs.Z };
-	}
-
-	FNiagaraUIntVector operator*(FNiagaraUIntVector Rhs)
-	{
-		return FNiagaraUIntVector{ X * Rhs.X, Y * Rhs.Y, Z * Rhs.Z };
-	}
-
-	FNiagaraUIntVector operator>>(uint32 Shift)
-	{
-		return FNiagaraUIntVector(X >> Shift, Y >> Shift, Z >> Shift);
-	}
-
-	FNiagaraUIntVector operator&(FNiagaraUIntVector Rhs)
-	{
-		return FNiagaraUIntVector{ X & Rhs.X, Y & Rhs.Y, Z & Rhs.Z };
-	}
-
-	uint32& operator[](int Index)
-	{
-		if (Index == 0) return X;
-		if (Index == 1) return Y;
-		if (Index == 2) return Z;
-		check(false); 
-		return X;
-	}
-
-	const uint32& operator[](int Index) const
-	{
-		if (Index == 0) return X;
-		if (Index == 1) return Y;
-		if (Index == 2) return Z;
-		check(false);
-		return X;
-	}
-};
-
-FNiagaraUIntVector Rand3DPCG16(FIntVector p)
-{
-	FNiagaraUIntVector v = FNiagaraUIntVector(p);
-	
-	v = v * 1664525u + 1013904223u;
-
-	v.X += v.Y*v.Z;
-	v.Y += v.Z*v.X;
-	v.Z += v.X*v.Y;
-	v.X += v.Y*v.Z;
-	v.Y += v.Z*v.X;
-	v.Z += v.X*v.Y;
-
-	return v >> 16u;
-}
-
-FVector3f NiagaraVectorFrac(FVector3f v)
-{
-	return FVector3f(FMath::Frac(v.X), FMath::Frac(v.Y), FMath::Frac(v.Z));
-}
-
-FVector3f NoiseTileWrap(FVector3f v, bool bTiling, float RepeatSize)
-{
-	FVector3f vv = bTiling ? (NiagaraVectorFrac(v / RepeatSize) * RepeatSize) : v;
-	return vv;
-}
-
-struct FNiagaraMatrix4x3
-{
-	FVector3f Row0;
-	FVector3f Row1;
-	FVector3f Row2;
-	FVector3f Row3;
-
-	FNiagaraMatrix4x3() : Row0(FVector3f::Zero()), Row1(FVector3f::Zero()), Row2(FVector3f::Zero()), Row3(FVector3f::Zero()) {}
-	FNiagaraMatrix4x3(FVector3f Row0, FVector3f Row1, FVector3f Row2, FVector3f Row3) : Row0(Row0), Row1(Row1), Row2(Row2), Row3(Row3) {}
-
-	FVector3f& operator[](int Row)
-	{
-		if (Row == 0) return Row0;
-		if (Row == 1) return Row1;
-		if (Row == 2) return Row2;
-		if (Row == 3) return Row3;
-		check(false);
-		return Row0;
-	}
-	const FVector3f& operator[](int Row) const
-	{
-		if (Row == 0) return Row0;
-		if (Row == 1) return Row1;
-		if (Row == 2) return Row2;
-		if (Row == 3) return Row3;
-		check(false);
-		return Row0;
-	}
-};
-
-FVector3f NiagaraVectorFloor(FVector3f v) {
-	return FVector3f(FGenericPlatformMath::FloorToFloat(v.X),
-		           FGenericPlatformMath::FloorToFloat(v.Y),
-		           FGenericPlatformMath::FloorToFloat(v.Z));
-};
-
-FIntVector NiagaraVectorFloorToInt(FVector3f v) {
-	return FIntVector(
-		FGenericPlatformMath::FloorToInt(v.X),
-		FGenericPlatformMath::FloorToInt(v.Y),
-		FGenericPlatformMath::FloorToInt(v.Z)
-	);
-};
-
-FVector3f NiagaraVectorStep(FVector3f v, FVector3f u)
-{
-	return FVector3f(u.X >= v.X ? 1.0f : 0.0f,
-		           u.Y >= v.Y ? 1.0f : 0.0f,
-		           u.Z >= v.Z ? 1.0f : 0.0f);
-}
-
-FVector3f NiagaraVectorSwizzle(FVector3f v, uint32 x, uint32 y, uint32 z)
-{
-	return FVector3f(v[x], v[y], v[z]);
-}
-
-FVector3f NiagaraVectorMin(FVector3f u, FVector3f v)
-{
-	return FVector3f(u.X < v.X ? u.X : v.X,
-		           u.Y < v.Y ? u.Y : v.Y, 
-		           u.Z < v.Z ? u.Z : v.Z);
-}
-
-FVector3f NiagaraVectorMax(FVector3f u, FVector3f v)
-{
-	return FVector3f(u.X > v.X ? u.X : v.X,
-		           u.Y > v.Y ? u.Y : v.Y,
-	           	   u.Z > v.Z ? u.Z : v.Z);
-} 
-
-FNiagaraMatrix4x3 SimplexCorners(FVector3f v)
-{
-	FVector3f tet = NiagaraVectorFloor(v + v.X / 3 + v.Y / 3 + v.Z / 3);
-	FVector3f base = tet - tet.X / 6 - tet.Y / 6 - tet.Z / 6;
-	FVector3f f = v - base;
-
-	FVector3f g = NiagaraVectorStep(NiagaraVectorSwizzle(f, 1, 2, 0), f), h = FVector3f(1.0f) - NiagaraVectorSwizzle(g, 2, 0, 1);
-	FVector3f a1 = NiagaraVectorMin(g, h) - 1. / 6., a2 = NiagaraVectorMax(g, h) - 1. / 3.;
-
-	return FNiagaraMatrix4x3(base, base + a1, base + a2, base + 0.5);
-}
-
-FVector4f NiagaraVector4Saturate(const FVector4f& v) {
-	return FVector4f(FMath::Clamp(v.X, 0.0f, 1.0f), FMath::Clamp(v.Y, 0.0f, 1.0f), FMath::Clamp(v.Z, 0.0f, 1.0f), FMath::Clamp(v.W, 0.0f, 1.0f));
-}
-
-FVector4f SimplexSmooth(FNiagaraMatrix4x3 f)
-{
-	const float scale = 1024.f / 375.f;
-	FVector4f d = FVector4f(FVector3f::DotProduct(f[0], f[0]), FVector3f::DotProduct(f[1], f[1]), FVector3f::DotProduct(f[2], f[2]), FVector3f::DotProduct(f[3], f[3]));
-	FVector4f s = NiagaraVector4Saturate(2.0f * d);
-	return scale * (FVector4f(1.0f, 1.0f, 1.0f, 1.0f) + s * (FVector4f(-3.0f, -3.0f, -3.0f, -3.0f) + s * (FVector4f(3.0f, 3.0f, 3.0f, 3.0f) - s)));
-}
-
-struct FNiagaraMatrix3x4
-{
-	FVector4f row0;
-	FVector4f row1;
-	FVector4f row2;
-
-	FNiagaraMatrix3x4() : row0(FVector4f::Zero()), row1(FVector4f::Zero()), row2(FVector4f::Zero()) {}
-	FNiagaraMatrix3x4(const FVector4f& row0, const FVector4f& row1, const FVector4f& row2) : row0(row0), row1(row1), row2(row2) {}
-
-	FVector4f& operator[](int row)
-	{
-		if (row == 0) return row0;
-		if (row == 1) return row1;
-		if (row == 2) return row2;
-		check(false);
-		return row0;
-	}
-	const FVector4f& operator[](int row) const
-	{
-		if (row == 0) return row0;
-		if (row == 1) return row1;
-		if (row == 2) return row2;
-		check(false); 
-		return row0; 
-	}
-};
-
-FNiagaraMatrix3x4 SimplexDSmooth(FNiagaraMatrix4x3 f)
-{
-	const float scale = 1024. / 375.;
-	FVector4f d = FVector4f(FVector3f::DotProduct(f[0], f[0]), FVector3f::DotProduct(f[1], f[1]), FVector3f::DotProduct(f[2], f[2]), FVector3f::DotProduct(f[3], f[3]));
-	FVector4f s = NiagaraVector4Saturate(2 * d);
-	s = -12 * FVector4f(scale, scale, scale, scale) + s * (24 * FVector4f(scale, scale, scale, scale) - s * 12 * scale);
-
-	return FNiagaraMatrix3x4(
-		s * FVector4f(f[0][0], f[1][0], f[2][0], f[3][0]),
-		s * FVector4f(f[0][1], f[1][1], f[2][1], f[3][1]),
-		s * FVector4f(f[0][2], f[1][2], f[2][2], f[3][2]));
-}
-
-FNiagaraUIntVector FNiagaraUIntVectorSwizzle(FNiagaraUIntVector v, int x, int y, int z)
-{
-	return FNiagaraUIntVector(v[x], v[y], v[z]);
-}
-
-FVector3f FNiagaraUIntVectorToFVector(FNiagaraUIntVector v)
-{
-	return FVector3f(float(v.X), float(v.Y), float(v.Z));
-}
-
-FVector3f MulFVector4AndFNiagaraMatrix4x3(const FVector4f& lhs, FNiagaraMatrix4x3 rhs)
-{
-	return FVector3f(lhs[0] * rhs[0][0] + lhs[1] * rhs[1][0] + lhs[2] * rhs[2][0] + lhs[3] * rhs[3][0],
-		           lhs[0] * rhs[0][1] + lhs[1] * rhs[1][1] + lhs[2] * rhs[2][1] + lhs[3] * rhs[3][1],
-	           	   lhs[0] * rhs[0][2] + lhs[1] * rhs[1][2] + lhs[2] * rhs[2][2] + lhs[3] * rhs[3][2]);
-}
-
-FVector3f MulFNiagaraMatrix3x4FAndVector4(const FNiagaraMatrix3x4& lhs, const FVector4f& rhs)
-{
-	return FVector3f(lhs[0][0] * rhs[0] + lhs[0][1] * rhs[1] + lhs[0][2] * rhs[2] + lhs[0][3] * rhs[3],
-		           lhs[1][0] * rhs[0] + lhs[1][1] * rhs[1] + lhs[1][2] * rhs[2] + lhs[1][3] * rhs[3],
-		           lhs[2][0] * rhs[0] + lhs[2][1] * rhs[1] + lhs[2][2] * rhs[2] + lhs[2][3] * rhs[3]);
-}
-
-#define MGradientMask FNiagaraUIntVector(0x8000, 0x4000, 0x2000)
-#define MGradientScale FVector3f(1.f / 0x4000, 1.f / 0x2000, 1.f / 0x1000)
-
-FNiagaraMatrix3x4 JacobianSimplex_ALU(FVector3f v)
-{
-	FNiagaraMatrix4x3 T = SimplexCorners(v);
-	FNiagaraUIntVector rand;
-	FNiagaraMatrix4x3 gvec[3], fv;
-	FNiagaraMatrix3x4 grad;
-
-	fv[0] = v - T[0];
-	rand = Rand3DPCG16(NiagaraVectorFloorToInt(6 * T[0] + 0.5));
-	gvec[0][0] = FVector3f(FNiagaraUIntVectorToFVector(FNiagaraUIntVectorSwizzle(rand, 0, 0, 0) & MGradientMask)) * MGradientScale - 1;
-	gvec[1][0] = FVector3f(FNiagaraUIntVectorToFVector(FNiagaraUIntVectorSwizzle(rand, 1, 1, 1) & MGradientMask)) * MGradientScale - 1;
-	gvec[2][0] = FVector3f(FNiagaraUIntVectorToFVector(FNiagaraUIntVectorSwizzle(rand, 2, 2, 2) & MGradientMask)) * MGradientScale - 1;
-	grad[0][0] = FVector3f::DotProduct(gvec[0][0], fv[0]);
-	grad[1][0] = FVector3f::DotProduct(gvec[1][0], fv[0]);
-	grad[2][0] = FVector3f::DotProduct(gvec[2][0], fv[0]);
-
-	fv[1] = v - T[1];
-	rand = Rand3DPCG16(NiagaraVectorFloorToInt(6 * T[1] + 0.5));
-	gvec[0][1] = FVector3f(FNiagaraUIntVectorToFVector(FNiagaraUIntVectorSwizzle(rand, 0, 0, 0) & MGradientMask)) * MGradientScale - 1;
-	gvec[1][1] = FVector3f(FNiagaraUIntVectorToFVector(FNiagaraUIntVectorSwizzle(rand, 1, 1, 1) & MGradientMask)) * MGradientScale - 1;
-	gvec[2][1] = FVector3f(FNiagaraUIntVectorToFVector(FNiagaraUIntVectorSwizzle(rand, 2, 2, 2) & MGradientMask)) * MGradientScale - 1;
-	grad[0][1] = FVector3f::DotProduct(gvec[0][1], fv[1]);
-	grad[1][1] = FVector3f::DotProduct(gvec[1][1], fv[1]);
-	grad[2][1] = FVector3f::DotProduct(gvec[2][1], fv[1]);
-
-	fv[2] = v - T[2];
-	rand = Rand3DPCG16(NiagaraVectorFloorToInt(6 * T[2] + 0.5));
-	gvec[0][2] = FVector3f(FNiagaraUIntVectorToFVector(FNiagaraUIntVectorSwizzle(rand, 0, 0, 0) & MGradientMask)) * MGradientScale - 1;
-	gvec[1][2] = FVector3f(FNiagaraUIntVectorToFVector(FNiagaraUIntVectorSwizzle(rand, 1, 1, 1) & MGradientMask)) * MGradientScale - 1;
-	gvec[2][2] = FVector3f(FNiagaraUIntVectorToFVector(FNiagaraUIntVectorSwizzle(rand, 2, 2, 2) & MGradientMask)) * MGradientScale - 1;
-	grad[0][2] = FVector3f::DotProduct(gvec[0][2], fv[2]);
-	grad[1][2] = FVector3f::DotProduct(gvec[1][2], fv[2]);
-	grad[2][2] = FVector3f::DotProduct(gvec[2][2], fv[2]);
-
-	fv[3] = v - T[3];
-	rand = Rand3DPCG16(NiagaraVectorFloorToInt(6 * T[3] + 0.5));
-	gvec[0][3] = FVector3f(FNiagaraUIntVectorToFVector(FNiagaraUIntVectorSwizzle(rand, 0, 0, 0) & MGradientMask)) * MGradientScale - 1;
-	gvec[1][3] = FVector3f(FNiagaraUIntVectorToFVector(FNiagaraUIntVectorSwizzle(rand, 1, 1, 1) & MGradientMask)) * MGradientScale - 1;
-	gvec[2][3] = FVector3f(FNiagaraUIntVectorToFVector(FNiagaraUIntVectorSwizzle(rand, 2, 2, 2) & MGradientMask)) * MGradientScale - 1;
-	grad[0][3] = FVector3f::DotProduct(gvec[0][3], fv[3]);
-	grad[1][3] = FVector3f::DotProduct(gvec[1][3], fv[3]);
-	grad[2][3] = FVector3f::DotProduct(gvec[2][3], fv[3]);
-
-	FVector4f sv = SimplexSmooth(fv);
-	FNiagaraMatrix3x4 ds = SimplexDSmooth(fv);
-
-	FNiagaraMatrix3x4 jacobian;
-	jacobian[0] = FVector4f(MulFVector4AndFNiagaraMatrix4x3(sv, gvec[0]) + MulFNiagaraMatrix3x4FAndVector4(ds, grad[0]), Dot4(sv, grad[0]));
-	jacobian[1] = FVector4f(MulFVector4AndFNiagaraMatrix4x3(sv, gvec[1]) + MulFNiagaraMatrix3x4FAndVector4(ds, grad[1]), Dot4(sv, grad[1]));
-	jacobian[2] = FVector4f(MulFVector4AndFNiagaraMatrix4x3(sv, gvec[2]) + MulFNiagaraMatrix3x4FAndVector4(ds, grad[2]), Dot4(sv, grad[2]));
-
-	return jacobian;
-}
-
-#undef MGradientMask
-#undef MGradientScale 
-
-// ----------------------------------------------------- end noise helpers -----------------------------------------------------
-
-
 UNiagaraDataInterfaceCurlNoise::UNiagaraDataInterfaceCurlNoise(FObjectInitializer const& ObjectInitializer)
 	: Super(ObjectInitializer)
 	, Seed(0)
 {
-	OffsetFromSeed = FNiagaraUIntVectorToFVector(Rand3DPCG16(FIntVector(Seed, Seed, Seed))) / 100.0f;
-
+	OffsetFromSeed = SimplexNoiseOffsetFromSeed(Seed);
 	Proxy.Reset(new FNiagaraDataInterfaceProxyCurlNoise(OffsetFromSeed));
 }
 
@@ -336,7 +33,7 @@ void UNiagaraDataInterfaceCurlNoise::PostInitProperties()
 void UNiagaraDataInterfaceCurlNoise::PostLoad()
 {
 	Super::PostLoad();
-	OffsetFromSeed = FNiagaraUIntVectorToFVector(Rand3DPCG16(FIntVector(Seed, Seed, Seed))) / 100.0f;
+	OffsetFromSeed = SimplexNoiseOffsetFromSeed(Seed);
 
 	MarkRenderDataDirty();
 }
@@ -359,7 +56,7 @@ void UNiagaraDataInterfaceCurlNoise::PostEditChangeProperty(struct FPropertyChan
 	if (PropertyChangedEvent.Property && PropertyChangedEvent.Property->GetFName() == GET_MEMBER_NAME_CHECKED(UNiagaraDataInterfaceCurlNoise, Seed))
 	{
 		// NOTE: Calculate the offset based on the seed on-change instead of on every invocation for every particle...
-		OffsetFromSeed = FNiagaraUIntVectorToFVector(Rand3DPCG16(FIntVector(Seed, Seed, Seed))) / 100.0f;
+		OffsetFromSeed = SimplexNoiseOffsetFromSeed(Seed);
 	}
 
 	MarkRenderDataDirty();
diff --git a/Engine/Plugins/FX/Niagara/Source/Niagara/Private/NiagaraSimplexNoise.cpp b/Engine/Plugins/FX/Niagara/Source/Niagara/Private/NiagaraSimplexNoise.cpp
new file mode 100644
--- /dev/null
+++ b/Engine/Plugins/FX/Niagara/Source/Niagara/Private/NiagaraSimplexNoise.cpp
@@ -0,0 +1,285 @@
+// Copyright Epic Games, Inc. All Rights Reserved.
+
+#include "NiagaraSimplexNoise.h"
+#include "Math/IntVector.h"
+
+namespace {
+
+// Fairly straightforward CPU implementation of the equivalent HLSL code found in Engine/Shaders/Private/Random.ush
+// Also contains a partial implementation of some math types and functions found in HLSL but not in UE
+
+struct FNiagaraUIntVector
+{
+	uint32 X;
+	uint32 Y;
+	uint32 Z;
+
+	FNiagaraUIntVector() : X(0), Y(0), Z(0) {}
+	FNiagaraUIntVector(uint32 X, uint32 Y, uint32 Z) : X(X), Y(Y), Z(Z) {}
+	FNiagaraUIntVector(uint32 Value) : X(Value), Y(Value), Z(Value) {}
+	FNiagaraUIntVector(FIntVector Values) : X(Values.X), Y(Values.Y), Z(Values.Z) {}
+
+	FNiagaraUIntVector operator+(FNiagaraUIntVector Rhs)
+	{
+		return FNiagaraUIntVector{ X + Rhs.X, Y + Rhs.Y, Z + Rhs.Z };
+	}
+
+	FNiagaraUIntVector operator*(FNiagaraUIntVector Rhs)
+	{
+		return FNiagaraUIntVector{ X * Rhs.X, Y * Rhs.Y, Z * Rhs.Z };
+	}
+
+	FNiagaraUIntVector operator>>(uint32 Shift)
+	{
+		return FNiagaraUIntVector(X >> Shift, Y >> Shift, Z >> Shift);
+	}
+
+	FNiagaraUIntVector operator&(FNiagaraUIntVector Rhs)
+	{
+		return FNiagaraUIntVector{ X & Rhs.X, Y & Rhs.Y, Z & Rhs.Z };
+	}
+
+	uint32& operator[](int Index)
+	{
+		if (Index == 0) return X;
+		if (Index == 1) return Y;
+		if (Index == 2) return Z;
+		check(false); 
+		return X;
+	}
+
+	const uint32& operator[](int Index) const
+	{
+		if (Index == 0) return X;
+		if (Index == 1) return Y;
+		if (Index == 2) return Z;
+		check(false);
+		return X;
+	}
+};
+
+FNiagaraUIntVector Rand3DPCG16(FIntVector p)
+{
+	FNiagaraUIntVector v = FNiagaraUIntVector(p);
+	
+	v = v * 1664525u + 1013904223u;
+
+	v.X += v.Y*v.Z;
+	v.Y += v.Z*v.X;
+	v.Z += v.X*v.Y;
+	v.X += v.Y*v.Z;
+	v.Y += v.Z*v.X;
+	v.Z += v.X*v.Y;
+
+	return v >> 16u;
+}
+
+FVector3f NiagaraVectorFrac(FVector3f v)
+{
+	return FVector3f(FMath::Frac(v.X), FMath::Frac(v.Y), FMath::Frac(v.Z));
+}
+
+FVector3f NoiseTileWrap(FVector3f v, bool bTiling, float RepeatSize)
+{
+	FVector3f vv = bTiling ? (NiagaraVectorFrac(v / RepeatSize) * RepeatSize) : v;
+	return vv;
+}
+
+struct FNiagaraMatrix4x3
+{
+	FVector3f Row0;
+	FVector3f Row1;
+	FVector3f Row2;
+	FVector3f Row3;
+
+	FNiagaraMatrix4x3() : Row0(FVector3f::Zero()), Row1(FVector3f::Zero()), Row2(FVector3f::Zero()), Row3(FVector3f::Zero()) {}
+	FNiagaraMatrix4x3(FVector3f Row0, FVector3f Row1, FVector3f Row2, FVector3f Row3) : Row0(Row0), Row1(Row1), Row2(Row2), Row3(Row3) {}
+
+	FVector3f& operator[](int Row)
+	{
+		if (Row == 0) return Row0;
+		if (Row == 1) return Row1;
+		if (Row == 2) return Row2;
+		if (Row == 3) return Row3;
+		check(false);
+		return Row0;
+	}
+	const FVector3f& operator[](int Row) const
+	{
+		if (Row == 0) return Row0;
+		if (Row == 1) return Row1;
+		if (Row == 2) return Row2;
+		if (Row == 3) return Row3;
+		check(false);
+		return Row0;
+	}
+};
+
+FVector3f NiagaraVectorFloor(FVector3f v) {
+	return FVector3f(FGenericPlatformMath::FloorToFloat(v.X),
+		           FGenericPlatformMath::FloorToFloat(v.Y),
+		           FGenericPlatformMath::FloorToFloat(v.Z));
+};
+
+FIntVector NiagaraVectorFloorToInt(FVector3f v) {
+	return FIntVector(
+		FGenericPlatformMath::FloorToInt(v.X),
+		FGenericPlatformMath::FloorToInt(v.Y),
+		FGenericPlatformMath::FloorToInt(v.Z)
+	);
+};
+
+FVector3f NiagaraVectorStep(FVector3f v, FVector3f u)
+{
+	return FVector3f(u.X >= v.X ? 1.0f : 0.0f,
+		           u.Y >= v.Y ? 1.0f : 0.0f,
+		           u.Z >= v.Z ? 1.0f : 0.0f);
+}
+
+FVector3f NiagaraVectorSwizzle(FVector3f v, uint32 x, uint32 y, uint32 z)
+{
+	return FVector3f(v[x], v[y], v[z]);
+}
+
+FVector3f NiagaraVectorMin(FVector3f u, FVector3f v)
+{
+	return FVector3f(u.X < v.X ? u.X : v.X,
+		           u.Y < v.Y ? u.Y : v.Y, 
+		           u.Z < v.Z ? u.Z : v.Z);
+}
+
+FVector3f NiagaraVectorMax(FVector3f u, FVector3f v)
+{
+	return FVector3f(u.X > v.X ? u.X : v.X,
+		           u.Y > v.Y ? u.Y : v.Y,
+	           	   u.Z > v.Z ? u.Z : v.Z);
+} 
+
+FNiagaraMatrix4x3 SimplexCorners(FVector3f v)
+{
+	FVector3f tet = NiagaraVectorFloor(v + v.X / 3 + v.Y / 3 + v.Z / 3);
+	FVector3f base = tet - tet.X / 6 - tet.Y / 6 - tet.Z / 6;
+	FVector3f f = v - base;
+
+	FVector3f g = NiagaraVectorStep(NiagaraVectorSwizzle(f, 1, 2, 0), f), h = FVector3f(1.0f) - NiagaraVectorSwizzle(g, 2, 0, 1);
+	FVector3f a1 = NiagaraVectorMin(g, h) - 1. / 6., a2 = NiagaraVectorMax(g, h) - 1. / 3.;
+
+	return FNiagaraMatrix4x3(base, base + a1, base + a2, base + 0.5);
+}
+
+FVector4f NiagaraVector4Saturate(const FVector4f& v) {
+	return FVector4f(FMath::Clamp(v.X, 0.0f, 1.0f), FMath::Clamp(v.Y, 0.0f, 1.0f), FMath::Clamp(v.Z, 0.0f, 1.0f), FMath::Clamp(v.W, 0.0f, 1.0f));
+}
+
+FVector4f SimplexSmooth(FNiagaraMatrix4x3 f)
+{
+	const float scale = 1024.f / 375.f;
+	FVector4f d = FVector4f(FVector3f::DotProduct(f[0], f[0]), FVector3f::DotProduct(f[1], f[1]), FVector3f::DotProduct(f[2], f[2]), FVector3f::DotProduct(f[3], f[3]));
+	FVector4f s = NiagaraVector4Saturate(2.0f * d);
+	return scale * (FVector4f(1.0f, 1.0f, 1.0f, 1.0f) + s * (FVector4f(-3.0f, -3.0f, -3.0f, -3.0f) + s * (FVector4f(3.0f, 3.0f, 3.0f, 3.0f) - s)));
+}
+
+FNiagaraMatrix3x4 SimplexDSmooth(FNiagaraMatrix4x3 f)
+{
+	const float scale = 1024. / 375.;
+	FVector4f d = FVector4f(FVector3f::DotProduct(f[0], f[0]), FVector3f::DotProduct(f[1], f[1]), FVector3f::DotProduct(f[2], f[2]), FVector3f::DotProduct(f[3], f[3]));
+	FVector4f s = NiagaraVector4Saturate(2 * d);
+	s = -12 * FVector4f(scale, scale, scale, scale) + s * (24 * FVector4f(scale, scale, scale, scale) - s * 12 * scale);
+
+	return FNiagaraMatrix3x4(
+		s * FVector4f(f[0][0], f[1][0], f[2][0], f[3][0]),
+		s * FVector4f(f[0][1], f[1][1], f[2][1], f[3][1]),
+		s * FVector4f(f[0][2], f[1][2], f[2][2], f[3][2]));
+}
+
+FNiagaraUIntVector FNiagaraUIntVectorSwizzle(FNiagaraUIntVector v, int x, int y, int z)
+{
+	return FNiagaraUIntVector(v[x], v[y], v[z]);
+}
+
+FVector3f FNiagaraUIntVectorToFVector(FNiagaraUIntVector v)
+{
+	return FVector3f(float(v.X), float(v.Y), float(v.Z));
+}
+
+FVector3f MulFVector4AndFNiagaraMatrix4x3(const FVector4f& lhs, FNiagaraMatrix4x3 rhs)
+{
+	return FVector3f(lhs[0] * rhs[0][0] + lhs[1] * rhs[1][0] + lhs[2] * rhs[2][0] + lhs[3] * rhs[3][0],
+		           lhs[0] * rhs[0][1] + lhs[1] * rhs[1][1] + lhs[2] * rhs[2][1] + lhs[3] * rhs[3][1],
+	           	   lhs[0] * rhs[0][2] + lhs[1] * rhs[1][2] + lhs[2] * rhs[2][2] + lhs[3] * rhs[3][2]);
+}
+
+FVector3f MulFNiagaraMatrix3x4FAndVector4(const FNiagaraMatrix3x4& lhs, const FVector4f& rhs)
+{
+	return FVector3f(lhs[0][0] * rhs[0] + lhs[0][1] * rhs[1] + lhs[0][2] * rhs[2] + lhs[0][3] * rhs[3],
+		           lhs[1][0] * rhs[0] + lhs[1][1] * rhs[1] + lhs[1][2] * rhs[2] + lhs[1][3] * rhs[3],
+		           lhs[2][0] * rhs[0] + lhs[2][1] * rhs[1] + lhs[2][2] * rhs[2] + lhs[2][3] * rhs[3]);
+}
+
+} // anonymous namespace
+
+FVector3f SimplexNoiseOffsetFromSeed(uint32 Seed)
+{
+	return FNiagaraUIntVectorToFVector(Rand3DPCG16(FIntVector(Seed, Seed, Seed))) / 100.0f;
+}
+
+#define MGradientMask FNiagaraUIntVector(0x8000, 0x4000, 0x2000)
+#define MGradientScale FVector3f(1.f / 0x4000, 1.f / 0x2000, 1.f / 0x1000)
+
+FNiagaraMatrix3x4 JacobianSimplex_ALU(FVector3f v)
+{
+	FNiagaraMatrix4x3 T = SimplexCorners(v);
+	FNiagaraUIntVector rand;
+	FNiagaraMatrix4x3 gvec[3], fv;
+	FNiagaraMatrix3x4 grad;
+
+	fv[0] = v - T[0];
+	rand = Rand3DPCG16(NiagaraVectorFloorToInt(6 * T[0] + 0.5));
+	gvec[0][0] = FVector3f(FNiagaraUIntVectorToFVector(FNiagaraUIntVectorSwizzle(rand, 0, 0, 0) & MGradientMask)) * MGradientScale - 1;
+	gvec[1][0] = FVector3f(FNiagaraUIntVectorToFVector(FNiagaraUIntVectorSwizzle(rand, 1, 1, 1) & MGradientMask)) * MGradientScale - 1;
+	gvec[2][0] = FVector3f(FNiagaraUIntVectorToFVector(FNiagaraUIntVectorSwizzle(rand, 2, 2, 2) & MGradientMask)) * MGradientScale - 1;
+	grad[0][0] = FVector3f::DotProduct(gvec[0][0], fv[0]);
+	grad[1][0] = FVector3f::DotProduct(gvec[1][0], fv[0]);
+	grad[2][0] = FVector3f::DotProduct(gvec[2][0], fv[0]);
+
+	fv[1] = v - T[1];
+	rand = Rand3DPCG16(NiagaraVectorFloorToInt(6 * T[1] + 0.5));
+	gvec[0][1] = FVector3f(FNiagaraUIntVectorToFVector(FNiagaraUIntVectorSwizzle(rand, 0, 0, 0) & MGradientMask)) * MGradientScale - 1;
+	gvec[1][1] = FVector3f(FNiagaraUIntVectorToFVector(FNiagaraUIntVectorSwizzle(rand, 1, 1, 1) & MGradientMask)) * MGradientScale - 1;
+	gvec[2][1] = FVector3f(FNiagaraUIntVectorToFVector(FNiagaraUIntVectorSwizzle(rand, 2, 2, 2) & MGradientMask)) * MGradientScale - 1;
+	grad[0][1] = FVector3f::DotProduct(gvec[0][1], fv[1]);
+	grad[1][1] = FVector3f::DotProduct(gvec[1][1], fv[1]);
+	grad[2][1] = FVector3f::DotProduct(gvec[2][1], fv[1]);
+
+	fv[2] = v - T[2];
+	rand = Rand3DPCG16(NiagaraVectorFloorToInt(6 * T[2] + 0.5));
+	gvec[0][2] = FVector3f(FNiagaraUIntVectorToFVector(FNiagaraUIntVectorSwizzle(rand, 0, 0, 0) & MGradientMask)) * MGradientScale - 1;
+	gvec[1][2] = FVector3f(FNiagaraUIntVectorToFVector(FNiagaraUIntVectorSwizzle(rand, 1, 1, 1) & MGradientMask)) * MGradientScale - 1;
+	gvec[2][2] = FVector3f(FNiagaraUIntVectorToFVector(FNiagaraUIntVectorSwizzle(rand, 2, 2, 2) & MGradientMask)) * MGradientScale - 1;
+	grad[0][2] = FVector3f::DotProduct(gvec[0][2], fv[2]);
+	grad[1][2] = FVector3f::DotProduct(gvec[1][2], fv[2]);
+	grad[2][2] = FVector3f::DotProduct(gvec[2][2], fv[2]);
+
+	fv[3] = v - T[3];
+	rand = Rand3DPCG16(NiagaraVectorFloorToInt(6 * T[3] + 0.5));
+	gvec[0][3] = FVector3f(FNiagaraUIntVectorToFVector(FNiagaraUIntVectorSwizzle(rand, 0, 0, 0) & MGradientMask)) * MGradientScale - 1;
+	gvec[1][3] = FVector3f(FNiagaraUIntVectorToFVector(FNiagaraUIntVectorSwizzle(rand, 1, 1, 1) & MGradientMask)) * MGradientScale - 1;
+	gvec[2][3] = FVector3f(FNiagaraUIntVectorToFVector(FNiagaraUIntVectorSwizzle(rand, 2, 2, 2) & MGradientMask)) * MGradientScale - 1;
+	grad[0][3] = FVector3f::DotProduct(gvec[0][3], fv[3]);
+	grad[1][3] = FVector3f::DotProduct(gvec[1][3], fv[3]);
+	grad[2][3] = FVector3f::DotProduct(gvec[2][3], fv[3]);
+
+	FVector4f sv = SimplexSmooth(fv);
+	FNiagaraMatrix3x4 ds = SimplexDSmooth(fv);
+
+	FNiagaraMatrix3x4 jacobian;
+	jacobian[0] = FVector4f(MulFVector4AndFNiagaraMatrix4x3(sv, gvec[0]) + MulFNiagaraMatrix3x4FAndVector4(ds, grad[0]), Dot4(sv, grad[0]));
+	jacobian[1] = FVector4f(MulFVector4AndFNiagaraMatrix4x3(sv, gvec[1]) + MulFNiagaraMatrix3x4FAndVector4(ds, grad[1]), Dot4(sv, grad[1]));
+	jacobian[2] = FVector4f(MulFVector4AndFNiagaraMatrix4x3(sv, gvec[2]) + MulFNiagaraMatrix3x4FAndVector4(ds, grad[2]), Dot4(sv, grad[2]));
+
+	return jacobian;
+}
+
+#undef MGradientMask
+#undef MGradientScale 
+
diff --git a/Engine/Plugins/FX/Niagara/Source/Niagara/Private/NiagaraSimplexNoise.h b/Engine/Plugins/FX/Niagara/Source/Niagara/Private/NiagaraSimplexNoise.h
new file mode 100644
--- /dev/null
+++ b/Engine/Plugins/FX/Niagara/Source/Niagara/Private/NiagaraSimplexNoise.h
@@ -0,0 +1,47 @@
+// Copyright Epic Games, Inc. All Rights Reserved.
+
+#pragma once
+
+#include "Math/Vector.h"
+#include "Math/Vector4.h"
+
+/**
+ * Jacobian as 3x4 matrix for JacobianSimplex_ALU
+ */
+struct FNiagaraMatrix3x4
+{
+	FVector4f row0;
+	FVector4f row1;
+	FVector4f row2;
+
+	FNiagaraMatrix3x4() : row0(FVector4f::Zero()), row1(FVector4f::Zero()), row2(FVector4f::Zero()) {}
+	FNiagaraMatrix3x4(const FVector4f& row0, const FVector4f& row1, const FVector4f& row2) : row0(row0), row1(row1), row2(row2) {}
+
+	FVector4f& operator[](int row)
+	{
+		if (row == 0) return row0;
+		if (row == 1) return row1;
+		if (row == 2) return row2;
+		check(false);
+		return row0;
+	}
+	const FVector4f& operator[](int row) const
+	{
+		if (row == 0) return row0;
+		if (row == 1) return row1;
+		if (row == 2) return row2;
+		check(false); 
+		return row0; 
+	}
+};
+
+/**
+ * Turns a random seed into a 3D vector offset to use when sampling the 3D simplex noise field.
+ */
+FVector3f SimplexNoiseOffsetFromSeed(uint32 Seed);
+
+/**
+ * Jacobian of simplex noise at given position. Used for curl noise.
+ */
+FNiagaraMatrix3x4 JacobianSimplex_ALU(FVector3f v);
+
diff --git a/Engine/Plugins/FX/Niagara/Source/Niagara/Private/Stateless/Modules/NiagaraStatelessModule_CurlNoiseForce.cpp b/Engine/Plugins/FX/Niagara/Source/Niagara/Private/Stateless/Modules/NiagaraStatelessModule_CurlNoiseForce.cpp
--- a/Engine/Plugins/FX/Niagara/Source/Niagara/Private/Stateless/Modules/NiagaraStatelessModule_CurlNoiseForce.cpp
+++ b/Engine/Plugins/FX/Niagara/Source/Niagara/Private/Stateless/Modules/NiagaraStatelessModule_CurlNoiseForce.cpp
@@ -1,304 +1,7 @@
 // Copyright Epic Games, Inc. All Rights Reserved.
 
 #include "Stateless/Modules/NiagaraStatelessModule_CurlNoiseForce.h"
-
-namespace NoiseHelper
-{
-	struct FNiagaraUIntVector
-	{
-		uint32 X;
-		uint32 Y;
-		uint32 Z;
-
-		FNiagaraUIntVector() : X(0), Y(0), Z(0) {}
-		FNiagaraUIntVector(uint32 X, uint32 Y, uint32 Z) : X(X), Y(Y), Z(Z) {}
-		FNiagaraUIntVector(uint32 Value) : X(Value), Y(Value), Z(Value) {}
-		FNiagaraUIntVector(FIntVector Values) : X(Values.X), Y(Values.Y), Z(Values.Z) {}
-
-		FNiagaraUIntVector operator+(FNiagaraUIntVector Rhs)
-		{
-			return FNiagaraUIntVector{ X + Rhs.X, Y + Rhs.Y, Z + Rhs.Z };
-		}
-
-		FNiagaraUIntVector operator*(FNiagaraUIntVector Rhs)
-		{
-			return FNiagaraUIntVector{ X * Rhs.X, Y * Rhs.Y, Z * Rhs.Z };
-		}
-
-		FNiagaraUIntVector operator>>(uint32 Shift)
-		{
-			return FNiagaraUIntVector(X >> Shift, Y >> Shift, Z >> Shift);
-		}
-
-		FNiagaraUIntVector operator&(FNiagaraUIntVector Rhs)
-		{
-			return FNiagaraUIntVector{ X & Rhs.X, Y & Rhs.Y, Z & Rhs.Z };
-		}
-
-		uint32& operator[](int Index)
-		{
-			if (Index == 0) return X;
-			if (Index == 1) return Y;
-			if (Index == 2) return Z;
-			check(false);
-			return X;
-		}
-
-		const uint32& operator[](int Index) const
-		{
-			if (Index == 0) return X;
-			if (Index == 1) return Y;
-			if (Index == 2) return Z;
-			check(false);
-			return X;
-		}
-	};
-
-	FNiagaraUIntVector Rand3DPCG16(FIntVector p)
-	{
-		FNiagaraUIntVector v = FNiagaraUIntVector(p);
-
-		v = v * 1664525u + 1013904223u;
-
-		v.X += v.Y * v.Z;
-		v.Y += v.Z * v.X;
-		v.Z += v.X * v.Y;
-		v.X += v.Y * v.Z;
-		v.Y += v.Z * v.X;
-		v.Z += v.X * v.Y;
-
-		return v >> 16u;
-	}
-
-	FVector3f NiagaraVectorFrac(FVector3f v)
-	{
-		return FVector3f(FMath::Frac(v.X), FMath::Frac(v.Y), FMath::Frac(v.Z));
-	}
-
-	FVector3f NoiseTileWrap(FVector3f v, bool bTiling, float RepeatSize)
-	{
-		FVector3f vv = bTiling ? (NiagaraVectorFrac(v / RepeatSize) * RepeatSize) : v;
-		return vv;
-	}
-
-	struct FNiagaraMatrix4x3
-	{
-		FVector3f Row0;
-		FVector3f Row1;
-		FVector3f Row2;
-		FVector3f Row3;
-
-		FNiagaraMatrix4x3() : Row0(FVector3f::Zero()), Row1(FVector3f::Zero()), Row2(FVector3f::Zero()), Row3(FVector3f::Zero()) {}
-		FNiagaraMatrix4x3(FVector3f Row0, FVector3f Row1, FVector3f Row2, FVector3f Row3) : Row0(Row0), Row1(Row1), Row2(Row2), Row3(Row3) {}
-
-		FVector3f& operator[](int Row)
-		{
-			if (Row == 0) return Row0;
-			if (Row == 1) return Row1;
-			if (Row == 2) return Row2;
-			if (Row == 3) return Row3;
-			check(false);
-			return Row0;
-		}
-		const FVector3f& operator[](int Row) const
-		{
-			if (Row == 0) return Row0;
-			if (Row == 1) return Row1;
-			if (Row == 2) return Row2;
-			if (Row == 3) return Row3;
-			check(false);
-			return Row0;
-		}
-	};
-
-	FVector3f NiagaraVectorFloor(FVector3f v) {
-		return FVector3f(FGenericPlatformMath::FloorToFloat(v.X),
-			FGenericPlatformMath::FloorToFloat(v.Y),
-			FGenericPlatformMath::FloorToFloat(v.Z));
-	};
-
-	FIntVector NiagaraVectorFloorToInt(FVector3f v) {
-		return FIntVector(
-			FGenericPlatformMath::FloorToInt(v.X),
-			FGenericPlatformMath::FloorToInt(v.Y),
-			FGenericPlatformMath::FloorToInt(v.Z)
-		);
-	};
-
-	FVector3f NiagaraVectorStep(FVector3f v, FVector3f u)
-	{
-		return FVector3f(u.X >= v.X ? 1.0f : 0.0f,
-			u.Y >= v.Y ? 1.0f : 0.0f,
-			u.Z >= v.Z ? 1.0f : 0.0f);
-	}
-
-	FVector3f NiagaraVectorSwizzle(FVector3f v, uint32 x, uint32 y, uint32 z)
-	{
-		return FVector3f(v[x], v[y], v[z]);
-	}
-
-	FVector3f NiagaraVectorMin(FVector3f u, FVector3f v)
-	{
-		return FVector3f(u.X < v.X ? u.X : v.X,
-			u.Y < v.Y ? u.Y : v.Y,
-			u.Z < v.Z ? u.Z : v.Z);
-	}
-
-	FVector3f NiagaraVectorMax(FVector3f u, FVector3f v)
-	{
-		return FVector3f(u.X > v.X ? u.X : v.X,
-			u.Y > v.Y ? u.Y : v.Y,
-			u.Z > v.Z ? u.Z : v.Z);
-	}
-
-	FNiagaraMatrix4x3 SimplexCorners(FVector3f v)
-	{
-		FVector3f tet = NiagaraVectorFloor(v + v.X / 3 + v.Y / 3 + v.Z / 3);
-		FVector3f base = tet - tet.X / 6 - tet.Y / 6 - tet.Z / 6;
-		FVector3f f = v - base;
-
-		FVector3f g = NiagaraVectorStep(NiagaraVectorSwizzle(f, 1, 2, 0), f), h = FVector3f(1.0f) - NiagaraVectorSwizzle(g, 2, 0, 1);
-		FVector3f a1 = NiagaraVectorMin(g, h) - 1. / 6., a2 = NiagaraVectorMax(g, h) - 1. / 3.;
-
-		return FNiagaraMatrix4x3(base, base + a1, base + a2, base + 0.5);
-	}
-
-	FVector4f NiagaraVector4Saturate(const FVector4f& v) {
-		return FVector4f(FMath::Clamp(v.X, 0.0f, 1.0f), FMath::Clamp(v.Y, 0.0f, 1.0f), FMath::Clamp(v.Z, 0.0f, 1.0f), FMath::Clamp(v.W, 0.0f, 1.0f));
-	}
-
-	FVector4f SimplexSmooth(FNiagaraMatrix4x3 f)
-	{
-		const float scale = 1024.f / 375.f;
-		FVector4f d = FVector4f(FVector3f::DotProduct(f[0], f[0]), FVector3f::DotProduct(f[1], f[1]), FVector3f::DotProduct(f[2], f[2]), FVector3f::DotProduct(f[3], f[3]));
-		FVector4f s = NiagaraVector4Saturate(2.0f * d);
-		return scale * (FVector4f(1.0f, 1.0f, 1.0f, 1.0f) + s * (FVector4f(-3.0f, -3.0f, -3.0f, -3.0f) + s * (FVector4f(3.0f, 3.0f, 3.0f, 3.0f) - s)));
-	}
-
-	struct FNiagaraMatrix3x4
-	{
-		FVector4f row0;
-		FVector4f row1;
-		FVector4f row2;
-
-		FNiagaraMatrix3x4() : row0(FVector4f::Zero()), row1(FVector4f::Zero()), row2(FVector4f::Zero()) {}
-		FNiagaraMatrix3x4(const FVector4f& row0, const FVector4f& row1, const FVector4f& row2) : row0(row0), row1(row1), row2(row2) {}
-
-		FVector4f& operator[](int row)
-		{
-			if (row == 0) return row0;
-			if (row == 1) return row1;
-			if (row == 2) return row2;
-			check(false);
-			return row0;
-		}
-		const FVector4f& operator[](int row) const
-		{
-			if (row == 0) return row0;
-			if (row == 1) return row1;
-			if (row == 2) return row2;
-			check(false);
-			return row0;
-		}
-	};
-
-	FNiagaraMatrix3x4 SimplexDSmooth(FNiagaraMatrix4x3 f)
-	{
-		const float scale = 1024. / 375.;
-		FVector4f d = FVector4f(FVector3f::DotProduct(f[0], f[0]), FVector3f::DotProduct(f[1], f[1]), FVector3f::DotProduct(f[2], f[2]), FVector3f::DotProduct(f[3], f[3]));
-		FVector4f s = NiagaraVector4Saturate(2 * d);
-		s = -12 * FVector4f(scale, scale, scale, scale) + s * (24 * FVector4f(scale, scale, scale, scale) - s * 12 * scale);
-
-		return FNiagaraMatrix3x4(
-			s * FVector4f(f[0][0], f[1][0], f[2][0], f[3][0]),
-			s * FVector4f(f[0][1], f[1][1], f[2][1], f[3][1]),
-			s * FVector4f(f[0][2], f[1][2], f[2][2], f[3][2]));
-	}
-
-	FNiagaraUIntVector FNiagaraUIntVectorSwizzle(FNiagaraUIntVector v, int x, int y, int z)
-	{
-		return FNiagaraUIntVector(v[x], v[y], v[z]);
-	}
-
-	FVector3f FNiagaraUIntVectorToFVector(FNiagaraUIntVector v)
-	{
-		return FVector3f(float(v.X), float(v.Y), float(v.Z));
-	}
-
-	FVector3f MulFVector4AndFNiagaraMatrix4x3(const FVector4f& lhs, FNiagaraMatrix4x3 rhs)
-	{
-		return FVector3f(lhs[0] * rhs[0][0] + lhs[1] * rhs[1][0] + lhs[2] * rhs[2][0] + lhs[3] * rhs[3][0],
-			lhs[0] * rhs[0][1] + lhs[1] * rhs[1][1] + lhs[2] * rhs[2][1] + lhs[3] * rhs[3][1],
-			lhs[0] * rhs[0][2] + lhs[1] * rhs[1][2] + lhs[2] * rhs[2][2] + lhs[3] * rhs[3][2]);
-	}
-
-	FVector3f MulFNiagaraMatrix3x4FAndVector4(const FNiagaraMatrix3x4& lhs, const FVector4f& rhs)
-	{
-		return FVector3f(lhs[0][0] * rhs[0] + lhs[0][1] * rhs[1] + lhs[0][2] * rhs[2] + lhs[0][3] * rhs[3],
-			lhs[1][0] * rhs[0] + lhs[1][1] * rhs[1] + lhs[1][2] * rhs[2] + lhs[1][3] * rhs[3],
-			lhs[2][0] * rhs[0] + lhs[2][1] * rhs[1] + lhs[2][2] * rhs[2] + lhs[2][3] * rhs[3]);
-	}
-
-#define MGradientMask FNiagaraUIntVector(0x8000, 0x4000, 0x2000)
-#define MGradientScale FVector3f(1.f / 0x4000, 1.f / 0x2000, 1.f / 0x1000)
-
-	FNiagaraMatrix3x4 JacobianSimplex_ALU(FVector3f v)
-	{
-		FNiagaraMatrix4x3 T = SimplexCorners(v);
-		FNiagaraUIntVector rand;
-		FNiagaraMatrix4x3 gvec[3], fv;
-		FNiagaraMatrix3x4 grad;
-
-		fv[0] = v - T[0];
-		rand = Rand3DPCG16(NiagaraVectorFloorToInt(6 * T[0] + 0.5));
-		gvec[0][0] = FVector3f(FNiagaraUIntVectorToFVector(FNiagaraUIntVectorSwizzle(rand, 0, 0, 0) & MGradientMask)) * MGradientScale - 1;
-		gvec[1][0] = FVector3f(FNiagaraUIntVectorToFVector(FNiagaraUIntVectorSwizzle(rand, 1, 1, 1) & MGradientMask)) * MGradientScale - 1;
-		gvec[2][0] = FVector3f(FNiagaraUIntVectorToFVector(FNiagaraUIntVectorSwizzle(rand, 2, 2, 2) & MGradientMask)) * MGradientScale - 1;
-		grad[0][0] = FVector3f::DotProduct(gvec[0][0], fv[0]);
-		grad[1][0] = FVector3f::DotProduct(gvec[1][0], fv[0]);
-		grad[2][0] = FVector3f::DotProduct(gvec[2][0], fv[0]);
-
-		fv[1] = v - T[1];
-		rand = Rand3DPCG16(NiagaraVectorFloorToInt(6 * T[1] + 0.5));
-		gvec[0][1] = FVector3f(FNiagaraUIntVectorToFVector(FNiagaraUIntVectorSwizzle(rand, 0, 0, 0) & MGradientMask)) * MGradientScale - 1;
-		gvec[1][1] = FVector3f(FNiagaraUIntVectorToFVector(FNiagaraUIntVectorSwizzle(rand, 1, 1, 1) & MGradientMask)) * MGradientScale - 1;
-		gvec[2][1] = FVector3f(FNiagaraUIntVectorToFVector(FNiagaraUIntVectorSwizzle(rand, 2, 2, 2) & MGradientMask)) * MGradientScale - 1;
-		grad[0][1] = FVector3f::DotProduct(gvec[0][1], fv[1]);
-		grad[1][1] = FVector3f::DotProduct(gvec[1][1], fv[1]);
-		grad[2][1] = FVector3f::DotProduct(gvec[2][1], fv[1]);
-
-		fv[2] = v - T[2];
-		rand = Rand3DPCG16(NiagaraVectorFloorToInt(6 * T[2] + 0.5));
-		gvec[0][2] = FVector3f(FNiagaraUIntVectorToFVector(FNiagaraUIntVectorSwizzle(rand, 0, 0, 0) & MGradientMask)) * MGradientScale - 1;
-		gvec[1][2] = FVector3f(FNiagaraUIntVectorToFVector(FNiagaraUIntVectorSwizzle(rand, 1, 1, 1) & MGradientMask)) * MGradientScale - 1;
-		gvec[2][2] = FVector3f(FNiagaraUIntVectorToFVector(FNiagaraUIntVectorSwizzle(rand, 2, 2, 2) & MGradientMask)) * MGradientScale - 1;
-		grad[0][2] = FVector3f::DotProduct(gvec[0][2], fv[2]);
-		grad[1][2] = FVector3f::DotProduct(gvec[1][2], fv[2]);
-		grad[2][2] = FVector3f::DotProduct(gvec[2][2], fv[2]);
-
-		fv[3] = v - T[3];
-		rand = Rand3DPCG16(NiagaraVectorFloorToInt(6 * T[3] + 0.5));
-		gvec[0][3] = FVector3f(FNiagaraUIntVectorToFVector(FNiagaraUIntVectorSwizzle(rand, 0, 0, 0) & MGradientMask)) * MGradientScale - 1;
-		gvec[1][3] = FVector3f(FNiagaraUIntVectorToFVector(FNiagaraUIntVectorSwizzle(rand, 1, 1, 1) & MGradientMask)) * MGradientScale - 1;
-		gvec[2][3] = FVector3f(FNiagaraUIntVectorToFVector(FNiagaraUIntVectorSwizzle(rand, 2, 2, 2) & MGradientMask)) * MGradientScale - 1;
-		grad[0][3] = FVector3f::DotProduct(gvec[0][3], fv[3]);
-		grad[1][3] = FVector3f::DotProduct(gvec[1][3], fv[3]);
-		grad[2][3] = FVector3f::DotProduct(gvec[2][3], fv[3]);
-
-		FVector4f sv = SimplexSmooth(fv);
-		FNiagaraMatrix3x4 ds = SimplexDSmooth(fv);
-
-		FNiagaraMatrix3x4 jacobian;
-		jacobian[0] = FVector4f(MulFVector4AndFNiagaraMatrix4x3(sv, gvec[0]) + MulFNiagaraMatrix3x4FAndVector4(ds, grad[0]), Dot4(sv, grad[0]));
-		jacobian[1] = FVector4f(MulFVector4AndFNiagaraMatrix4x3(sv, gvec[1]) + MulFNiagaraMatrix3x4FAndVector4(ds, grad[1]), Dot4(sv, grad[1]));
-		jacobian[2] = FVector4f(MulFVector4AndFNiagaraMatrix4x3(sv, gvec[2]) + MulFNiagaraMatrix3x4FAndVector4(ds, grad[2]), Dot4(sv, grad[2]));
-
-		return jacobian;
-	}
-
-#undef MGradientMask
-#undef MGradientScale 
-}
+#include "NiagaraSimplexNoise.h"
 
 void UNiagaraStatelessModule_CurlNoiseForce::BuildEmitterData(const FNiagaraStatelessEmitterDataBuildContext& BuildContext) const
 {
@@ -332,7 +35,7 @@ void UNiagaraStatelessModule_CurlNoiseForce::BuildEmitterData(const FNiagaraStat
 			{
 				//NoiseValue += FMath::RandRange(-1.0f, 1.0f);
 
-				NoiseHelper::FNiagaraMatrix3x4 J = NoiseHelper::JacobianSimplex_ALU(NoisePosition);
+				FNiagaraMatrix3x4 J = JacobianSimplex_ALU(NoisePosition);
 				NoiseValue += J[1][2] - J[2][1];
 				NoisePosition += NoiseTravel;
 
