- Fix PVS Studio nit pick about final function inside a final class

#rnx

diff --git a/Engine/Plugins/FX/Niagara/Source/Niagara/Internal/Stateless/Expressions/NiagaraStatelessExpressionColor.h b/Engine/Plugins/FX/Niagara/Source/Niagara/Internal/Stateless/Expressions/NiagaraStatelessExpressionColor.h
--- a/Engine/Plugins/FX/Niagara/Source/Niagara/Internal/Stateless/Expressions/NiagaraStatelessExpressionColor.h
+++ b/Engine/Plugins/FX/Niagara/Source/Niagara/Internal/Stateless/Expressions/NiagaraStatelessExpressionColor.h
@@ -29,7 +29,7 @@ struct FNiagaraStatelessExpressionColorConstant final : public FNiagaraStateless
 	GENERATED_BODY()
 
 	virtual ~FNiagaraStatelessExpressionColorConstant() = default;
-	virtual FLinearColor EvaluateInternal(const FEvaluateContext& Context) const override final;
+	virtual FLinearColor EvaluateInternal(const FEvaluateContext& Context) const override;
 	virtual bool IsConstant() const override;
 
 	UPROPERTY(EditAnywhere, Category = "Values")
@@ -46,7 +46,7 @@ struct FNiagaraStatelessExpressionColorBinding final : public FNiagaraStatelessE
 	virtual ~FNiagaraStatelessExpressionColorBinding() = default;
 	virtual bool IsConstant() const { return false; }
 	virtual FInstancedStruct BuildInternal(const FNiagaraStatelessEmitterDataBuildContext& BuildContext) const override;
-	virtual FLinearColor EvaluateInternal(const FEvaluateContext& Context) const override final;
+	virtual FLinearColor EvaluateInternal(const FEvaluateContext& Context) const override;
 
 	UPROPERTY(EditAnywhere, Category = "Values")
 	FName A;
@@ -64,7 +64,7 @@ struct FNiagaraStatelessExpressionColorAdd final : public FNiagaraStatelessExpre
 	FNiagaraStatelessExpressionColorAdd();
 	virtual ~FNiagaraStatelessExpressionColorAdd() = default;
 	virtual FInstancedStruct BuildInternal(const FNiagaraStatelessEmitterDataBuildContext& BuildContext) const override;
-	virtual FLinearColor EvaluateInternal(const FEvaluateContext& Context) const override final;
+	virtual FLinearColor EvaluateInternal(const FEvaluateContext& Context) const override;
 	virtual bool IsConstant() const override;
 
 	UPROPERTY(EditAnywhere, Category = "Values", meta = (BaseStruct = "/Script/Niagara.NiagaraStatelessExpressionColor"))
@@ -84,7 +84,7 @@ struct FNiagaraStatelessExpressionColorSubtract final : public FNiagaraStateless
 	FNiagaraStatelessExpressionColorSubtract();
 	virtual ~FNiagaraStatelessExpressionColorSubtract() = default;
 	virtual FInstancedStruct BuildInternal(const FNiagaraStatelessEmitterDataBuildContext& BuildContext) const override;
-	virtual FLinearColor EvaluateInternal(const FEvaluateContext& Context) const override final;
+	virtual FLinearColor EvaluateInternal(const FEvaluateContext& Context) const override;
 	virtual bool IsConstant() const override;
 
 	UPROPERTY(EditAnywhere, Category = "Values", meta = (BaseStruct = "/Script/Niagara.NiagaraStatelessExpressionColor"))
@@ -104,7 +104,7 @@ struct FNiagaraStatelessExpressionColorMultiply final : public FNiagaraStateless
 	FNiagaraStatelessExpressionColorMultiply();
 	virtual ~FNiagaraStatelessExpressionColorMultiply() = default;
 	virtual FInstancedStruct BuildInternal(const FNiagaraStatelessEmitterDataBuildContext& BuildContext) const override;
-	virtual FLinearColor EvaluateInternal(const FEvaluateContext& Context) const override final;
+	virtual FLinearColor EvaluateInternal(const FEvaluateContext& Context) const override;
 	virtual bool IsConstant() const override;
 
 	UPROPERTY(EditAnywhere, Category = "Values", meta = (BaseStruct = "/Script/Niagara.NiagaraStatelessExpressionColor"))
@@ -124,7 +124,7 @@ struct FNiagaraStatelessExpressionColorDivide final : public FNiagaraStatelessEx
 	FNiagaraStatelessExpressionColorDivide();
 	virtual ~FNiagaraStatelessExpressionColorDivide() = default;
 	virtual FInstancedStruct BuildInternal(const FNiagaraStatelessEmitterDataBuildContext& BuildContext) const override;
-	virtual FLinearColor EvaluateInternal(const FEvaluateContext& Context) const override final;
+	virtual FLinearColor EvaluateInternal(const FEvaluateContext& Context) const override;
 	virtual bool IsConstant() const override;
 
 	UPROPERTY(EditAnywhere, Category = "Values", meta = (BaseStruct = "/Script/Niagara.NiagaraStatelessExpressionColor"))
diff --git a/Engine/Plugins/FX/Niagara/Source/Niagara/Internal/Stateless/Expressions/NiagaraStatelessExpressionFloat.h b/Engine/Plugins/FX/Niagara/Source/Niagara/Internal/Stateless/Expressions/NiagaraStatelessExpressionFloat.h
--- a/Engine/Plugins/FX/Niagara/Source/Niagara/Internal/Stateless/Expressions/NiagaraStatelessExpressionFloat.h
+++ b/Engine/Plugins/FX/Niagara/Source/Niagara/Internal/Stateless/Expressions/NiagaraStatelessExpressionFloat.h
@@ -29,7 +29,7 @@ struct FNiagaraStatelessExpressionFloatConstant final : public FNiagaraStateless
 	GENERATED_BODY()
 
 	virtual ~FNiagaraStatelessExpressionFloatConstant() = default;
-	virtual float EvaluateInternal(const FEvaluateContext& Context) const override final;
+	virtual float EvaluateInternal(const FEvaluateContext& Context) const override;
 	virtual bool IsConstant() const override;
 
 	UPROPERTY(EditAnywhere, Category = "Values")
@@ -46,7 +46,7 @@ struct FNiagaraStatelessExpressionFloatBinding final : public FNiagaraStatelessE
 	virtual ~FNiagaraStatelessExpressionFloatBinding() = default;
 	virtual bool IsConstant() const { return false; }
 	virtual FInstancedStruct BuildInternal(const FNiagaraStatelessEmitterDataBuildContext& BuildContext) const override;
-	virtual float EvaluateInternal(const FEvaluateContext& Context) const override final;
+	virtual float EvaluateInternal(const FEvaluateContext& Context) const override;
 
 	UPROPERTY(EditAnywhere, Category = "Values")
 	FName A;
@@ -64,7 +64,7 @@ struct FNiagaraStatelessExpressionFloatAdd final : public FNiagaraStatelessExpre
 	FNiagaraStatelessExpressionFloatAdd();
 	virtual ~FNiagaraStatelessExpressionFloatAdd() = default;
 	virtual FInstancedStruct BuildInternal(const FNiagaraStatelessEmitterDataBuildContext& BuildContext) const override;
-	virtual float EvaluateInternal(const FEvaluateContext& Context) const override final;
+	virtual float EvaluateInternal(const FEvaluateContext& Context) const override;
 	virtual bool IsConstant() const override;
 
 	UPROPERTY(EditAnywhere, Category = "Values", meta = (BaseStruct = "/Script/Niagara.NiagaraStatelessExpressionFloat"))
@@ -84,7 +84,7 @@ struct FNiagaraStatelessExpressionFloatSubtract final : public FNiagaraStateless
 	FNiagaraStatelessExpressionFloatSubtract();
 	virtual ~FNiagaraStatelessExpressionFloatSubtract() = default;
 	virtual FInstancedStruct BuildInternal(const FNiagaraStatelessEmitterDataBuildContext& BuildContext) const override;
-	virtual float EvaluateInternal(const FEvaluateContext& Context) const override final;
+	virtual float EvaluateInternal(const FEvaluateContext& Context) const override;
 	virtual bool IsConstant() const override;
 
 	UPROPERTY(EditAnywhere, Category = "Values", meta = (BaseStruct = "/Script/Niagara.NiagaraStatelessExpressionFloat"))
@@ -104,7 +104,7 @@ struct FNiagaraStatelessExpressionFloatMultiply final : public FNiagaraStateless
 	FNiagaraStatelessExpressionFloatMultiply();
 	virtual ~FNiagaraStatelessExpressionFloatMultiply() = default;
 	virtual FInstancedStruct BuildInternal(const FNiagaraStatelessEmitterDataBuildContext& BuildContext) const override;
-	virtual float EvaluateInternal(const FEvaluateContext& Context) const override final;
+	virtual float EvaluateInternal(const FEvaluateContext& Context) const override;
 	virtual bool IsConstant() const override;
 
 	UPROPERTY(EditAnywhere, Category = "Values", meta = (BaseStruct = "/Script/Niagara.NiagaraStatelessExpressionFloat"))
@@ -124,7 +124,7 @@ struct FNiagaraStatelessExpressionFloatDivide final : public FNiagaraStatelessEx
 	FNiagaraStatelessExpressionFloatDivide();
 	virtual ~FNiagaraStatelessExpressionFloatDivide() = default;
 	virtual FInstancedStruct BuildInternal(const FNiagaraStatelessEmitterDataBuildContext& BuildContext) const override;
-	virtual float EvaluateInternal(const FEvaluateContext& Context) const override final;
+	virtual float EvaluateInternal(const FEvaluateContext& Context) const override;
 	virtual bool IsConstant() const override;
 
 	UPROPERTY(EditAnywhere, Category = "Values", meta = (BaseStruct = "/Script/Niagara.NiagaraStatelessExpressionFloat"))
diff --git a/Engine/Plugins/FX/Niagara/Source/Niagara/Internal/Stateless/Expressions/NiagaraStatelessExpressionVec2.h b/Engine/Plugins/FX/Niagara/Source/Niagara/Internal/Stateless/Expressions/NiagaraStatelessExpressionVec2.h
--- a/Engine/Plugins/FX/Niagara/Source/Niagara/Internal/Stateless/Expressions/NiagaraStatelessExpressionVec2.h
+++ b/Engine/Plugins/FX/Niagara/Source/Niagara/Internal/Stateless/Expressions/NiagaraStatelessExpressionVec2.h
@@ -29,7 +29,7 @@ struct FNiagaraStatelessExpressionVec2Constant final : public FNiagaraStatelessE
 	GENERATED_BODY()
 
 	virtual ~FNiagaraStatelessExpressionVec2Constant() = default;
-	virtual FVector2f EvaluateInternal(const FEvaluateContext& Context) const override final;
+	virtual FVector2f EvaluateInternal(const FEvaluateContext& Context) const override;
 	virtual bool IsConstant() const override;
 
 	UPROPERTY(EditAnywhere, Category = "Values")
@@ -46,7 +46,7 @@ struct FNiagaraStatelessExpressionVec2Binding final : public FNiagaraStatelessEx
 	virtual ~FNiagaraStatelessExpressionVec2Binding() = default;
 	virtual bool IsConstant() const { return false; }
 	virtual FInstancedStruct BuildInternal(const FNiagaraStatelessEmitterDataBuildContext& BuildContext) const override;
-	virtual FVector2f EvaluateInternal(const FEvaluateContext& Context) const override final;
+	virtual FVector2f EvaluateInternal(const FEvaluateContext& Context) const override;
 
 	UPROPERTY(EditAnywhere, Category = "Values")
 	FName A;
@@ -64,7 +64,7 @@ struct FNiagaraStatelessExpressionVec2Add final : public FNiagaraStatelessExpres
 	FNiagaraStatelessExpressionVec2Add();
 	virtual ~FNiagaraStatelessExpressionVec2Add() = default;
 	virtual FInstancedStruct BuildInternal(const FNiagaraStatelessEmitterDataBuildContext& BuildContext) const override;
-	virtual FVector2f EvaluateInternal(const FEvaluateContext& Context) const override final;
+	virtual FVector2f EvaluateInternal(const FEvaluateContext& Context) const override;
 	virtual bool IsConstant() const override;
 
 	UPROPERTY(EditAnywhere, Category = "Values", meta = (BaseStruct = "/Script/Niagara.NiagaraStatelessExpressionVec2"))
@@ -84,7 +84,7 @@ struct FNiagaraStatelessExpressionVec2Subtract final : public FNiagaraStatelessE
 	FNiagaraStatelessExpressionVec2Subtract();
 	virtual ~FNiagaraStatelessExpressionVec2Subtract() = default;
 	virtual FInstancedStruct BuildInternal(const FNiagaraStatelessEmitterDataBuildContext& BuildContext) const override;
-	virtual FVector2f EvaluateInternal(const FEvaluateContext& Context) const override final;
+	virtual FVector2f EvaluateInternal(const FEvaluateContext& Context) const override;
 	virtual bool IsConstant() const override;
 
 	UPROPERTY(EditAnywhere, Category = "Values", meta = (BaseStruct = "/Script/Niagara.NiagaraStatelessExpressionVec2"))
@@ -104,7 +104,7 @@ struct FNiagaraStatelessExpressionVec2Multiply final : public FNiagaraStatelessE
 	FNiagaraStatelessExpressionVec2Multiply();
 	virtual ~FNiagaraStatelessExpressionVec2Multiply() = default;
 	virtual FInstancedStruct BuildInternal(const FNiagaraStatelessEmitterDataBuildContext& BuildContext) const override;
-	virtual FVector2f EvaluateInternal(const FEvaluateContext& Context) const override final;
+	virtual FVector2f EvaluateInternal(const FEvaluateContext& Context) const override;
 	virtual bool IsConstant() const override;
 
 	UPROPERTY(EditAnywhere, Category = "Values", meta = (BaseStruct = "/Script/Niagara.NiagaraStatelessExpressionVec2"))
@@ -124,7 +124,7 @@ struct FNiagaraStatelessExpressionVec2Divide final : public FNiagaraStatelessExp
 	FNiagaraStatelessExpressionVec2Divide();
 	virtual ~FNiagaraStatelessExpressionVec2Divide() = default;
 	virtual FInstancedStruct BuildInternal(const FNiagaraStatelessEmitterDataBuildContext& BuildContext) const override;
-	virtual FVector2f EvaluateInternal(const FEvaluateContext& Context) const override final;
+	virtual FVector2f EvaluateInternal(const FEvaluateContext& Context) const override;
 	virtual bool IsConstant() const override;
 
 	UPROPERTY(EditAnywhere, Category = "Values", meta = (BaseStruct = "/Script/Niagara.NiagaraStatelessExpressionVec2"))
diff --git a/Engine/Plugins/FX/Niagara/Source/Niagara/Internal/Stateless/Expressions/NiagaraStatelessExpressionVec3.h b/Engine/Plugins/FX/Niagara/Source/Niagara/Internal/Stateless/Expressions/NiagaraStatelessExpressionVec3.h
--- a/Engine/Plugins/FX/Niagara/Source/Niagara/Internal/Stateless/Expressions/NiagaraStatelessExpressionVec3.h
+++ b/Engine/Plugins/FX/Niagara/Source/Niagara/Internal/Stateless/Expressions/NiagaraStatelessExpressionVec3.h
@@ -29,7 +29,7 @@ struct FNiagaraStatelessExpressionVec3Constant final : public FNiagaraStatelessE
 	GENERATED_BODY()
 
 	virtual ~FNiagaraStatelessExpressionVec3Constant() = default;
-	virtual FVector3f EvaluateInternal(const FEvaluateContext& Context) const override final;
+	virtual FVector3f EvaluateInternal(const FEvaluateContext& Context) const override;
 	virtual bool IsConstant() const override;
 
 	UPROPERTY(EditAnywhere, Category = "Values")
@@ -46,7 +46,7 @@ struct FNiagaraStatelessExpressionVec3Binding final : public FNiagaraStatelessEx
 	virtual ~FNiagaraStatelessExpressionVec3Binding() = default;
 	virtual bool IsConstant() const { return false; }
 	virtual FInstancedStruct BuildInternal(const FNiagaraStatelessEmitterDataBuildContext& BuildContext) const override;
-	virtual FVector3f EvaluateInternal(const FEvaluateContext& Context) const override final;
+	virtual FVector3f EvaluateInternal(const FEvaluateContext& Context) const override;
 
 	UPROPERTY(EditAnywhere, Category = "Values")
 	FName A;
@@ -64,7 +64,7 @@ struct FNiagaraStatelessExpressionVec3Add final : public FNiagaraStatelessExpres
 	FNiagaraStatelessExpressionVec3Add();
 	virtual ~FNiagaraStatelessExpressionVec3Add() = default;
 	virtual FInstancedStruct BuildInternal(const FNiagaraStatelessEmitterDataBuildContext& BuildContext) const override;
-	virtual FVector3f EvaluateInternal(const FEvaluateContext& Context) const override final;
+	virtual FVector3f EvaluateInternal(const FEvaluateContext& Context) const override;
 	virtual bool IsConstant() const override;
 
 	UPROPERTY(EditAnywhere, Category = "Values", meta = (BaseStruct = "/Script/Niagara.NiagaraStatelessExpressionVec3"))
@@ -84,7 +84,7 @@ struct FNiagaraStatelessExpressionVec3Subtract final : public FNiagaraStatelessE
 	FNiagaraStatelessExpressionVec3Subtract();
 	virtual ~FNiagaraStatelessExpressionVec3Subtract() = default;
 	virtual FInstancedStruct BuildInternal(const FNiagaraStatelessEmitterDataBuildContext& BuildContext) const override;
-	virtual FVector3f EvaluateInternal(const FEvaluateContext& Context) const override final;
+	virtual FVector3f EvaluateInternal(const FEvaluateContext& Context) const override;
 	virtual bool IsConstant() const override;
 
 	UPROPERTY(EditAnywhere, Category = "Values", meta = (BaseStruct = "/Script/Niagara.NiagaraStatelessExpressionVec3"))
@@ -104,7 +104,7 @@ struct FNiagaraStatelessExpressionVec3Multiply final : public FNiagaraStatelessE
 	FNiagaraStatelessExpressionVec3Multiply();
 	virtual ~FNiagaraStatelessExpressionVec3Multiply() = default;
 	virtual FInstancedStruct BuildInternal(const FNiagaraStatelessEmitterDataBuildContext& BuildContext) const override;
-	virtual FVector3f EvaluateInternal(const FEvaluateContext& Context) const override final;
+	virtual FVector3f EvaluateInternal(const FEvaluateContext& Context) const override;
 	virtual bool IsConstant() const override;
 
 	UPROPERTY(EditAnywhere, Category = "Values", meta = (BaseStruct = "/Script/Niagara.NiagaraStatelessExpressionVec3"))
@@ -124,7 +124,7 @@ struct FNiagaraStatelessExpressionVec3Divide final : public FNiagaraStatelessExp
 	FNiagaraStatelessExpressionVec3Divide();
 	virtual ~FNiagaraStatelessExpressionVec3Divide() = default;
 	virtual FInstancedStruct BuildInternal(const FNiagaraStatelessEmitterDataBuildContext& BuildContext) const override;
-	virtual FVector3f EvaluateInternal(const FEvaluateContext& Context) const override final;
+	virtual FVector3f EvaluateInternal(const FEvaluateContext& Context) const override;
 	virtual bool IsConstant() const override;
 
 	UPROPERTY(EditAnywhere, Category = "Values", meta = (BaseStruct = "/Script/Niagara.NiagaraStatelessExpressionVec3"))
diff --git a/Engine/Plugins/FX/Niagara/Source/Niagara/Internal/Stateless/Expressions/NiagaraStatelessExpressionVec4.h b/Engine/Plugins/FX/Niagara/Source/Niagara/Internal/Stateless/Expressions/NiagaraStatelessExpressionVec4.h
--- a/Engine/Plugins/FX/Niagara/Source/Niagara/Internal/Stateless/Expressions/NiagaraStatelessExpressionVec4.h
+++ b/Engine/Plugins/FX/Niagara/Source/Niagara/Internal/Stateless/Expressions/NiagaraStatelessExpressionVec4.h
@@ -29,7 +29,7 @@ struct FNiagaraStatelessExpressionVec4Constant final : public FNiagaraStatelessE
 	GENERATED_BODY()
 
 	virtual ~FNiagaraStatelessExpressionVec4Constant() = default;
-	virtual FVector4f EvaluateInternal(const FEvaluateContext& Context) const override final;
+	virtual FVector4f EvaluateInternal(const FEvaluateContext& Context) const override;
 	virtual bool IsConstant() const override;
 
 	UPROPERTY(EditAnywhere, Category = "Values")
@@ -46,7 +46,7 @@ struct FNiagaraStatelessExpressionVec4Binding final : public FNiagaraStatelessEx
 	virtual ~FNiagaraStatelessExpressionVec4Binding() = default;
 	virtual bool IsConstant() const { return false; }
 	virtual FInstancedStruct BuildInternal(const FNiagaraStatelessEmitterDataBuildContext& BuildContext) const override;
-	virtual FVector4f EvaluateInternal(const FEvaluateContext& Context) const override final;
+	virtual FVector4f EvaluateInternal(const FEvaluateContext& Context) const override;
 
 	UPROPERTY(EditAnywhere, Category = "Values")
 	FName A;
@@ -64,7 +64,7 @@ struct FNiagaraStatelessExpressionVec4Add final : public FNiagaraStatelessExpres
 	FNiagaraStatelessExpressionVec4Add();
 	virtual ~FNiagaraStatelessExpressionVec4Add() = default;
 	virtual FInstancedStruct BuildInternal(const FNiagaraStatelessEmitterDataBuildContext& BuildContext) const override;
-	virtual FVector4f EvaluateInternal(const FEvaluateContext& Context) const override final;
+	virtual FVector4f EvaluateInternal(const FEvaluateContext& Context) const override;
 	virtual bool IsConstant() const override;
 
 	UPROPERTY(EditAnywhere, Category = "Values", meta = (BaseStruct = "/Script/Niagara.NiagaraStatelessExpressionVec4"))
@@ -84,7 +84,7 @@ struct FNiagaraStatelessExpressionVec4Subtract final : public FNiagaraStatelessE
 	FNiagaraStatelessExpressionVec4Subtract();
 	virtual ~FNiagaraStatelessExpressionVec4Subtract() = default;
 	virtual FInstancedStruct BuildInternal(const FNiagaraStatelessEmitterDataBuildContext& BuildContext) const override;
-	virtual FVector4f EvaluateInternal(const FEvaluateContext& Context) const override final;
+	virtual FVector4f EvaluateInternal(const FEvaluateContext& Context) const override;
 	virtual bool IsConstant() const override;
 
 	UPROPERTY(EditAnywhere, Category = "Values", meta = (BaseStruct = "/Script/Niagara.NiagaraStatelessExpressionVec4"))
@@ -104,7 +104,7 @@ struct FNiagaraStatelessExpressionVec4Multiply final : public FNiagaraStatelessE
 	FNiagaraStatelessExpressionVec4Multiply();
 	virtual ~FNiagaraStatelessExpressionVec4Multiply() = default;
 	virtual FInstancedStruct BuildInternal(const FNiagaraStatelessEmitterDataBuildContext& BuildContext) const override;
-	virtual FVector4f EvaluateInternal(const FEvaluateContext& Context) const override final;
+	virtual FVector4f EvaluateInternal(const FEvaluateContext& Context) const override;
 	virtual bool IsConstant() const override;
 
 	UPROPERTY(EditAnywhere, Category = "Values", meta = (BaseStruct = "/Script/Niagara.NiagaraStatelessExpressionVec4"))
@@ -124,7 +124,7 @@ struct FNiagaraStatelessExpressionVec4Divide final : public FNiagaraStatelessExp
 	FNiagaraStatelessExpressionVec4Divide();
 	virtual ~FNiagaraStatelessExpressionVec4Divide() = default;
 	virtual FInstancedStruct BuildInternal(const FNiagaraStatelessEmitterDataBuildContext& BuildContext) const override;
-	virtual FVector4f EvaluateInternal(const FEvaluateContext& Context) const override final;
+	virtual FVector4f EvaluateInternal(const FEvaluateContext& Context) const override;
 	virtual bool IsConstant() const override;
 
 	UPROPERTY(EditAnywhere, Category = "Values", meta = (BaseStruct = "/Script/Niagara.NiagaraStatelessExpressionVec4"))
