- IWYU fix

diff --git a/Engine/Plugins/FX/Niagara/Source/Niagara/Internal/Stateless/Modules/NiagaraStatelessModule_ScaleRibbonWidth.h b/Engine/Plugins/FX/Niagara/Source/Niagara/Internal/Stateless/Modules/NiagaraStatelessModule_ScaleRibbonWidth.h
--- a/Engine/Plugins/FX/Niagara/Source/Niagara/Internal/Stateless/Modules/NiagaraStatelessModule_ScaleRibbonWidth.h
+++ b/Engine/Plugins/FX/Niagara/Source/Niagara/Internal/Stateless/Modules/NiagaraStatelessModule_ScaleRibbonWidth.h
@@ -7,6 +7,7 @@
 #include "Stateless/NiagaraStatelessEmitterDataBuildContext.h"
 #include "Stateless/NiagaraStatelessModuleShaderParameters.h"
 #include "Stateless/NiagaraStatelessParticleSimContext.h"
+#include "NiagaraParameterBinding.h"
 
 #include "NiagaraStatelessModule_ScaleRibbonWidth.generated.h"
 
