- Niagara World Manager Tick changes
- Add tick passes to the world manager
- Two passes, one for none overlapping and one for overlapping assets
- None overlapping pass is only enabled if we have an simulation that requires overlapping
- None overlapping pass will run first to ensure potentially blocking simulations run first
- Ticks are dependent so we can guarantee ordering
- Add a different task priority for none overlapping simulations
- Instances now release the system simulation when they are complete
- Testing showed we ended up with buckets of ~40 simulations vs ~4 in the same spot
- Upside is we release memory faster, reduce the regular WM tick cost (less sims to iterate)
- Downside is we may have to spin up simulations more often

diff --git a/Engine/Plugins/FX/Niagara/Source/Niagara/Private/NiagaraSimulationTaskPriority.cpp b/Engine/Plugins/FX/Niagara/Source/Niagara/Private/NiagaraSimulationTaskPriority.cpp
--- a/Engine/Plugins/FX/Niagara/Source/Niagara/Private/NiagaraSimulationTaskPriority.cpp
+++ b/Engine/Plugins/FX/Niagara/Source/Niagara/Private/NiagaraSimulationTaskPriority.cpp
@@ -26,6 +26,16 @@ namespace NiagaraSimulationTaskPriority
 		ECVF_Default | ECVF_Scalability
 	);
 
+	// Priority for work which is not allowed to overlap tick groups
+	static int32 GNoneOverlappingPriority = 1;
+	static FAutoConsoleVariableRef CVarNoneOverlappingPriority(
+		TEXT("fx.Niagara.TaskPriorities.NoneOverlappingPriority"),
+		GNoneOverlappingPriority,
+		TEXT("Priority to use when we are not allowed to overlap tick groups.")
+		TEXT("Default is 1, set to -1 to not override the default priorities."),
+		ECVF_Default
+	);
+
 	// Translation from a Niagara Priority into a Task Priority
 	static FAutoConsoleTaskPriority GTaskPriorities[] =
 	{
@@ -97,7 +107,7 @@ namespace NiagaraSimulationTaskPriority
 				{
 					DumpPriority(
 						GTickGroupPriorities[i],
-						GetTickGroupPriority(ETickingGroup(i)),
+						GetTickGroupPriority(ETickingGroup(i), true),
 						*TickingGroupEnum->GetNameStringByValue(i)
 					);
 				}
@@ -167,7 +177,7 @@ namespace NiagaraSimulationTaskPriority
 		return GTaskPriorities[Priority].Get();
 	}
 
-	ENamedThreads::Type GetTickGroupPriority(ETickingGroup TickGroup)
+	ENamedThreads::Type GetTickGroupPriority(ETickingGroup TickGroup, bool bCanOverlapTickGroups)
 	{
 		int32 Priority = 0;
 	#if WITH_PARTICLE_PERF_STATS
@@ -183,6 +193,7 @@ namespace NiagaraSimulationTaskPriority
 			const int32 TickGroupIndex = FMath::Clamp(int32(Priority), 0, (int32)UE_ARRAY_COUNT(GTickGroupPriorities) - 1);
 			Priority = GTickGroupPriorities[TickGroupIndex];
 		}
+		Priority = bCanOverlapTickGroups || GNoneOverlappingPriority < 0 ? Priority : GNoneOverlappingPriority;
 		Priority = FMath::Clamp(Priority, 0, (int32)UE_ARRAY_COUNT(GTaskPriorities) - 1);
 		return GTaskPriorities[Priority].Get();
 	}
diff --git a/Engine/Plugins/FX/Niagara/Source/Niagara/Private/NiagaraSimulationTaskPriority.h b/Engine/Plugins/FX/Niagara/Source/Niagara/Private/NiagaraSimulationTaskPriority.h
--- a/Engine/Plugins/FX/Niagara/Source/Niagara/Private/NiagaraSimulationTaskPriority.h
+++ b/Engine/Plugins/FX/Niagara/Source/Niagara/Private/NiagaraSimulationTaskPriority.h
@@ -9,5 +9,5 @@ namespace NiagaraSimulationTaskPriority
 {
 	extern void Initialize();
 	extern ENamedThreads::Type GetPostActorTickPriority();
-	extern ENamedThreads::Type GetTickGroupPriority(ETickingGroup TickGroup);
+	extern ENamedThreads::Type GetTickGroupPriority(ETickingGroup TickGroup, bool bCanOverlapTickGroups);
 }
diff --git a/Engine/Plugins/FX/Niagara/Source/Niagara/Private/NiagaraSystemInstance.cpp b/Engine/Plugins/FX/Niagara/Source/Niagara/Private/NiagaraSystemInstance.cpp
--- a/Engine/Plugins/FX/Niagara/Source/Niagara/Private/NiagaraSystemInstance.cpp
+++ b/Engine/Plugins/FX/Niagara/Source/Niagara/Private/NiagaraSystemInstance.cpp
@@ -707,6 +707,13 @@ void FNiagaraSystemInstance::Complete(bool bExternalCompletion)
 		UnbindParameters(true);
 	}
 
+	// When tick passes are enabled we need to quick release the simulation otherwise they can linger causing more tick groups to be active
+	FNiagaraWorldManager* WorldManager = GetWorldManager();
+	if (WorldManager && WorldManager->UseTickPasses())
+	{
+		SystemSimulation = nullptr;
+	}
+
 	if (bNeedToNotifyOthers)
 	{
 		// We've already notified once, no need to do so again.
diff --git a/Engine/Plugins/FX/Niagara/Source/Niagara/Private/NiagaraSystemSimulation.cpp b/Engine/Plugins/FX/Niagara/Source/Niagara/Private/NiagaraSystemSimulation.cpp
--- a/Engine/Plugins/FX/Niagara/Source/Niagara/Private/NiagaraSystemSimulation.cpp
+++ b/Engine/Plugins/FX/Niagara/Source/Niagara/Private/NiagaraSystemSimulation.cpp
@@ -293,7 +293,7 @@ struct FNiagaraSystemSimulationTickConcurrentTask
 	FNiagaraSystemSimulationTickConcurrentTask(FNiagaraSystemSimulationTickContext InContext, FGraphEventRef& CompletionGraphEvent)
 		: Context(InContext)
 	{
-		TaskThread = NiagaraSimulationTaskPriority::GetTickGroupPriority(Context.Owner->GetTickGroup());
+		TaskThread = NiagaraSimulationTaskPriority::GetTickGroupPriority(Context.Owner->GetTickGroup(), Context.System->AsyncWorkCanOverlapTickGroups());
 		CompletionTask = TGraphTask<FNiagaraSystemSimulationAllWorkCompleteTask>::CreateTask(nullptr, ENamedThreads::GameThread).ConstructAndHold(Context.CompletionEvents);
 		CompletionGraphEvent = CompletionTask->GetCompletionEvent();
 	}
@@ -363,7 +363,7 @@ struct FNiagaraSystemInstanceTickConcurrentTask
 		, Batch(InBatch)
 		, WorldContext(InWorldContext)
 	{
-		TaskThread = NiagaraSimulationTaskPriority::GetTickGroupPriority(SystemSimulation->GetTickGroup());
+		TaskThread = NiagaraSimulationTaskPriority::GetTickGroupPriority(SystemSimulation->GetTickGroup(), SystemSimulation->GetSystem()->AsyncWorkCanOverlapTickGroups());
 	}
 
 	FORCEINLINE TStatId GetStatId() const { RETURN_QUICK_DECLARE_CYCLE_STAT(FNiagaraSystemInstanceTickConcurrentTask, STATGROUP_TaskGraphTasks); }
diff --git a/Engine/Plugins/FX/Niagara/Source/Niagara/Private/NiagaraWorldManager.cpp b/Engine/Plugins/FX/Niagara/Source/Niagara/Private/NiagaraWorldManager.cpp
--- a/Engine/Plugins/FX/Niagara/Source/Niagara/Private/NiagaraWorldManager.cpp
+++ b/Engine/Plugins/FX/Niagara/Source/Niagara/Private/NiagaraWorldManager.cpp
@@ -43,11 +43,11 @@ DECLARE_CYCLE_STAT(TEXT("Niagara Manager Wait Pre Garbage Collect [GT]"), STAT_N
 DECLARE_CYCLE_STAT(TEXT("Niagara Manager Refresh Owner Allows Scalability"), STAT_NiagaraWorldManRefreshOwnerAllowsScalability, STATGROUP_Niagara);
 DECLARE_CYCLE_STAT(TEXT("Niagara Manager Tick Parameter Collections [GT]"), STAT_NiagaraWorldManTickParamCollections, STATGROUP_Niagara);
 
-bool GNiagaraWorldManagerKillUniqueSims = 1;
-static FAutoConsoleVariableRef CVarNiagaraWorldManagerKillUniqueSims(
-	TEXT("fx.Niagara.WorldManager.KillUniqueSims"),
-	GNiagaraWorldManagerKillUniqueSims,
-	TEXT("System simulations will be removed when unique rather that waiting for GC."),
+bool GNiagaraWorldManagerUseTickPasses = true;
+static FAutoConsoleVariableRef CVarNiagaraWorldManagerUseTickPasses(
+	TEXT("fx.Niagara.WorldManager.UseTickPasses"),
+	GNiagaraWorldManagerUseTickPasses,
+	TEXT("When enabled we will use the tick passes to avoid blocking overlapping sims at tick boundaries."),
 	ECVF_Default
 );
 
@@ -59,14 +59,6 @@ static FAutoConsoleVariableRef CVarNiagaraWorldManagerObjectPoolEnabled(
 	ECVF_Default
 );
 
-static int GNiagaraAllowAsyncWorkToEndOfFrame = 1;
-static FAutoConsoleVariableRef CVarNiagaraAllowAsyncWorkToEndOfFrame(
-	TEXT("fx.Niagara.AllowAsyncWorkToEndOfFrame"),
-	GNiagaraAllowAsyncWorkToEndOfFrame,
-	TEXT("Allow async work to continue until the end of the frame, if false it will complete within the tick group it's started in."),
-	ECVF_Default
-); 
-
 static int GNiagaraWaitOnPreGC = 1;
 static FAutoConsoleVariableRef CVarNiagaraWaitOnPreGC(
 	TEXT("fx.Niagara.WaitOnPreGC"),
@@ -99,16 +91,18 @@ static FAutoConsoleVariableRef CVarAllowVisibilityCullingForDynamicBounds(
 	ECVF_Default
 );
 
-FAutoConsoleCommandWithWorld DumpNiagaraWorldManagerCommand(
+FAutoConsoleCommandWithWorldAndArgs DumpNiagaraWorldManagerCommand(
 	TEXT("DumpNiagaraWorldManager"),
 	TEXT("Dump Information About the Niagara World Manager Contents"),
-	FConsoleCommandWithWorldDelegate::CreateLambda(
-		[](UWorld* World)
+	FConsoleCommandWithWorldAndArgsDelegate::CreateLambda(
+		[](const TArray<FString>& Args, UWorld* World)
 		{
 			FNiagaraWorldManager* WorldManager = FNiagaraWorldManager::Get(World);
 			if (WorldManager != nullptr && GLog != nullptr)
 			{
-				WorldManager->DumpDetails(*GLog);
+				const bool bDetailed = Args.Num() > 0 ? FCString::Atoi(*Args[0]) != 0 : false;
+
+				WorldManager->DumpDetails(*GLog, bDetailed);
 			}
 		}
 	)
@@ -333,7 +327,7 @@ namespace NiagaraWorldManagerInternal
 void FNiagaraWorldManagerTickFunction::ExecuteTick(float DeltaTime, enum ELevelTick TickType, ENamedThreads::Type CurrentThread, const FGraphEventRef& MyCompletionGraphEvent)
 {
 	check(Owner);
-	Owner->Tick(TickGroup, DeltaTime, TickType, CurrentThread, MyCompletionGraphEvent);
+	Owner->Tick(TickPass, TickGroup, DeltaTime, MyCompletionGraphEvent);
 }
 
 FString FNiagaraWorldManagerTickFunction::DiagnosticMessage()
@@ -368,10 +362,10 @@ void FNiagaraWorldManager::SetScalabilityCullingMode(ENiagaraScalabilityCullingM
 	}
 }
 
-
 FNiagaraWorldManager::FNiagaraWorldManager()
 	: World(nullptr)
 	, ActiveNiagaraTickGroup(-1)
+	, bUseTickPasses(GNiagaraWorldManagerUseTickPasses)
 	, bAppHasFocus(true)
 	, DataChannelManager(new FNiagaraDataChannelManager(this))
 {
@@ -392,21 +386,32 @@ void FNiagaraWorldManager::Init(UWorld* InWorld)
 	DebugPlaybackRate = 1.0f;
 
 	World = InWorld;
+
+	ComponentPool = NewObject<UNiagaraComponentPool>();
+
 	for (int32 TickGroup = 0; TickGroup < NiagaraNumTickGroups; ++TickGroup)
 	{
-		FNiagaraWorldManagerTickFunction& TickFunc = TickFunctions[TickGroup];
-		TickFunc.TickGroup = ETickingGroup(NiagaraFirstTickGroup + TickGroup);
-		TickFunc.EndTickGroup = GNiagaraAllowAsyncWorkToEndOfFrame ? TG_LastDemotable : (ETickingGroup)TickFunc.TickGroup;
-		TickFunc.bCanEverTick = true;
-		TickFunc.bStartWithTickEnabled = true;
-		TickFunc.bAllowTickOnDedicatedServer = false;
-		TickFunc.bHighPriority = NiagaraWorldManagerInternal::GetNiagaraTickGroupPriority(TickGroup);
-		TickFunc.Owner = this;
-		TickFunc.RegisterTickFunction(InWorld->PersistentLevel);
+		for (int32 iTickPass=0; iTickPass < int32(ENiagaraWorldManagerTickPass::Num); ++iTickPass )
+		{
+			const ENiagaraWorldManagerTickPass TickPass = ENiagaraWorldManagerTickPass(iTickPass);
+
+			FNiagaraWorldManagerTickFunction& TickFunc = TickFunctions[TickGroup][iTickPass];
+			TickFunc.TickGroup						= ETickingGroup(NiagaraFirstTickGroup + TickGroup);
+			TickFunc.EndTickGroup					= TickPass == ENiagaraWorldManagerTickPass::Overlapping ? TG_LastDemotable : (ETickingGroup)TickFunc.TickGroup;
+			TickFunc.bCanEverTick					= true;
+			TickFunc.bStartWithTickEnabled			= TickPass == ENiagaraWorldManagerTickPass::Overlapping;
+			TickFunc.bAllowTickOnDedicatedServer	= false;
+			TickFunc.bHighPriority					= NiagaraWorldManagerInternal::GetNiagaraTickGroupPriority(TickGroup);
+			TickFunc.Owner							= this;
+			TickFunc.TickPass						= TickPass;
+			TickFunc.RegisterTickFunction(InWorld->PersistentLevel);
+			if (iTickPass > 0)
+			{
+				TickFunc.AddPrerequisite(ComponentPool, TickFunctions[TickGroup][iTickPass - 1]);
+			}
+		}
 	}
 
-	ComponentPool = NewObject<UNiagaraComponentPool>();
-
 	//Ideally we'd do this here but it's too early in the init process and the world does not have a Scene yet.
 	//Possibly a later hook we can use.
 	//PrimePoolForAllSystems();
@@ -652,30 +657,44 @@ FNiagaraSystemSimulationRef FNiagaraWorldManager::GetSystemSimulation(ETickingGr
 		ActualTickGroup = DemotedTickGroup == ActualTickGroup ? 0 : DemotedTickGroup;
 	}
 
-	FNiagaraSystemSimulationRef* SimPtr = SystemSimulations[ActualTickGroup].Find(System);
+	const bool bCanOverlapTickGroups = System->AsyncWorkCanOverlapTickGroups();
+	const ENiagaraWorldManagerTickPass TickPass = bCanOverlapTickGroups || bUseTickPasses == false ? ENiagaraWorldManagerTickPass::Overlapping : ENiagaraWorldManagerTickPass::NoneOverlapping;
+	FNiagaraSystemSimulationRef* SimPtr = SystemSimulations[ActualTickGroup][int(TickPass)].Find(System);
 	if (SimPtr != nullptr)
 	{
 		return *SimPtr;
 	}
 	
 	FNiagaraSystemSimulationRef Sim = MakeShared<FNiagaraSystemSimulation, ESPMode::ThreadSafe>();
-	SystemSimulations[ActualTickGroup].Add(System, Sim);
+	SystemSimulations[ActualTickGroup][int(TickPass)].Add(System, Sim);
 	Sim->Init(System, World, false, TickGroup);
 
-	//If this system doesn't allow async work to overlap TGs then we must ensure our Tick task is performed in it's correct TG.
-	//The async work for this system can then properly be chained to the tick task so that everything completes on time.
-	//Most systems most of the time will no need this and we'll clear this requirement as soon as all systems no longer need it.
-	//TODO: This should go away when we move over to the new task graph system and rejig our dependency tracking.
-	if(!System->AsyncWorkCanOverlapTickGroups())
+	// Old path which will block simulation at tick boundaries
+	if (bUseTickPasses == false)
+	{
+		//If this system doesn't allow async work to overlap TGs then we must ensure our Tick task is performed in it's correct TG.
+		//The async work for this system can then properly be chained to the tick task so that everything completes on time.
+		//Most systems most of the time will no need this and we'll clear this requirement as soon as all systems no longer need it.
+		//TODO: This should go away when we move over to the new task graph system and rejig our dependency tracking.
+		if(!bCanOverlapTickGroups)
+		{
+			TickFunctions[ActualTickGroup][int(TickPass)].EndTickGroup = (ETickingGroup)ActualTickGroup;
+		}
+	}
+	// New path which will only block none overlapping simulations are tick boundaries
+	else
 	{
-		TickFunctions[ActualTickGroup].EndTickGroup = (ETickingGroup)ActualTickGroup;
+		// Ensure tick function is running
+		if (TickFunctions[ActualTickGroup][int(TickPass)].IsTickFunctionEnabled() == false )
+		{
+			TickFunctions[ActualTickGroup][int(TickPass)].SetTickFunctionEnable(true);
+		}
 	}
 
 #if WITH_EDITOR
 	if (!SystemPostChangeDelegateHandles.Contains(System))
 	{
-		SystemPostChangeDelegateHandles.Emplace(System,
-			System->OnSystemPostEditChange().AddRaw(this, &FNiagaraWorldManager::OnSystemPostChange));
+		SystemPostChangeDelegateHandles.Emplace(System, System->OnSystemPostEditChange().AddRaw(this, &FNiagaraWorldManager::OnSystemPostChange));
 	}
 #endif
 
@@ -686,14 +705,17 @@ void FNiagaraWorldManager::DestroySystemSimulation(UNiagaraSystem* System)
 {
 	for ( int TG=0; TG < NiagaraNumTickGroups; ++TG )
 	{
-		FNiagaraSystemSimulationRef* Simulation = SystemSimulations[TG].Find(System);
-		if (Simulation != nullptr)
+		for (int iTickPass=0; iTickPass < int(ENiagaraWorldManagerTickPass::Num); ++iTickPass)
 		{
-			SimulationsWithPostActorWork.Remove(*Simulation);
-			SimulationsWithEndOfFrameWait.Remove(*Simulation);
+			FNiagaraSystemSimulationRef* Simulation = SystemSimulations[TG][iTickPass].Find(System);
+			if (Simulation != nullptr)
+			{
+				SimulationsWithPostActorWork.Remove(*Simulation);
+				SimulationsWithEndOfFrameWait.Remove(*Simulation);
 
-			(*Simulation)->Destroy();
-			SystemSimulations[TG].Remove(System);
+				(*Simulation)->Destroy();
+				SystemSimulations[TG][iTickPass].Remove(System);
+			}
 		}
 	}
 	ComponentPool->RemoveComponentsBySystem(System);
@@ -754,11 +776,14 @@ void FNiagaraWorldManager::OnWorldCleanup(bool bSessionEnded, bool bCleanupResou
 
 	for (int TG = 0; TG < NiagaraNumTickGroups; ++TG)
 	{
-		for (TPair<UNiagaraSystem*, FNiagaraSystemSimulationRef>& SimPair : SystemSimulations[TG])
+		for (int iTickPass = 0; iTickPass < int(ENiagaraWorldManagerTickPass::Num); ++iTickPass)
 		{
-			SimPair.Value->Destroy();
+			for (TPair<UNiagaraSystem*, FNiagaraSystemSimulationRef>& SimPair : SystemSimulations[TG][iTickPass])
+			{
+				SimPair.Value->Destroy();
+			}
+			SystemSimulations[TG][iTickPass].Empty();
 		}
-		SystemSimulations[TG].Empty();
 	}
 	CleanupParameterCollections();
 
@@ -794,15 +819,10 @@ void FNiagaraWorldManager::PreGarbageCollect()
 	if (GNiagaraWaitOnPreGC)
 	{
 		SCOPE_CYCLE_COUNTER(STAT_NiagaraWorldManWaitPreGC);
+
 		// We must wait for system simulation & instance async ticks to complete before garbage collection can start
 		// The reason for this is that our async ticks could be referencing GC objects, i.e. skel meshes, etc, and we don't want them to become unreachable while we are potentially using them
-		for (int TG = 0; TG < NiagaraNumTickGroups; ++TG)
-		{
-			for (TPair<UNiagaraSystem*, FNiagaraSystemSimulationRef>& SimPair : SystemSimulations[TG])
-			{
-				SimPair.Value->WaitForInstancesTickComplete();
-			}
-		}
+		WaitForAsyncWork();
 	}
 	
 	for (auto It = ScalabilityManagers.CreateIterator(); It; ++It)
@@ -903,13 +923,12 @@ UNiagaraDataChannelHandler* FNiagaraWorldManager::FindDataChannelHandler(const U
 
 void FNiagaraWorldManager::WaitForAsyncWork()
 {
-	for (int TG = 0; TG < NiagaraNumTickGroups; ++TG)
-	{
-		for (TPair<UNiagaraSystem*, FNiagaraSystemSimulationRef>& SimPair : SystemSimulations[TG])
+	ForAllSystemSimulations(
+		[](FNiagaraSystemSimulation& SystemSimulation)
 		{
-			SimPair.Value->WaitForInstancesTickComplete();
+			SystemSimulation.WaitForInstancesTickComplete();
 		}
-	}
+	);
 }
 
 void FNiagaraWorldManager::OnWorldInit(UWorld* World, const UWorld::InitializationValues IVS)
@@ -988,17 +1007,6 @@ void FNiagaraWorldManager::OnWorldBeginTearDown(UWorld* World)
 			WorldMan->OnWorldBeginTearDown();
 		}		
 	}
-// 	FNiagaraWorldManager** Manager = WorldManagers.Find(World);
-// 	if (Manager)
-// 	{
-// 		delete (*Manager);
-// 		WorldManagers.Remove(World);
-// 	}
-// 	FNiagaraWorldManager** Manager = WorldManagers.Find(World);
-// 	if (Manager)
-// 	{
-// 		Manager->SystemSimulations
-// 	}
 }
 
 void FNiagaraWorldManager::OnComputeDispatchInterfaceDestroyed(FNiagaraGpuComputeDispatchInterface* InComputeDispatchInterface)
@@ -1340,7 +1348,7 @@ bool FNiagaraWorldManager::PrepareCachedViewInfo(const APlayerController* Player
 	return false;
 }
 
-void FNiagaraWorldManager::Tick(ETickingGroup TickGroup, float DeltaSeconds, ELevelTick TickType, ENamedThreads::Type CurrentThread, const FGraphEventRef& MyCompletionGraphEvent)
+void FNiagaraWorldManager::Tick(ENiagaraWorldManagerTickPass TickPass, ETickingGroup TickGroup, float DeltaSeconds, const FGraphEventRef& MyCompletionGraphEvent)
 {
 	check(TickGroup >= NiagaraFirstTickGroup && TickGroup <= NiagaraLastTickGroup);
 
@@ -1349,6 +1357,46 @@ void FNiagaraWorldManager::Tick(ETickingGroup TickGroup, float DeltaSeconds, ELe
 	SCOPE_CYCLE_COUNTER(STAT_NiagaraWorldManTick);
 	SCOPE_CYCLE_COUNTER(STAT_NiagaraOverview_GT);
 
+	// Determine what work needs to run
+	const int ActualTickGroup = NiagaraWorldManagerInternal::GetNiagaraTickGroup(TickGroup);
+	const FNiagaraWorldManagerTickFunction& NonOverlappingTickFunction = TickFunctions[ActualTickGroup][int(ENiagaraWorldManagerTickPass::NoneOverlapping)];
+	const bool bIsFirstPass = TickPass == ENiagaraWorldManagerTickPass::First || NonOverlappingTickFunction.IsTickFunctionEnabled() == false;
+	const bool bIsLastPass = TickPass == ENiagaraWorldManagerTickPass::Last || NonOverlappingTickFunction.IsTickFunctionEnabled() == false;
+
+	// Execute work at the start of each tick group
+	if (bIsFirstPass)
+	{
+		ExecutePreTickWork(TickGroup, DeltaSeconds);
+	}
+
+	// If we are in paused don't do anything
+	if (DebugPlaybackMode == ENiagaraDebugPlaybackMode::Paused)
+	{
+		return;
+	}
+
+	// Tick generated data (First tick pass only)
+	if (bIsFirstPass)
+	{
+		CSV_SCOPED_TIMING_STAT(Particles, CoreSystems_NiagaraDIGeneratedData);
+		for (auto& GeneratedData : DIGeneratedData)
+		{
+			GeneratedData.Value->Tick(TickGroup, DeltaSeconds);
+		}
+	}
+
+	// Execute simulations for this tick group
+	ExecuteSimulations(TickPass, TickGroup, DeltaSeconds, MyCompletionGraphEvent);
+
+	// Execute work at the end of each tick group
+	if (bIsLastPass)
+	{
+		ExecutePostTickWork(TickGroup, DeltaSeconds);
+	}
+}
+
+void FNiagaraWorldManager::ExecutePreTickWork(ETickingGroup TickGroup, float DeltaSeconds)
+{
 	// Execute any global and local deferred functions
 	NiagaraWorldManagerInternal::ExecuteGlobalDeferredCallbacks();
 	DeferredMethods.ExecuteAndClear();
@@ -1467,39 +1515,26 @@ void FNiagaraWorldManager::Tick(ETickingGroup TickGroup, float DeltaSeconds, ELe
 			TickParameterCollections();
 		}
 	}
+}
 
-	// If we are in paused don't do anything
-	if ( DebugPlaybackMode == ENiagaraDebugPlaybackMode::Paused )
-	{
-		return;
-	}
-
-	// Tick generated data
-	{
-		CSV_SCOPED_TIMING_STAT(Particles, CoreSystems_NiagaraDIGeneratedData);
-		for (auto& GeneratedData : DIGeneratedData)
-		{
-			GeneratedData.Value->Tick(TickGroup, DeltaSeconds);
-		}
-	}
-
+void FNiagaraWorldManager::ExecuteSimulations(ENiagaraWorldManagerTickPass TickPass, ETickingGroup TickGroup, float DeltaSeconds, const FGraphEventRef& MyCompletionGraphEvent)
+{
 	// Now tick all system instances. 
 	const int ActualTickGroup = NiagaraWorldManagerInternal::GetNiagaraTickGroup(TickGroup);
 
 	ActiveNiagaraTickGroup = ActualTickGroup;
 
-	const bool bKillUniqueSims = GNiagaraWorldManagerKillUniqueSims;
-	bool bTickFunctionCanOverlapTickGroups = GNiagaraAllowAsyncWorkToEndOfFrame != 0;
-	for (auto SimIt=SystemSimulations[ActualTickGroup].CreateIterator(); SimIt; ++SimIt)
+	bool bHasSimulationsInTickPass = false;
+	bool bCanOverlapTickGroups = true;
+	for (auto SimIt=SystemSimulations[ActualTickGroup][int(TickPass)].CreateIterator(); SimIt; ++SimIt)
 	{
 		FNiagaraSystemSimulation* Sim = &SimIt.Value().Get();
 
-		const bool bUniqueKill = bKillUniqueSims && SimIt.Value().IsUnique();
-		if (!bUniqueKill && Sim->IsValid())
+		if (!SimIt.Value().IsUnique() && Sim->IsValid())
 		{
+			bHasSimulationsInTickPass = true;
 			Sim->Tick_GameThread(DeltaSeconds, MyCompletionGraphEvent);
-
-			bTickFunctionCanOverlapTickGroups &= Sim->GetSystem()->AsyncWorkCanOverlapTickGroups();
+			bCanOverlapTickGroups &= Sim->GetSystem()->AsyncWorkCanOverlapTickGroups();
 		}
 		else
 		{
@@ -1508,21 +1543,39 @@ void FNiagaraWorldManager::Tick(ETickingGroup TickGroup, float DeltaSeconds, ELe
 		}
 	}
 
-	//If all systems in this TG can overlap their async work then we can allow the tick as a whole to run until last demotable.
-	//When adding new system sims to this tick group we can change this.
-	//Any systems needing their async work to complete must have this Tick function finish inside this TG so that we can properly complete our async work before the TG finishes.
-	if(bTickFunctionCanOverlapTickGroups)
+	// New path that uses different tick passes for overlapping / non overlapping work
+	if (bUseTickPasses)
 	{
-		//Set here but this wont take effect until next tick.
-		TickFunctions[ActualTickGroup].EndTickGroup = TG_LastDemotable;
+		// Should we disable the none overlapping tick function from running?
+		if (TickPass == ENiagaraWorldManagerTickPass::NoneOverlapping && !bHasSimulationsInTickPass)
+		{
+			FNiagaraWorldManagerTickFunction& NonOverlappingTickFunction = TickFunctions[ActualTickGroup][int(ENiagaraWorldManagerTickPass::NoneOverlapping)];
+			NonOverlappingTickFunction.SetTickFunctionEnable(false);
+		}
 	}
+	// Old tick path which modifies the end tick group which can cause stalls
 	else
 	{
-		TickFunctions[ActualTickGroup].EndTickGroup = (ETickingGroup)ActualTickGroup;
+		//If all systems in this TG can overlap their async work then we can allow the tick as a whole to run until last demotable.
+		//When adding new system sims to this tick group we can change this.
+		//Any systems needing their async work to complete must have this Tick function finish inside this TG so that we can properly complete our async work before the TG finishes.
+		if ( bCanOverlapTickGroups )
+		{
+			//Set here but this wont take effect until next tick.
+			TickFunctions[ActualTickGroup][int(TickPass)].EndTickGroup = TG_LastDemotable;
+		}
+		else
+		{
+			TickFunctions[ActualTickGroup][int(TickPass)].EndTickGroup = (ETickingGroup)ActualTickGroup;
+		}
 	}
 
+
 	ActiveNiagaraTickGroup = -1;
+}
 
+void FNiagaraWorldManager::ExecutePostTickWork(ETickingGroup TickGroup, float DeltaSeconds)
+{
 	//We update scalability managers here so that any new systems can be culled or setup with other scalability based parameters correctly for their spawn.
 	UpdateScalabilityManagers(DeltaSeconds, true);
 
@@ -1541,31 +1594,44 @@ void FNiagaraWorldManager::Tick(ETickingGroup TickGroup, float DeltaSeconds, ELe
 	}
 }
 
-void FNiagaraWorldManager::DumpDetails(FOutputDevice& Ar)
+void FNiagaraWorldManager::DumpDetails(FOutputDevice& Ar, bool bDetailed)
 {
-	Ar.Logf(TEXT("=== FNiagaraWorldManager Dumping Detailed Information"));
-
-	static const UEnum* TickingGroupEnum = FindObjectChecked<UEnum>(nullptr, TEXT("/Script/Engine.ETickingGroup"));
+	Ar.Logf(TEXT("=== FNiagaraWorldManager Dumping Information"));
 
+	const UEnum* TickingGroupEnum = StaticEnum<ETickingGroup>();
 	for ( int TG=0; TG < NiagaraNumTickGroups; ++TG )
 	{
-		if (SystemSimulations[TG].Num() == 0 )
+		TStringBuilder<8> TickStateString;
+		int32 TotalSimulations = 0;
+		for (int iTickPass=0; iTickPass < int(ENiagaraWorldManagerTickPass::Num); ++iTickPass)
 		{
-			continue;
+			TickStateString.Append(TickFunctions[TG][iTickPass].IsTickFunctionEnabled() ? TEXT("X") : TEXT("."));
+			TotalSimulations += SystemSimulations[TG][iTickPass].Num();
 		}
 
-		Ar.Logf(TEXT("TickingGroup %s"), *TickingGroupEnum->GetNameStringByIndex(TG + NiagaraFirstTickGroup));
+		Ar.Logf(
+			TEXT("%-20s | TickState(%s) | %d Simulations"),
+			*TickingGroupEnum->GetNameStringByIndex(TG + NiagaraFirstTickGroup),
+			TickStateString.ToString(),
+			TotalSimulations
+		);
 
-		for (TPair<UNiagaraSystem*, FNiagaraSystemSimulationRef>& SystemSim : SystemSimulations[TG])
+		if (bDetailed)
 		{
-			FNiagaraSystemSimulation* Sim = &SystemSim.Value.Get();
-			if ( !Sim->IsValid() )
+			for (int iTickPass = 0; iTickPass < int(ENiagaraWorldManagerTickPass::Num); ++iTickPass)
 			{
-				continue;
-			}
+				for (TPair<UNiagaraSystem*, FNiagaraSystemSimulationRef>& SystemSim : SystemSimulations[TG][iTickPass])
+				{
+					FNiagaraSystemSimulation* Sim = &SystemSim.Value.Get();
+					if (!Sim->IsValid())
+					{
+						continue;
+					}
 
-			Ar.Logf(TEXT("\tSimulation %s"), *Sim->GetSystem()->GetFullName());
-			Sim->DumpTickInfo(Ar);
+					Ar.Logf(TEXT("\tSimulation %s"), *Sim->GetSystem()->GetFullName());
+					Sim->DumpTickInfo(Ar);
+				}
+			}
 		}
 	}
 }
@@ -2327,7 +2393,10 @@ bool FNiagaraWorldManager::IsComponentLocalPlayerLinked(const USceneComponent* I
 void FNiagaraWorldManager::SetTickGroupPriority(const int32 NiagaraTickGroupIndex, const bool bHighPriority)
 {
 	check ((NiagaraTickGroupIndex >= 0) && (NiagaraTickGroupIndex < NiagaraNumTickGroups));
-	TickFunctions[NiagaraTickGroupIndex].SetPriorityIncludingPrerequisites(bHighPriority);
+	for ( int iTickPass=0; iTickPass < int(ENiagaraWorldManagerTickPass::Num); ++iTickPass)
+	{
+		TickFunctions[NiagaraTickGroupIndex][iTickPass].SetPriorityIncludingPrerequisites(bHighPriority);
+	}
 }
 
 #if DEBUG_SCALABILITY_STATE
diff --git a/Engine/Plugins/FX/Niagara/Source/Niagara/Public/NiagaraWorldManager.h b/Engine/Plugins/FX/Niagara/Source/Niagara/Public/NiagaraWorldManager.h
--- a/Engine/Plugins/FX/Niagara/Source/Niagara/Public/NiagaraWorldManager.h
+++ b/Engine/Plugins/FX/Niagara/Source/Niagara/Public/NiagaraWorldManager.h
@@ -33,6 +33,16 @@ class UNiagaraComponentPool;
 struct FNiagaraScalabilityState;
 class UNiagaraDataChannelHandler;
 
+enum class ENiagaraWorldManagerTickPass : int
+{
+	NoneOverlapping,
+	Overlapping,
+
+	Num,
+	First = NoneOverlapping,
+	Last  = Overlapping,
+};
+
 USTRUCT()
 struct FNiagaraWorldManagerTickFunction : public FTickFunction
 {
@@ -44,7 +54,8 @@ struct FNiagaraWorldManagerTickFunction : public FTickFunction
 	virtual FName DiagnosticContext(bool bDetailed)  override;
 	//~ FTickFunction Interface
 
-	FNiagaraWorldManager* Owner;
+	FNiagaraWorldManager*			Owner = nullptr;
+	ENiagaraWorldManagerTickPass	TickPass = ENiagaraWorldManagerTickPass::Num;
 };
 
 template<>
@@ -127,8 +138,13 @@ public:
 	/** Called at the very beginning of world ticking. */
 	void TickStart(float DeltaSeconds);
 
-	void Tick(ETickingGroup TickGroup, float DeltaSeconds, ELevelTick TickType, ENamedThreads::Type CurrentThread, const FGraphEventRef& MyCompletionGraphEvent);
+	void Tick(ENiagaraWorldManagerTickPass TickPass, ETickingGroup TickGroup, float DeltaSeconds, const FGraphEventRef& MyCompletionGraphEvent);
+private:
+	void ExecutePreTickWork(ETickingGroup TickGroup, float DeltaSeconds);
+	void ExecuteSimulations(ENiagaraWorldManagerTickPass TickPass, ETickingGroup TickGroup, float DeltaSeconds, const FGraphEventRef& MyCompletionGraphEvent);
+	void ExecutePostTickWork(ETickingGroup TickGroup, float DeltaSeconds);
 
+public:
 	void PreActorTick(ELevelTick InLevelTick, float InDeltaSeconds);
 
 	/** Called after all actor tick groups are complete. */
@@ -176,7 +192,7 @@ public:
 	void UpdateScalabilityManagers(float DeltaSeconds, bool bNewSpawnsOnly);
 
 	// Dump details about what's inside the world manager
-	void DumpDetails(FOutputDevice& Ar);
+	void DumpDetails(FOutputDevice& Ar, bool bDetailed);
 	
 	UWorld* GetWorld();
 	FORCEINLINE UWorld* GetWorld()const { return World; }
@@ -296,6 +312,8 @@ public:
 	/** Waits for all currently in flight async work to be completed. */
 	void WaitForAsyncWork();
 
+	bool UseTickPasses() const { return bUseTickPasses; }
+
 private:
 	// Callback function registered with global world delegates to instantiate world manager when a game world is created
 	static void OnWorldInit(UWorld* World, const UWorld::InitializationValues IVS);
@@ -378,11 +396,12 @@ private:
 
 	int ActiveNiagaraTickGroup;
 
-	FNiagaraWorldManagerTickFunction TickFunctions[NiagaraNumTickGroups];
+	bool bUseTickPasses = true;
+	FNiagaraWorldManagerTickFunction TickFunctions[NiagaraNumTickGroups][int(ENiagaraWorldManagerTickPass::Num)];
 
 	TMap<TObjectPtr<UNiagaraParameterCollection>, TObjectPtr<UNiagaraParameterCollectionInstance>> ParameterCollections;
 
-	TMap<UNiagaraSystem*, FNiagaraSystemSimulationRef> SystemSimulations[NiagaraNumTickGroups];
+	TMap<UNiagaraSystem*, FNiagaraSystemSimulationRef> SystemSimulations[NiagaraNumTickGroups][int(ENiagaraWorldManagerTickPass::Num)];
 
 	TArray<FNiagaraSystemSimulationRef> SimulationsWithPostActorWork;
 
@@ -454,9 +473,12 @@ void FNiagaraWorldManager::ForAllSystemSimulations(TAction Func)
 {
 	for (int TG = 0; TG < NiagaraNumTickGroups; ++TG)
 	{
-		for (TPair<UNiagaraSystem*, FNiagaraSystemSimulationRef>& SimPair : SystemSimulations[TG])
+		for (int iTickPass = 0; iTickPass < int(ENiagaraWorldManagerTickPass::Num); ++iTickPass)
 		{
-			Func(SimPair.Value.Get());
+			for (TPair<UNiagaraSystem*, FNiagaraSystemSimulationRef>& SimPair : SystemSimulations[TG][iTickPass])
+			{
+				Func(SimPair.Value.Get());
+			}
 		}
 	}
 }
