- Return element offset rather than byte offset when adding an expression

#rnx

diff --git a/Engine/Plugins/FX/Niagara/Source/Niagara/Private/Stateless/NiagaraStatelessEmitterDataBuildContext.cpp b/Engine/Plugins/FX/Niagara/Source/Niagara/Private/Stateless/NiagaraStatelessEmitterDataBuildContext.cpp
--- a/Engine/Plugins/FX/Niagara/Source/Niagara/Private/Stateless/NiagaraStatelessEmitterDataBuildContext.cpp
+++ b/Engine/Plugins/FX/Niagara/Source/Niagara/Private/Stateless/NiagaraStatelessEmitterDataBuildContext.cpp
@@ -111,6 +111,7 @@ int32 FNiagaraStatelessEmitterDataBuildContext::AddExpression(const FInstancedSt
 		RendererBindings.AddParameter(ExpressionVariable, false, false, &DataOffset);
 
 		Expressions.Emplace(DataOffset, Expression->Build(*this));
+		DataOffset /= sizeof(uint32);
 	}
 	return DataOffset;
 }
