- PVS studio fix

diff --git a/Engine/Plugins/FX/Niagara/Source/Niagara/Internal/Stateless/Expressions/NiagaraStatelessExpressionColor.h b/Engine/Plugins/FX/Niagara/Source/Niagara/Internal/Stateless/Expressions/NiagaraStatelessExpressionColor.h
--- a/Engine/Plugins/FX/Niagara/Source/Niagara/Internal/Stateless/Expressions/NiagaraStatelessExpressionColor.h
+++ b/Engine/Plugins/FX/Niagara/Source/Niagara/Internal/Stateless/Expressions/NiagaraStatelessExpressionColor.h
@@ -44,7 +44,7 @@ struct FNiagaraStatelessExpressionColorBinding final : public FNiagaraStatelessE
 	GENERATED_BODY()
 
 	virtual ~FNiagaraStatelessExpressionColorBinding() = default;
-	virtual bool IsConstant() const { return false; }
+	virtual bool IsConstant() const override { return false; }
 	virtual FInstancedStruct BuildInternal(const FNiagaraStatelessEmitterDataBuildContext& BuildContext) const override;
 	virtual FLinearColor EvaluateInternal(const FEvaluateContext& Context) const override;
 
diff --git a/Engine/Plugins/FX/Niagara/Source/Niagara/Internal/Stateless/Expressions/NiagaraStatelessExpressionFloat.h b/Engine/Plugins/FX/Niagara/Source/Niagara/Internal/Stateless/Expressions/NiagaraStatelessExpressionFloat.h
--- a/Engine/Plugins/FX/Niagara/Source/Niagara/Internal/Stateless/Expressions/NiagaraStatelessExpressionFloat.h
+++ b/Engine/Plugins/FX/Niagara/Source/Niagara/Internal/Stateless/Expressions/NiagaraStatelessExpressionFloat.h
@@ -44,7 +44,7 @@ struct FNiagaraStatelessExpressionFloatBinding final : public FNiagaraStatelessE
 	GENERATED_BODY()
 
 	virtual ~FNiagaraStatelessExpressionFloatBinding() = default;
-	virtual bool IsConstant() const { return false; }
+	virtual bool IsConstant() const override { return false; }
 	virtual FInstancedStruct BuildInternal(const FNiagaraStatelessEmitterDataBuildContext& BuildContext) const override;
 	virtual float EvaluateInternal(const FEvaluateContext& Context) const override;
 
diff --git a/Engine/Plugins/FX/Niagara/Source/Niagara/Internal/Stateless/Expressions/NiagaraStatelessExpressionVec2.h b/Engine/Plugins/FX/Niagara/Source/Niagara/Internal/Stateless/Expressions/NiagaraStatelessExpressionVec2.h
--- a/Engine/Plugins/FX/Niagara/Source/Niagara/Internal/Stateless/Expressions/NiagaraStatelessExpressionVec2.h
+++ b/Engine/Plugins/FX/Niagara/Source/Niagara/Internal/Stateless/Expressions/NiagaraStatelessExpressionVec2.h
@@ -44,7 +44,7 @@ struct FNiagaraStatelessExpressionVec2Binding final : public FNiagaraStatelessEx
 	GENERATED_BODY()
 
 	virtual ~FNiagaraStatelessExpressionVec2Binding() = default;
-	virtual bool IsConstant() const { return false; }
+	virtual bool IsConstant() const override { return false; }
 	virtual FInstancedStruct BuildInternal(const FNiagaraStatelessEmitterDataBuildContext& BuildContext) const override;
 	virtual FVector2f EvaluateInternal(const FEvaluateContext& Context) const override;
 
diff --git a/Engine/Plugins/FX/Niagara/Source/Niagara/Internal/Stateless/Expressions/NiagaraStatelessExpressionVec3.h b/Engine/Plugins/FX/Niagara/Source/Niagara/Internal/Stateless/Expressions/NiagaraStatelessExpressionVec3.h
--- a/Engine/Plugins/FX/Niagara/Source/Niagara/Internal/Stateless/Expressions/NiagaraStatelessExpressionVec3.h
+++ b/Engine/Plugins/FX/Niagara/Source/Niagara/Internal/Stateless/Expressions/NiagaraStatelessExpressionVec3.h
@@ -44,7 +44,7 @@ struct FNiagaraStatelessExpressionVec3Binding final : public FNiagaraStatelessEx
 	GENERATED_BODY()
 
 	virtual ~FNiagaraStatelessExpressionVec3Binding() = default;
-	virtual bool IsConstant() const { return false; }
+	virtual bool IsConstant() const override { return false; }
 	virtual FInstancedStruct BuildInternal(const FNiagaraStatelessEmitterDataBuildContext& BuildContext) const override;
 	virtual FVector3f EvaluateInternal(const FEvaluateContext& Context) const override;
 
diff --git a/Engine/Plugins/FX/Niagara/Source/Niagara/Internal/Stateless/Expressions/NiagaraStatelessExpressionVec4.h b/Engine/Plugins/FX/Niagara/Source/Niagara/Internal/Stateless/Expressions/NiagaraStatelessExpressionVec4.h
--- a/Engine/Plugins/FX/Niagara/Source/Niagara/Internal/Stateless/Expressions/NiagaraStatelessExpressionVec4.h
+++ b/Engine/Plugins/FX/Niagara/Source/Niagara/Internal/Stateless/Expressions/NiagaraStatelessExpressionVec4.h
@@ -44,7 +44,7 @@ struct FNiagaraStatelessExpressionVec4Binding final : public FNiagaraStatelessEx
 	GENERATED_BODY()
 
 	virtual ~FNiagaraStatelessExpressionVec4Binding() = default;
-	virtual bool IsConstant() const { return false; }
+	virtual bool IsConstant() const override { return false; }
 	virtual FInstancedStruct BuildInternal(const FNiagaraStatelessEmitterDataBuildContext& BuildContext) const override;
 	virtual FVector4f EvaluateInternal(const FEvaluateContext& Context) const override;
 
diff --git a/Engine/Plugins/FX/Niagara/Source/Niagara/Internal/Stateless/NiagaraStatelessExpression.h b/Engine/Plugins/FX/Niagara/Source/Niagara/Internal/Stateless/NiagaraStatelessExpression.h
--- a/Engine/Plugins/FX/Niagara/Source/Niagara/Internal/Stateless/NiagaraStatelessExpression.h
+++ b/Engine/Plugins/FX/Niagara/Source/Niagara/Internal/Stateless/NiagaraStatelessExpression.h
@@ -27,6 +27,8 @@ struct FNiagaraStatelessExpression
 
 	GENERATED_BODY()
 
+	virtual ~FNiagaraStatelessExpression() = default;
+
 	virtual FInstancedStruct Build(const FNiagaraStatelessEmitterDataBuildContext& BuildContext) const { checkNoEntry(); return FInstancedStruct();  }
 	virtual void Evaluate(const FEvaluateContext& Context, void* ValueAddress) const { checkNoEntry(); }
 	virtual FNiagaraTypeDefinition GetOutputTypeDef() const { checkNoEntry(); return FNiagaraTypeDefinition(); }
