- Fix for not copying data to the internal array data when CopyToInternal is called
- This can result in an off by one error in the editor

diff --git a/Engine/Plugins/FX/Niagara/Source/Niagara/Classes/NiagaraDataInterfaceArrayImpl.h b/Engine/Plugins/FX/Niagara/Source/Niagara/Classes/NiagaraDataInterfaceArrayImpl.h
--- a/Engine/Plugins/FX/Niagara/Source/Niagara/Classes/NiagaraDataInterfaceArrayImpl.h
+++ b/Engine/Plugins/FX/Niagara/Source/Niagara/Classes/NiagaraDataInterfaceArrayImpl.h
@@ -70,6 +70,7 @@
 			if ( TypedDestination != nullptr ) \
 			{ \
 				TypedDestination->MEMBERNAME = MEMBERNAME; \
+				TypedDestination->GetProxyAs<FProxyType>()->template SetArrayData<decltype(MEMBERNAME)::ElementType>(MakeArrayView(TypedDestination->MEMBERNAME)); \
 			} \
 			return TypedDestination != nullptr;  \
 		} \
@@ -528,10 +529,6 @@ struct FNDIArrayProxyImpl : public INDIArrayProxyBase
 		}
 		else
 		{
-			// Note: Can be called from PostEditChangeProperty where we can expect to the have > 1 data interface.  For BP user parameters we can expect this to be 1 exactly.
-		#if WITH_EDITORONLY_DATA
-			ensure(GIsTransacting || PerInstanceData_GameThread.Num() == 1);
-		#endif
 			BeginSetArrayFromBP(false);
 
 			Owner->GetArrayReference().SetNum(InArrayData.Num());
