NiagaraSimplexNoise simplify

The helper types etc. are all there to maintain correspondence to the HLSL shader code,
but that original code is fairly opaque anyway, and is heavily geared towards "everything
is inlined" and fast math-style optimizations compiled by the shader compiler that we don't
get on the CPU side. Refactor it into a more natural way for CPU instead.

(This changes the results slightly, but within FP round-off expectations for the values in question.)

Tested on 1 million random sample points in [-4,4]^3 and in Niagara EngineTest.

diff --git a/Engine/Plugins/FX/Niagara/Source/Niagara/Private/NiagaraSimplexNoise.cpp b/Engine/Plugins/FX/Niagara/Source/Niagara/Private/NiagaraSimplexNoise.cpp
--- a/Engine/Plugins/FX/Niagara/Source/Niagara/Private/NiagaraSimplexNoise.cpp
+++ b/Engine/Plugins/FX/Niagara/Source/Niagara/Private/NiagaraSimplexNoise.cpp
@@ -3,125 +3,29 @@
 #include "NiagaraSimplexNoise.h"
 #include "Math/IntVector.h"
 
-namespace {
+// This matches the HLSL code found in Engine/Shaders/Private/Random.ush, but is not
+// fully unrolled and also doesn't expect everything to get 100% inlined.
 
-// Fairly straightforward CPU implementation of the equivalent HLSL code found in Engine/Shaders/Private/Random.ush
-// Also contains a partial implementation of some math types and functions found in HLSL but not in UE
-
-struct FNiagaraUIntVector
-{
-	uint32 X;
-	uint32 Y;
-	uint32 Z;
-
-	FNiagaraUIntVector() : X(0), Y(0), Z(0) {}
-	FNiagaraUIntVector(uint32 X, uint32 Y, uint32 Z) : X(X), Y(Y), Z(Z) {}
-	FNiagaraUIntVector(uint32 Value) : X(Value), Y(Value), Z(Value) {}
-	FNiagaraUIntVector(FIntVector Values) : X(Values.X), Y(Values.Y), Z(Values.Z) {}
-
-	FNiagaraUIntVector operator+(FNiagaraUIntVector Rhs)
-	{
-		return FNiagaraUIntVector{ X + Rhs.X, Y + Rhs.Y, Z + Rhs.Z };
-	}
-
-	FNiagaraUIntVector operator*(FNiagaraUIntVector Rhs)
-	{
-		return FNiagaraUIntVector{ X * Rhs.X, Y * Rhs.Y, Z * Rhs.Z };
-	}
-
-	FNiagaraUIntVector operator>>(uint32 Shift)
-	{
-		return FNiagaraUIntVector(X >> Shift, Y >> Shift, Z >> Shift);
-	}
-
-	FNiagaraUIntVector operator&(FNiagaraUIntVector Rhs)
-	{
-		return FNiagaraUIntVector{ X & Rhs.X, Y & Rhs.Y, Z & Rhs.Z };
-	}
-
-	uint32& operator[](int Index)
-	{
-		if (Index == 0) return X;
-		if (Index == 1) return Y;
-		if (Index == 2) return Z;
-		check(false); 
-		return X;
-	}
-
-	const uint32& operator[](int Index) const
-	{
-		if (Index == 0) return X;
-		if (Index == 1) return Y;
-		if (Index == 2) return Z;
-		check(false);
-		return X;
-	}
-};
-
-FNiagaraUIntVector Rand3DPCG16(FIntVector p)
+static FIntVector Rand3DPCG16(FIntVector p)
 {
-	FNiagaraUIntVector v = FNiagaraUIntVector(p);
-	
-	v = v * 1664525u + 1013904223u;
-
-	v.X += v.Y*v.Z;
-	v.Y += v.Z*v.X;
-	v.Z += v.X*v.Y;
-	v.X += v.Y*v.Z;
-	v.Y += v.Z*v.X;
-	v.Z += v.X*v.Y;
+	constexpr uint32 Mul = 1664525u;
+	constexpr uint32 Add = 1013904223u;
 
-	return v >> 16u;
-}
+	uint32 X = uint32(p.X) * Mul + Add;
+	uint32 Y = uint32(p.Y) * Mul + Add;
+	uint32 Z = uint32(p.Z) * Mul + Add;
 
-FVector3f NiagaraVectorFrac(FVector3f v)
-{
-	return FVector3f(FMath::Frac(v.X), FMath::Frac(v.Y), FMath::Frac(v.Z));
-}
+	X += Y * Z;
+	Y += Z * X;
+	Z += X * Y;
+	X += Y * Z;
+	Y += Z * X;
+	Z += X * Y;
 
-FVector3f NoiseTileWrap(FVector3f v, bool bTiling, float RepeatSize)
-{
-	FVector3f vv = bTiling ? (NiagaraVectorFrac(v / RepeatSize) * RepeatSize) : v;
-	return vv;
+	return FIntVector(X >> 16, Y >> 16, Z >> 16);
 }
 
-struct FNiagaraMatrix4x3
-{
-	FVector3f Row0;
-	FVector3f Row1;
-	FVector3f Row2;
-	FVector3f Row3;
-
-	FNiagaraMatrix4x3() : Row0(FVector3f::Zero()), Row1(FVector3f::Zero()), Row2(FVector3f::Zero()), Row3(FVector3f::Zero()) {}
-	FNiagaraMatrix4x3(FVector3f Row0, FVector3f Row1, FVector3f Row2, FVector3f Row3) : Row0(Row0), Row1(Row1), Row2(Row2), Row3(Row3) {}
-
-	FVector3f& operator[](int Row)
-	{
-		if (Row == 0) return Row0;
-		if (Row == 1) return Row1;
-		if (Row == 2) return Row2;
-		if (Row == 3) return Row3;
-		check(false);
-		return Row0;
-	}
-	const FVector3f& operator[](int Row) const
-	{
-		if (Row == 0) return Row0;
-		if (Row == 1) return Row1;
-		if (Row == 2) return Row2;
-		if (Row == 3) return Row3;
-		check(false);
-		return Row0;
-	}
-};
-
-FVector3f NiagaraVectorFloor(FVector3f v) {
-	return FVector3f(FGenericPlatformMath::FloorToFloat(v.X),
-		           FGenericPlatformMath::FloorToFloat(v.Y),
-		           FGenericPlatformMath::FloorToFloat(v.Z));
-};
-
-FIntVector NiagaraVectorFloorToInt(FVector3f v) {
+static FIntVector NiagaraVectorFloorToInt(FVector3f v) {
 	return FIntVector(
 		FGenericPlatformMath::FloorToInt(v.X),
 		FGenericPlatformMath::FloorToInt(v.Y),
@@ -129,157 +33,79 @@ FIntVector NiagaraVectorFloorToInt(FVector3f v) {
 	);
 };
 
-FVector3f NiagaraVectorStep(FVector3f v, FVector3f u)
-{
-	return FVector3f(u.X >= v.X ? 1.0f : 0.0f,
-		           u.Y >= v.Y ? 1.0f : 0.0f,
-		           u.Z >= v.Z ? 1.0f : 0.0f);
-}
-
-FVector3f NiagaraVectorSwizzle(FVector3f v, uint32 x, uint32 y, uint32 z)
-{
-	return FVector3f(v[x], v[y], v[z]);
-}
-
-FVector3f NiagaraVectorMin(FVector3f u, FVector3f v)
-{
-	return FVector3f(u.X < v.X ? u.X : v.X,
-		           u.Y < v.Y ? u.Y : v.Y, 
-		           u.Z < v.Z ? u.Z : v.Z);
-}
-
-FVector3f NiagaraVectorMax(FVector3f u, FVector3f v)
-{
-	return FVector3f(u.X > v.X ? u.X : v.X,
-		           u.Y > v.Y ? u.Y : v.Y,
-	           	   u.Z > v.Z ? u.Z : v.Z);
-} 
-
-FNiagaraMatrix4x3 SimplexCorners(FVector3f v)
-{
-	FVector3f tet = NiagaraVectorFloor(v + v.X / 3 + v.Y / 3 + v.Z / 3);
-	FVector3f base = tet - tet.X / 6 - tet.Y / 6 - tet.Z / 6;
-	FVector3f f = v - base;
-
-	FVector3f g = NiagaraVectorStep(NiagaraVectorSwizzle(f, 1, 2, 0), f), h = FVector3f(1.0f) - NiagaraVectorSwizzle(g, 2, 0, 1);
-	FVector3f a1 = NiagaraVectorMin(g, h) - 1. / 6., a2 = NiagaraVectorMax(g, h) - 1. / 3.;
-
-	return FNiagaraMatrix4x3(base, base + a1, base + a2, base + 0.5);
-}
-
-FVector4f NiagaraVector4Saturate(const FVector4f& v) {
-	return FVector4f(FMath::Clamp(v.X, 0.0f, 1.0f), FMath::Clamp(v.Y, 0.0f, 1.0f), FMath::Clamp(v.Z, 0.0f, 1.0f), FMath::Clamp(v.W, 0.0f, 1.0f));
-}
-
-FVector4f SimplexSmooth(FNiagaraMatrix4x3 f)
-{
-	const float scale = 1024.f / 375.f;
-	FVector4f d = FVector4f(FVector3f::DotProduct(f[0], f[0]), FVector3f::DotProduct(f[1], f[1]), FVector3f::DotProduct(f[2], f[2]), FVector3f::DotProduct(f[3], f[3]));
-	FVector4f s = NiagaraVector4Saturate(2.0f * d);
-	return scale * (FVector4f(1.0f, 1.0f, 1.0f, 1.0f) + s * (FVector4f(-3.0f, -3.0f, -3.0f, -3.0f) + s * (FVector4f(3.0f, 3.0f, 3.0f, 3.0f) - s)));
-}
-
-FNiagaraMatrix3x4 SimplexDSmooth(FNiagaraMatrix4x3 f)
-{
-	const float scale = 1024. / 375.;
-	FVector4f d = FVector4f(FVector3f::DotProduct(f[0], f[0]), FVector3f::DotProduct(f[1], f[1]), FVector3f::DotProduct(f[2], f[2]), FVector3f::DotProduct(f[3], f[3]));
-	FVector4f s = NiagaraVector4Saturate(2 * d);
-	s = -12 * FVector4f(scale, scale, scale, scale) + s * (24 * FVector4f(scale, scale, scale, scale) - s * 12 * scale);
-
-	return FNiagaraMatrix3x4(
-		s * FVector4f(f[0][0], f[1][0], f[2][0], f[3][0]),
-		s * FVector4f(f[0][1], f[1][1], f[2][1], f[3][1]),
-		s * FVector4f(f[0][2], f[1][2], f[2][2], f[3][2]));
-}
-
-FNiagaraUIntVector FNiagaraUIntVectorSwizzle(FNiagaraUIntVector v, int x, int y, int z)
+FVector3f SimplexNoiseOffsetFromSeed(uint32 Seed)
 {
-	return FNiagaraUIntVector(v[x], v[y], v[z]);
+	FIntVector Rand = Rand3DPCG16(FIntVector(Seed, Seed, Seed));
+	return FVector3f(Rand.X, Rand.Y, Rand.Z) * (1.0f / 100.0f);
 }
 
-FVector3f FNiagaraUIntVectorToFVector(FNiagaraUIntVector v)
+static FVector3f SimplexGvec(int Rand)
 {
-	return FVector3f(float(v.X), float(v.Y), float(v.Z));
+	FVector3f MGradientScale(1.f / 0x4000, 1.f / 0x2000, 1.f / 0x1000);
+	return FVector3f(float(Rand & 0x8000), float(Rand & 0x4000), float(Rand & 0x2000)) * MGradientScale - FVector3f(1.f);
 }
 
-FVector3f MulFVector4AndFNiagaraMatrix4x3(const FVector4f& lhs, FNiagaraMatrix4x3 rhs)
+FNiagaraMatrix3x4 JacobianSimplex_ALU(FVector3f V)
 {
-	return FVector3f(lhs[0] * rhs[0][0] + lhs[1] * rhs[1][0] + lhs[2] * rhs[2][0] + lhs[3] * rhs[3][0],
-		           lhs[0] * rhs[0][1] + lhs[1] * rhs[1][1] + lhs[2] * rhs[2][1] + lhs[3] * rhs[3][1],
-	           	   lhs[0] * rhs[0][2] + lhs[1] * rhs[1][2] + lhs[2] * rhs[2][2] + lhs[3] * rhs[3][2]);
-}
+	// SimplexCorners in Random.ush
+	float SkewAmount = (V.X + V.Y + V.Z) * (1.0f / 3.0f);
+	FVector3f Tet {
+		FGenericPlatformMath::FloorToFloat(V.X + SkewAmount),
+		FGenericPlatformMath::FloorToFloat(V.Y + SkewAmount),
+		FGenericPlatformMath::FloorToFloat(V.Z + SkewAmount)
+	};
+	FVector3f Base = Tet - FVector3f((Tet.X + Tet.Y + Tet.Z) * (1.0f / 6.0f));
+	FVector3f F = V - Base;
 
-FVector3f MulFNiagaraMatrix3x4FAndVector4(const FNiagaraMatrix3x4& lhs, const FVector4f& rhs)
-{
-	return FVector3f(lhs[0][0] * rhs[0] + lhs[0][1] * rhs[1] + lhs[0][2] * rhs[2] + lhs[0][3] * rhs[3],
-		           lhs[1][0] * rhs[0] + lhs[1][1] * rhs[1] + lhs[1][2] * rhs[2] + lhs[1][3] * rhs[3],
-		           lhs[2][0] * rhs[0] + lhs[2][1] * rhs[1] + lhs[2][2] * rhs[2] + lhs[2][3] * rhs[3]);
-}
+	FVector3f G {
+		F.X >= F.Y ? 1.0f : 0.0f,
+		F.Y >= F.Z ? 1.0f : 0.0f,
+		F.Z >= F.X ? 1.0f : 0.0f
+	};
+	FVector3f H { 1.0f - G.Z, 1.0f - G.X, 1.0f - G.Y };
 
-} // anonymous namespace
+	FVector3f A1 = G.ComponentMin(H) - FVector3f(1.f / 6.f);
+	FVector3f A2 = G.ComponentMax(H) - FVector3f(1.f / 3.f);
 
-FVector3f SimplexNoiseOffsetFromSeed(uint32 Seed)
-{
-	return FNiagaraUIntVectorToFVector(Rand3DPCG16(FIntVector(Seed, Seed, Seed))) / 100.0f;
-}
+	FVector3f TetCorners[4] {
+		Base,
+		Base + A1,
+		Base + A2,
+		Base + FVector3f(0.5f)
+	};
 
-#define MGradientMask FNiagaraUIntVector(0x8000, 0x4000, 0x2000)
-#define MGradientScale FVector3f(1.f / 0x4000, 1.f / 0x2000, 1.f / 0x1000)
+	// Actual Jacobian calc
+	FNiagaraMatrix3x4 Jacobian;
 
-FNiagaraMatrix3x4 JacobianSimplex_ALU(FVector3f v)
-{
-	FNiagaraMatrix4x3 T = SimplexCorners(v);
-	FNiagaraUIntVector rand;
-	FNiagaraMatrix4x3 gvec[3], fv;
-	FNiagaraMatrix3x4 grad;
+	for (int CornerIndex = 0; CornerIndex < 4; CornerIndex++)
+	{
+		FVector3f RelVec = V - TetCorners[CornerIndex];
+		FIntVector Rand = Rand3DPCG16(NiagaraVectorFloorToInt(6 * TetCorners[CornerIndex] + 0.5));
 
-	fv[0] = v - T[0];
-	rand = Rand3DPCG16(NiagaraVectorFloorToInt(6 * T[0] + 0.5));
-	gvec[0][0] = FVector3f(FNiagaraUIntVectorToFVector(FNiagaraUIntVectorSwizzle(rand, 0, 0, 0) & MGradientMask)) * MGradientScale - 1;
-	gvec[1][0] = FVector3f(FNiagaraUIntVectorToFVector(FNiagaraUIntVectorSwizzle(rand, 1, 1, 1) & MGradientMask)) * MGradientScale - 1;
-	gvec[2][0] = FVector3f(FNiagaraUIntVectorToFVector(FNiagaraUIntVectorSwizzle(rand, 2, 2, 2) & MGradientMask)) * MGradientScale - 1;
-	grad[0][0] = FVector3f::DotProduct(gvec[0][0], fv[0]);
-	grad[1][0] = FVector3f::DotProduct(gvec[1][0], fv[0]);
-	grad[2][0] = FVector3f::DotProduct(gvec[2][0], fv[0]);
+		FVector3f GvecX = SimplexGvec(Rand.X);
+		FVector3f GvecY = SimplexGvec(Rand.Y);
+		FVector3f GvecZ = SimplexGvec(Rand.Z);
 
-	fv[1] = v - T[1];
-	rand = Rand3DPCG16(NiagaraVectorFloorToInt(6 * T[1] + 0.5));
-	gvec[0][1] = FVector3f(FNiagaraUIntVectorToFVector(FNiagaraUIntVectorSwizzle(rand, 0, 0, 0) & MGradientMask)) * MGradientScale - 1;
-	gvec[1][1] = FVector3f(FNiagaraUIntVectorToFVector(FNiagaraUIntVectorSwizzle(rand, 1, 1, 1) & MGradientMask)) * MGradientScale - 1;
-	gvec[2][1] = FVector3f(FNiagaraUIntVectorToFVector(FNiagaraUIntVectorSwizzle(rand, 2, 2, 2) & MGradientMask)) * MGradientScale - 1;
-	grad[0][1] = FVector3f::DotProduct(gvec[0][1], fv[1]);
-	grad[1][1] = FVector3f::DotProduct(gvec[1][1], fv[1]);
-	grad[2][1] = FVector3f::DotProduct(gvec[2][1], fv[1]);
+		float GradX = GvecX.Dot(RelVec);
+		float GradY = GvecY.Dot(RelVec);
+		float GradZ = GvecZ.Dot(RelVec);
 
-	fv[2] = v - T[2];
-	rand = Rand3DPCG16(NiagaraVectorFloorToInt(6 * T[2] + 0.5));
-	gvec[0][2] = FVector3f(FNiagaraUIntVectorToFVector(FNiagaraUIntVectorSwizzle(rand, 0, 0, 0) & MGradientMask)) * MGradientScale - 1;
-	gvec[1][2] = FVector3f(FNiagaraUIntVectorToFVector(FNiagaraUIntVectorSwizzle(rand, 1, 1, 1) & MGradientMask)) * MGradientScale - 1;
-	gvec[2][2] = FVector3f(FNiagaraUIntVectorToFVector(FNiagaraUIntVectorSwizzle(rand, 2, 2, 2) & MGradientMask)) * MGradientScale - 1;
-	grad[0][2] = FVector3f::DotProduct(gvec[0][2], fv[2]);
-	grad[1][2] = FVector3f::DotProduct(gvec[1][2], fv[2]);
-	grad[2][2] = FVector3f::DotProduct(gvec[2][2], fv[2]);
+		// Shared terms
+		constexpr float Scale = 1024.f / 375.f;
+		float SquaredDist = RelVec.Dot(RelVec);
+		float S = FMath::Clamp(SquaredDist + SquaredDist, 0.0f, 1.0f);
 
-	fv[3] = v - T[3];
-	rand = Rand3DPCG16(NiagaraVectorFloorToInt(6 * T[3] + 0.5));
-	gvec[0][3] = FVector3f(FNiagaraUIntVectorToFVector(FNiagaraUIntVectorSwizzle(rand, 0, 0, 0) & MGradientMask)) * MGradientScale - 1;
-	gvec[1][3] = FVector3f(FNiagaraUIntVectorToFVector(FNiagaraUIntVectorSwizzle(rand, 1, 1, 1) & MGradientMask)) * MGradientScale - 1;
-	gvec[2][3] = FVector3f(FNiagaraUIntVectorToFVector(FNiagaraUIntVectorSwizzle(rand, 2, 2, 2) & MGradientMask)) * MGradientScale - 1;
-	grad[0][3] = FVector3f::DotProduct(gvec[0][3], fv[3]);
-	grad[1][3] = FVector3f::DotProduct(gvec[1][3], fv[3]);
-	grad[2][3] = FVector3f::DotProduct(gvec[2][3], fv[3]);
+		// sv / SimplexSmooth in Random.ush
+		float SimplexValue = Scale * (1.0f + S * (-3.0f + S * (3.0f - S)));
 
-	FVector4f sv = SimplexSmooth(fv);
-	FNiagaraMatrix3x4 ds = SimplexDSmooth(fv);
+		// ds / SimplexDSmooth in Random.ush
+		float DerivScale = Scale * (-12.0f + S * (24.0f - S * 12.0f));
 
-	FNiagaraMatrix3x4 jacobian;
-	jacobian[0] = FVector4f(MulFVector4AndFNiagaraMatrix4x3(sv, gvec[0]) + MulFNiagaraMatrix3x4FAndVector4(ds, grad[0]), Dot4(sv, grad[0]));
-	jacobian[1] = FVector4f(MulFVector4AndFNiagaraMatrix4x3(sv, gvec[1]) + MulFNiagaraMatrix3x4FAndVector4(ds, grad[1]), Dot4(sv, grad[1]));
-	jacobian[2] = FVector4f(MulFVector4AndFNiagaraMatrix4x3(sv, gvec[2]) + MulFNiagaraMatrix3x4FAndVector4(ds, grad[2]), Dot4(sv, grad[2]));
+		Jacobian[0] += FVector4f(SimplexValue * GvecX + (DerivScale * GradX) * RelVec, SimplexValue * GradX);
+		Jacobian[1] += FVector4f(SimplexValue * GvecY + (DerivScale * GradY) * RelVec, SimplexValue * GradY);
+		Jacobian[2] += FVector4f(SimplexValue * GvecZ + (DerivScale * GradZ) * RelVec, SimplexValue * GradZ);
+	}
 
-	return jacobian;
+	return Jacobian;
 }
 
-#undef MGradientMask
-#undef MGradientScale 
-
