- Fix for toggling library only not working if library only was checked
- When gathering actions we should not filter as we filter the display menu, so adjusting the filter after gathering resulted in not change

diff --git a/Engine/Plugins/FX/Niagara/Source/NiagaraEditorWidgets/Private/Stack/SNiagaraStackFunctionInputValue.cpp b/Engine/Plugins/FX/Niagara/Source/NiagaraEditorWidgets/Private/Stack/SNiagaraStackFunctionInputValue.cpp
--- a/Engine/Plugins/FX/Niagara/Source/NiagaraEditorWidgets/Private/Stack/SNiagaraStackFunctionInputValue.cpp
+++ b/Engine/Plugins/FX/Niagara/Source/NiagaraEditorWidgets/Private/Stack/SNiagaraStackFunctionInputValue.cpp
@@ -1241,7 +1241,7 @@ TArray<TSharedPtr<FNiagaraMenuAction_Generic>> SNiagaraStackFunctionInputValue::
 	{
 		const FText CategoryName = LOCTEXT("DynamicInputValueCategory", "Dynamic Inputs");
 		TArray<UNiagaraScript*> DynamicInputScripts;
-		FunctionInput->GetAvailableDynamicInputs(DynamicInputScripts, bLibraryOnly == false);
+		FunctionInput->GetAvailableDynamicInputs(DynamicInputScripts, true);	// Note we do not filter here as we filter the displayed data
 
 		// we add scratch pad scripts here so we can check if an available dynamic input is a scratch pad script or asset based
 		TSet<UNiagaraScript*> ScratchPadDynamicInputs;
