- Disable tick passes as there's an issue with parameter bindings from early simulation release

#rnx

diff --git a/Engine/Plugins/FX/Niagara/Source/Niagara/Private/NiagaraWorldManager.cpp b/Engine/Plugins/FX/Niagara/Source/Niagara/Private/NiagaraWorldManager.cpp
--- a/Engine/Plugins/FX/Niagara/Source/Niagara/Private/NiagaraWorldManager.cpp
+++ b/Engine/Plugins/FX/Niagara/Source/Niagara/Private/NiagaraWorldManager.cpp
@@ -43,7 +43,7 @@ DECLARE_CYCLE_STAT(TEXT("Niagara Manager Wait Pre Garbage Collect [GT]"), STAT_N
 DECLARE_CYCLE_STAT(TEXT("Niagara Manager Refresh Owner Allows Scalability"), STAT_NiagaraWorldManRefreshOwnerAllowsScalability, STATGROUP_Niagara);
 DECLARE_CYCLE_STAT(TEXT("Niagara Manager Tick Parameter Collections [GT]"), STAT_NiagaraWorldManTickParamCollections, STATGROUP_Niagara);
 
-bool GNiagaraWorldManagerUseTickPasses = true;
+bool GNiagaraWorldManagerUseTickPasses = false;
 static FAutoConsoleVariableRef CVarNiagaraWorldManagerUseTickPasses(
 	TEXT("fx.Niagara.WorldManager.UseTickPasses"),
 	GNiagaraWorldManagerUseTickPasses,
@@ -404,6 +404,7 @@ void FNiagaraWorldManager::Init(UWorld* InWorld)
 			TickFunc.bHighPriority					= NiagaraWorldManagerInternal::GetNiagaraTickGroupPriority(TickGroup);
 			TickFunc.Owner							= this;
 			TickFunc.TickPass						= TickPass;
+			TickFunc.SetTickFunctionEnable(TickFunc.bStartWithTickEnabled);
 			TickFunc.RegisterTickFunction(InWorld->PersistentLevel);
 			if (iTickPass > 0)
 			{
