Added an option for data channel data interfaces to only process functions once per frame, to prevent erratic behavior in systems which tick multiple times per frame (e.g. systems with fixed tick delta timing).

#tests data channel engine tests are still green, validated new behavior with fixed tick delta systems that write and spawn multiple times per tick

diff --git a/Engine/Plugins/FX/Niagara/Source/Niagara/Internal/DataInterface/NiagaraDataInterfaceDataChannelRead.h b/Engine/Plugins/FX/Niagara/Source/Niagara/Internal/DataInterface/NiagaraDataInterfaceDataChannelRead.h
--- a/Engine/Plugins/FX/Niagara/Source/Niagara/Internal/DataInterface/NiagaraDataInterfaceDataChannelRead.h
+++ b/Engine/Plugins/FX/Niagara/Source/Niagara/Internal/DataInterface/NiagaraDataInterfaceDataChannelRead.h
@@ -126,6 +126,15 @@ public:
 	UPROPERTY(EditAnywhere, Category = "Spawning", AdvancedDisplay)
 	bool bOverrideSpawnGroupToDataChannelIndex = false;
 
+	/**
+	 * If a system is ticking multiple times in a single frame (for example due to fixed tick delta time or because the user is scrubbing the desired time in the timeline),
+	 * it is usually not desirable to also spawn new particles every tick for the data channel entries, as the data channel itself doesn't change.
+	 *
+	 * When enabled, the emitter will only spawn particles on the last subtick of the frame.
+	 */
+	UPROPERTY(EditAnywhere, Category = "Data Channel", AdvancedDisplay)
+	bool bOnlySpawnOnceOnSubticks = true;
+
 	//UObject Interface
 	NIAGARA_API virtual void PostInitProperties() override;
 	NIAGARA_API virtual void BeginDestroy() override;
diff --git a/Engine/Plugins/FX/Niagara/Source/Niagara/Internal/DataInterface/NiagaraDataInterfaceDataChannelWrite.h b/Engine/Plugins/FX/Niagara/Source/Niagara/Internal/DataInterface/NiagaraDataInterfaceDataChannelWrite.h
--- a/Engine/Plugins/FX/Niagara/Source/Niagara/Internal/DataInterface/NiagaraDataInterfaceDataChannelWrite.h
+++ b/Engine/Plugins/FX/Niagara/Source/Niagara/Internal/DataInterface/NiagaraDataInterfaceDataChannelWrite.h
@@ -147,6 +147,15 @@ public:
 	UPROPERTY(EditAnywhere, Category = "Data Channel", AdvancedDisplay)
 	bool bUpdateDestinationDataEveryTick = true;
 
+	/**
+	 * If a system is ticking multiple times in a single frame (for example due to fixed tick delta time or because the user is scrubbing the desired time in the timeline),
+	 * it is usually not desirable to also write data channel entries for each subtick, as consumers reading on the next frame will potentially see a huge number of entries.
+	 *
+	 * When enabled, the emitter will only execute write functions on the last subtick of the frame.
+	 */
+	UPROPERTY(EditAnywhere, Category = "Data Channel", AdvancedDisplay)
+	bool bOnlyWriteOnceOnSubticks = true;
+
 	/** The Data Channel to write to. */
 	UPROPERTY(EditAnywhere, Category = "Data Channel")
 	TObjectPtr<UNiagaraDataChannelAsset> Channel;
diff --git a/Engine/Plugins/FX/Niagara/Source/Niagara/Private/DataInterface/NiagaraDataInterfaceDataChannelRead.cpp b/Engine/Plugins/FX/Niagara/Source/Niagara/Private/DataInterface/NiagaraDataInterfaceDataChannelRead.cpp
--- a/Engine/Plugins/FX/Niagara/Source/Niagara/Private/DataInterface/NiagaraDataInterfaceDataChannelRead.cpp
+++ b/Engine/Plugins/FX/Niagara/Source/Niagara/Private/DataInterface/NiagaraDataInterfaceDataChannelRead.cpp
@@ -30,7 +30,6 @@
 #if WITH_EDITOR
 #include "INiagaraEditorOnlyDataUtlities.h"
 #include "Modules/ModuleManager.h"
-#include "NiagaraModule.h"
 #endif
 
 //////////////////////////////////////////////////////////////////////////
@@ -1158,7 +1157,8 @@ bool UNiagaraDataInterfaceDataChannelRead::Equals(const UNiagaraDataInterface* O
 			Channel == OtherTyped->Channel &&
 			bReadCurrentFrame == OtherTyped->bReadCurrentFrame &&
 			bUpdateSourceDataEveryTick == OtherTyped->bUpdateSourceDataEveryTick &&
-			bOverrideSpawnGroupToDataChannelIndex == OtherTyped->bOverrideSpawnGroupToDataChannelIndex)
+			bOverrideSpawnGroupToDataChannelIndex == OtherTyped->bOverrideSpawnGroupToDataChannelIndex &&
+			bOnlySpawnOnceOnSubticks == OtherTyped->bOnlySpawnOnceOnSubticks)
 		{
 			return true;
 		}
@@ -1183,7 +1183,8 @@ bool UNiagaraDataInterfaceDataChannelRead::CopyToInternal(UNiagaraDataInterface*
 		DestTyped->CompiledData = CompiledData;
 		DestTyped->bReadCurrentFrame = bReadCurrentFrame;
 		DestTyped->bUpdateSourceDataEveryTick = bUpdateSourceDataEveryTick;
-		DestTyped->bOverrideSpawnGroupToDataChannelIndex = bOverrideSpawnGroupToDataChannelIndex;		
+		DestTyped->bOverrideSpawnGroupToDataChannelIndex = bOverrideSpawnGroupToDataChannelIndex;
+		DestTyped->bOnlySpawnOnceOnSubticks = bOnlySpawnOnceOnSubticks;
 		return true;
 	}
 
@@ -1741,21 +1742,22 @@ void UNiagaraDataInterfaceDataChannelRead::SpawnConditional(FVectorVMExternalFun
 	FNDIVariadicInputHandler<16> VariadicInputs(Context, BindingInfo);//TODO: Make static / avoid allocation
 
 	const FNiagaraEmitterID EmitterID = InEmitterID.GetAndAdvance();
-	const int32 NumEmitters = SystemInstance->GetEmitters().Num();
 	FNiagaraEmitterInstance* EmitterInst = SystemInstance->GetEmitterByID(EmitterID);
 
 	FNiagaraDataChannelData* DataChannelData = InstData->DataChannelData.Get();
 
 	const bool bReadPrevFrame = bReadCurrentFrame == false || GNDCReadForcePrevFrame;
 	FNiagaraDataBuffer* Data = DataChannelData ? DataChannelData->GetCPUData(bReadPrevFrame) : nullptr;
-
-	const bool bSpawn = INiagaraModule::DataChannelsEnabled() && Data && Data->GetNumInstances() > 0 && EmitterInst && EmitterInst->IsActive() && InEnabled.GetAndAdvance();
+	const FNiagaraTickInfo& TickInfo = SystemInstance->GetSystemSimulation()->GetTickInfo();
+	bool bProcessCurrentTick = bOnlySpawnOnceOnSubticks ? TickInfo.TickNumber == TickInfo.TickCount - 1 : true;
+	
+	const bool bSpawn = INiagaraModule::DataChannelsEnabled() && Data && Data->GetNumInstances() > 0 && EmitterInst && EmitterInst->IsActive() && InEnabled.GetAndAdvance() && bProcessCurrentTick;
 	if(bSpawn)
 	{
 		FNDIRandomHelperFromStream RandHelper(Context);
 
-		ENDIDataChannelSpawnMode Mode = (ENDIDataChannelSpawnMode)InMode.GetAndAdvance();
-		ENiagaraConditionalOperator Op = (ENiagaraConditionalOperator)InOp.GetAndAdvance();
+		ENDIDataChannelSpawnMode Mode = static_cast<ENDIDataChannelSpawnMode>(InMode.GetAndAdvance());
+		ENiagaraConditionalOperator Op = static_cast<ENiagaraConditionalOperator>(InOp.GetAndAdvance());
 		int32 NumDataChannelInstances = Data->GetNumInstances();
 
 		//Is mode none or invalid?
@@ -1876,8 +1878,10 @@ void UNiagaraDataInterfaceDataChannelRead::SpawnDirect(FVectorVMExternalFunction
 
 	const bool bReadPrevFrame = bReadCurrentFrame == false || GNDCReadForcePrevFrame;
 	FNiagaraDataBuffer* Data = DataChannelData ? DataChannelData->GetCPUData(bReadPrevFrame) : nullptr;
-
-	const bool bApply = INiagaraModule::DataChannelsEnabled() && Data && Data->GetNumInstances() > 0 && EmitterInst && EmitterInst->IsActive() && InEnabled.GetAndAdvance();
+	const FNiagaraTickInfo& TickInfo = SystemInstance->GetSystemSimulation()->GetTickInfo();
+	bool bProcessCurrentTick = bOnlySpawnOnceOnSubticks ? TickInfo.TickNumber == TickInfo.TickCount - 1 : true;
+	
+	const bool bApply = INiagaraModule::DataChannelsEnabled() && Data && Data->GetNumInstances() > 0 && EmitterInst && EmitterInst->IsActive() && InEnabled.GetAndAdvance() && bProcessCurrentTick;
 	if (bApply)
 	{
 		FNDIRandomHelperFromStream RandHelper(Context);
diff --git a/Engine/Plugins/FX/Niagara/Source/Niagara/Private/DataInterface/NiagaraDataInterfaceDataChannelWrite.cpp b/Engine/Plugins/FX/Niagara/Source/Niagara/Private/DataInterface/NiagaraDataInterfaceDataChannelWrite.cpp
--- a/Engine/Plugins/FX/Niagara/Source/Niagara/Private/DataInterface/NiagaraDataInterfaceDataChannelWrite.cpp
+++ b/Engine/Plugins/FX/Niagara/Source/Niagara/Private/DataInterface/NiagaraDataInterfaceDataChannelWrite.cpp
@@ -14,9 +14,6 @@
 
 #include "NiagaraCompileHashVisitor.h"
 #include "ShaderCompilerCore.h"
-#include "SystemTextures.h"
-#include "RenderGraphUtils.h"
-#include "NiagaraModule.h"
 #include "NiagaraShaderParametersBuilder.h"
 
 #include "NiagaraDataChannel.h"
@@ -287,12 +284,16 @@ struct FNDIDataChannelWriteInstanceData
 	//Buffer we're currently writing into this frame.
 	FNiagaraDataBuffer* DestinationData = nullptr;
 
+	FNiagaraSystemInstance* Owner = nullptr;
+
 	~FNDIDataChannelWriteInstanceData()
 	{
 	}
 
 	bool Init(UNiagaraDataInterfaceDataChannelWrite* Interface, FNiagaraSystemInstance* Instance)
 	{
+		Owner = Instance;
+		
 		//In non test/shipping builds we gather and log and missing parameters that cause us to fail to find correct bindings.
 		TArray<FNiagaraVariableBase> MissingParams;
 
@@ -312,6 +313,7 @@ struct FNDIDataChannelWriteInstanceData
 		DynamicAllocationCount = 0;
 		AtomicNumInstances = 0;
 		DestinationData = nullptr;
+		Owner = Instance;
 		
 		LwcTile = Instance->GetLWCTile();
 		if (Interface->ShouldPublish())
@@ -762,6 +764,7 @@ bool UNiagaraDataInterfaceDataChannelWrite::Equals(const UNiagaraDataInterface*
 			bPublishToCPU == OtherTyped->bPublishToCPU &&
 			bPublishToGPU == OtherTyped->bPublishToGPU &&
 			Channel == OtherTyped->Channel &&
+			bOnlyWriteOnceOnSubticks == OtherTyped->bOnlyWriteOnceOnSubticks &&
 			bUpdateDestinationDataEveryTick == OtherTyped->bUpdateDestinationDataEveryTick)
 		{
 			return true;
@@ -1090,6 +1093,7 @@ bool UNiagaraDataInterfaceDataChannelWrite::CopyToInternal(UNiagaraDataInterface
 		DestTyped->Channel = Channel;
 		DestTyped->CompiledData = CompiledData;
 		DestTyped->bUpdateDestinationDataEveryTick = bUpdateDestinationDataEveryTick;
+		DestTyped->bOnlyWriteOnceOnSubticks = bOnlyWriteOnceOnSubticks;
 		return true;
 	}
 
@@ -1188,9 +1192,15 @@ void UNiagaraDataInterfaceDataChannelWrite::Write(FVectorVMExternalFunctionConte
 	FNDIOutputParam<FNiagaraBool> OutSuccess(Context);
 
 	std::atomic<uint32>& AtomicNumInstances = InstData->AtomicNumInstances;
+	bool bProcessCurrentTick = true;
+	if (bOnlyWriteOnceOnSubticks && InstData->Owner)
+	{
+		const FNiagaraTickInfo& TickInfo = InstData->Owner->GetSystemSimulation()->GetTickInfo();
+		bProcessCurrentTick = TickInfo.TickNumber == TickInfo.TickCount - 1;
+	}
 
 	bool bAllFailedFallback = true;
-	if (InstData->DestinationData && BindingInfo && INiagaraModule::DataChannelsEnabled())
+	if (InstData->DestinationData && BindingInfo && INiagaraModule::DataChannelsEnabled() && bProcessCurrentTick)
 	{
 		if(FNiagaraDataBuffer* Data = InstData->DestinationData)
 		{			
@@ -1298,8 +1308,15 @@ void UNiagaraDataInterfaceDataChannelWrite::Append(FVectorVMExternalFunctionCont
 
 	std::atomic<uint32>& AtomicNumInstances = InstData->AtomicNumInstances;
 
+	bool bProcessCurrentTick = true;
+	if (bOnlyWriteOnceOnSubticks && InstData->Owner)
+	{
+		const FNiagaraTickInfo& TickInfo = InstData->Owner->GetSystemSimulation()->GetTickInfo();
+		bProcessCurrentTick = TickInfo.TickNumber == TickInfo.TickCount - 1;
+	}
+
 	bool bAllFailedFallback = true;
-	if (InstData->DestinationData && BindingInfo &&  INiagaraModule::DataChannelsEnabled())
+	if (InstData->DestinationData && BindingInfo && INiagaraModule::DataChannelsEnabled() && bProcessCurrentTick)
 	{
 		if(FNiagaraDataBuffer* Data = InstData->DestinationData)
 		{
diff --git a/Engine/Plugins/FX/Niagara/Source/Niagara/Private/NiagaraComponent.cpp b/Engine/Plugins/FX/Niagara/Source/Niagara/Private/NiagaraComponent.cpp
--- a/Engine/Plugins/FX/Niagara/Source/Niagara/Private/NiagaraComponent.cpp
+++ b/Engine/Plugins/FX/Niagara/Source/Niagara/Private/NiagaraComponent.cpp
@@ -947,8 +947,7 @@ void UNiagaraComponent::TickComponent(float DeltaSeconds, enum ELevelTick TickTy
 					AgeDiff = DesiredAge - SystemInstanceController->GetAge();
 				}
 
-				FNiagaraSystemSimulation* SystemSim = SystemInstanceController->GetSoloSystemSimulation().Get();
-				if (SystemSim)
+				if (SystemInstanceController->GetSoloSystemSimulation().IsValid())
 				{
 					FGraphEventRef CompletionEvent = (GNiagaraAllowDesiredAgeAsync && ThisTickFunction && ThisTickFunction->IsCompletionHandleValid()) ? ThisTickFunction->GetCompletionHandle() : nullptr;
 
@@ -957,33 +956,10 @@ void UNiagaraComponent::TickComponent(float DeltaSeconds, enum ELevelTick TickTy
 					if (bLockDesiredAgeDeltaTimeToSeekDelta || AgeDiff > TickDelta)
 					{
 						// If we're locking the delta time to the seek delta, or we need to seek more than a frame, tick the simulation by the seek delta.
-						const bool bUseMaxSimTime = MaxSimTime > 0.0f;
-						const double EndMaxSimTime = bUseMaxSimTime ? FPlatformTime::Seconds() + MaxSimTime : 0.0f;
-
 						TicksToProcess = FMath::FloorToInt(AgeDiff / TickDelta);
 						CompletionEvent = TicksToProcess == 1 ? CompletionEvent : nullptr;
 
-						while ( TicksToProcess > 0 )
-						{
-							// Cannot do multiple tick off the game thread here without additional work. So we pass in null for the completion event which will force GT execution.
-							// If this becomes a perf problem I can add a new path for the tick code to handle multiple ticks.
-							SystemInstanceController->ManualTick(TickDelta, CompletionEvent);
-							if (IsComplete())
-							{
-								break;
-							}
-							--TicksToProcess;
-
-							// This limits the amount of time we will consume seeking forward
-							if ( bUseMaxSimTime )
-							{
-								const double CurrentTime = FPlatformTime::Seconds();
-								if ( CurrentTime >= EndMaxSimTime )
-								{
-									break;
-								}
-							}
-						}
+						SystemInstanceController->MultiTick(TicksToProcess, TickDelta, MaxSimTime, CompletionEvent);
 					}
 					else
 					{
diff --git a/Engine/Plugins/FX/Niagara/Source/Niagara/Private/NiagaraSystemInstance.cpp b/Engine/Plugins/FX/Niagara/Source/Niagara/Private/NiagaraSystemInstance.cpp
--- a/Engine/Plugins/FX/Niagara/Source/Niagara/Private/NiagaraSystemInstance.cpp
+++ b/Engine/Plugins/FX/Niagara/Source/Niagara/Private/NiagaraSystemInstance.cpp
@@ -2273,6 +2273,26 @@ void FNiagaraSystemInstance::ManualTick(float DeltaSeconds, const FGraphEventRef
 	SystemSim->Tick_GameThread(DeltaSeconds, MyCompletionGraphEvent);
 }
 
+void FNiagaraSystemInstance::MultiTick(int32 TicksToProcess, float DeltaSecondsPerTick, float MaxSimTime, const FGraphEventRef& MyCompletionGraphEvent)
+{
+	SCOPE_CYCLE_COUNTER(STAT_NiagaraOverview_GT);
+	SCOPE_CYCLE_COUNTER(STAT_NiagaraSystemInst_ComponentTickGT);
+	CSV_SCOPED_TIMING_STAT_EXCLUSIVE(Effects);
+	LLM_SCOPE(ELLMTag::Niagara);
+
+	if (IsDisabled())
+	{
+		return;
+	}
+
+	TSharedPtr<FNiagaraSystemSimulation> SystemSim = GetSystemSimulation();
+	check(SystemSim.IsValid());
+	check(IsInGameThread());
+	check(bSolo);
+
+	SystemSim->MultiTick_GameThread(TicksToProcess, DeltaSecondsPerTick, MaxSimTime, MyCompletionGraphEvent);
+}
+
 void FNiagaraSystemInstance::SimCacheTick_GameThread(UNiagaraSimCache* SimCache, float DesiredAge, float DeltaSeconds, const FGraphEventRef& MyCompletionGraphEvent)
 {
 	SCOPE_CYCLE_COUNTER(STAT_NiagaraSystemInst_ComponentTickGT);
diff --git a/Engine/Plugins/FX/Niagara/Source/Niagara/Private/NiagaraSystemSimulation.cpp b/Engine/Plugins/FX/Niagara/Source/Niagara/Private/NiagaraSystemSimulation.cpp
--- a/Engine/Plugins/FX/Niagara/Source/Niagara/Private/NiagaraSystemSimulation.cpp
+++ b/Engine/Plugins/FX/Niagara/Source/Niagara/Private/NiagaraSystemSimulation.cpp
@@ -1112,6 +1112,48 @@ void FNiagaraSystemSimulation::Tick_GameThread(float DeltaSeconds, const FGraphE
 	}
 }
 
+void FNiagaraSystemSimulation::MultiTick_GameThread(int32 TicksToProcess, float DeltaSecondsPerTick, float MaxSimTime, const FGraphEventRef& MyCompletionGraphEvent)
+{
+	if (DeltaSecondsPerTick <= 0 || TicksToProcess <= 0)
+	{
+		return;
+	}
+
+	TickInfo.UsesFixedTick = true;
+	TickInfo.EngineTick = TicksToProcess * DeltaSecondsPerTick;
+	TickInfo.SystemTick = DeltaSecondsPerTick;
+	TickInfo.TickCount = TicksToProcess;
+	TickInfo.TickNumber = -1;
+	TickInfo.TimeStepFraction = 0;
+	
+	const bool bUseMaxSimTime = MaxSimTime > 0.0f;
+	const double EndMaxSimTime = bUseMaxSimTime ? FPlatformTime::Seconds() + MaxSimTime : 0.0f;
+		
+	for (int i = 0; i < TicksToProcess; i++)
+	{
+		TickInfo.TickNumber = i;
+		TickInfo.TimeStepFraction = 1.0f * (i + 1) / TicksToProcess;
+
+		//TODO: allow passing in MyCompletionGraphEvent to tick off the game thread here
+		Tick_GameThread_Internal(DeltaSecondsPerTick, nullptr);
+
+		if (System->HasFixedTickDelta())
+		{
+			FixedDeltaTickAge += DeltaSecondsPerTick;
+		}
+
+		// This limits the amount of time we will consume seeking forward
+		if (bUseMaxSimTime)
+		{
+			const double CurrentTime = FPlatformTime::Seconds();
+			if ( CurrentTime >= EndMaxSimTime )
+			{
+				break;
+			}
+		}
+	}
+}
+
 void FNiagaraSystemSimulation::Tick_GameThread_Internal(float DeltaSeconds, const FGraphEventRef& MyCompletionGraphEvent)
 {
 	TArray<FNiagaraSystemInstance*>& SystemInstances = GetSystemInstances(ENiagaraSystemInstanceState::Running);
diff --git a/Engine/Plugins/FX/Niagara/Source/Niagara/Public/NiagaraSystemInstance.h b/Engine/Plugins/FX/Niagara/Source/Niagara/Public/NiagaraSystemInstance.h
--- a/Engine/Plugins/FX/Niagara/Source/Niagara/Public/NiagaraSystemInstance.h
+++ b/Engine/Plugins/FX/Niagara/Source/Niagara/Public/NiagaraSystemInstance.h
@@ -194,8 +194,16 @@ public:
 	/** Requests the the simulation be reset on the next tick. */
 	NIAGARA_API void Reset(EResetMode Mode);
 
+	/** Manually advances the simulation. Ticks on the game thread. */
 	NIAGARA_API void ManualTick(float DeltaSeconds, const FGraphEventRef& MyCompletionGraphEvent);
 
+	/** Manually advances the simulation by multiple ticks. Ticks on the game thread.
+	 *  If MaxSimTime >= 0, then unprocessed ticks will be discarded after the time is elapsed to prevent stalling the game thread.
+	 *
+	 *  Note that this overrides a system's fixed tick delta time (if it has one) and uses DeltaSecondsPerTick instead. 
+	 */
+	NIAGARA_API void MultiTick(int32 TicksToProcess, float DeltaSecondsPerTick, float MaxSimTime, const FGraphEventRef& MyCompletionGraphEvent);
+
 	/** Ticks the system using the a SimCache. */
 	NIAGARA_API void SimCacheTick_GameThread(UNiagaraSimCache* SimCache, float DesiredAge, float DeltaSeconds, const FGraphEventRef& MyCompletionGraphEvent);
 	/** Concurrent work for SimCache tick */
diff --git a/Engine/Plugins/FX/Niagara/Source/Niagara/Public/NiagaraSystemInstanceController.h b/Engine/Plugins/FX/Niagara/Source/Niagara/Public/NiagaraSystemInstanceController.h
--- a/Engine/Plugins/FX/Niagara/Source/Niagara/Public/NiagaraSystemInstanceController.h
+++ b/Engine/Plugins/FX/Niagara/Source/Niagara/Public/NiagaraSystemInstanceController.h
@@ -134,6 +134,7 @@ public:
 	NIAGARA_SYSTEM_INSTANCE_CONTROLLER_SHIM(IsPendingSpawn, const)
 	NIAGARA_SYSTEM_INSTANCE_CONTROLLER_SHIM(Reset,)
 	NIAGARA_SYSTEM_INSTANCE_CONTROLLER_SHIM(ManualTick,)
+	NIAGARA_SYSTEM_INSTANCE_CONTROLLER_SHIM(MultiTick,)
 	NIAGARA_SYSTEM_INSTANCE_CONTROLLER_SHIM(SimCacheTick_GameThread,)
 	NIAGARA_SYSTEM_INSTANCE_CONTROLLER_SHIM(SetTickBehavior,)
 	NIAGARA_SYSTEM_INSTANCE_CONTROLLER_SHIM(GetTickBehavior,)
diff --git a/Engine/Plugins/FX/Niagara/Source/Niagara/Public/NiagaraSystemSimulation.h b/Engine/Plugins/FX/Niagara/Source/Niagara/Public/NiagaraSystemSimulation.h
--- a/Engine/Plugins/FX/Niagara/Source/Niagara/Public/NiagaraSystemSimulation.h
+++ b/Engine/Plugins/FX/Niagara/Source/Niagara/Public/NiagaraSystemSimulation.h
@@ -264,6 +264,9 @@ public:
 	/** Second phase of system sim tick that can run on any thread. */
 	void Tick_Concurrent(FNiagaraSystemSimulationTickContext& Context);
 
+	/** Processes multiple ticks at once on the game thread. */
+	void MultiTick_GameThread(int32 TicksToProcess, float DeltaSecondsPerTick, float MaxSimTime, const FGraphEventRef& MyCompletionGraphEvent);
+
 	/** Update TickGroups for pending instances and execute tick group promotions. */
 	void UpdateTickGroups_GameThread();
 	/** Spawn any pending instances, assumes that you have update tick groups ahead of time. */
