Fixup some NiagaraDataInterface non-determinism.  Cleans up when we clean up the RF_Public flag on load (moved to PostLoad since the Outer isn't necessarily loaded during Serialize) and mucks around with how we clear out the contents of the CachedDefaultDataInterfaces so that we can reliably get deterministic names for the generated DI when we apply compilation results.
#tests validated cooks in Lyra with iterative cooking and comparing loose file contents

diff --git a/Engine/Plugins/FX/Niagara/Source/Niagara/Classes/NiagaraDataInterface.h b/Engine/Plugins/FX/Niagara/Source/Niagara/Classes/NiagaraDataInterface.h
--- a/Engine/Plugins/FX/Niagara/Source/Niagara/Classes/NiagaraDataInterface.h
+++ b/Engine/Plugins/FX/Niagara/Source/Niagara/Classes/NiagaraDataInterface.h
@@ -579,7 +579,7 @@ public:
 #if WITH_EDITOR
 	NIAGARA_API virtual void PostEditChangeProperty(struct FPropertyChangedEvent& PropertyChangedEvent) override;
 #endif
-	NIAGARA_API virtual void Serialize(FStructuredArchive::FRecord Record) override;
+	NIAGARA_API virtual void PostLoad() override;
 	// UObject Interface END
 
 #if WITH_EDITOR
diff --git a/Engine/Plugins/FX/Niagara/Source/Niagara/Classes/NiagaraScript.h b/Engine/Plugins/FX/Niagara/Source/Niagara/Classes/NiagaraScript.h
--- a/Engine/Plugins/FX/Niagara/Source/Niagara/Classes/NiagaraScript.h
+++ b/Engine/Plugins/FX/Niagara/Source/Niagara/Classes/NiagaraScript.h
@@ -1513,7 +1513,25 @@ private:
 	private :
 
 #if WITH_EDITORONLY_DATA
-		void ComputeVMCompilationId_EmitterShared(FNiagaraVMExecutableDataId& Id, const FVersionedNiagaraEmitter& Emitter, ENiagaraRendererSourceDataMode InSourceMode) const;
+	void ComputeVMCompilationId_EmitterShared(FNiagaraVMExecutableDataId& Id, const FVersionedNiagaraEmitter& Emitter, ENiagaraRendererSourceDataMode InSourceMode) const;
+
+	// These data interfaces are populated via InvalidateAndEmptyDefaultDataInterfaces().  This helps avoid naming conflicts of the DataInterfaces stored within
+	// CachedDefaultDataInterfaces.  When we are updating the compiled data we end up creating new DataInterfaces but we want to ensure that they are deterministically
+	// named, and simply having the original DI become unreferenced leaves us in a state where a non-deterministic GC during the cook process could free the DI, but
+	// is not guaranteed.  We resolve this by renaming our invalidated DI, thus creating room for the newly created DI to be named properly.  The last caveat is that
+	// we have to be careful about when we rename the DI as the package containing the DI may get reloaded, which would recreate the DI (with the original name) leaving
+	// us, once again, in a state where the named slot is taken.
+	UPROPERTY(Transient)
+	TArray<TObjectPtr<UNiagaraDataInterface>> DataInterfacesPendingInvalidation;
+
+	enum EDataInterfaceInvalidationMode
+	{
+		QueueInvaliation,	// data interfaces will be queued up to invalidate in the future
+		ProcessInvaliation, // immediately invalidates any new or previously queued data interfaces
+	};
+
+	// Empties CachedDefaultDataInterfaces and moves the referenced DI to DataInterfacesPendingInvalidation.
+	void InvalidateAndEmptyDefaultDataInterfaces(int32 SlackCount, EDataInterfaceInvalidationMode InvalidationMode);
 #endif
 };
 
diff --git a/Engine/Plugins/FX/Niagara/Source/Niagara/Private/NiagaraDataInterface.cpp b/Engine/Plugins/FX/Niagara/Source/Niagara/Private/NiagaraDataInterface.cpp
--- a/Engine/Plugins/FX/Niagara/Source/Niagara/Private/NiagaraDataInterface.cpp
+++ b/Engine/Plugins/FX/Niagara/Source/Niagara/Private/NiagaraDataInterface.cpp
@@ -171,20 +171,17 @@ void UNiagaraDataInterface::PostEditChangeProperty(struct FPropertyChangedEvent&
 }
 #endif
 
-void UNiagaraDataInterface::Serialize(FStructuredArchive::FRecord Record)
+void UNiagaraDataInterface::PostLoad()
 {
-	Super::Serialize(Record);
+	Super::PostLoad();
 #if WITH_EDITORONLY_DATA
-	if (Record.GetUnderlyingArchive().IsLoading())
+	if (BuildObjectFlagsForOwner(GetOuter()) & RF_Public)
 	{
-		if (BuildObjectFlagsForOwner(GetOuter()) & RF_Public)
-		{
-			SetFlags(RF_Public);
-		}
-		else
-		{
-			ClearFlags(RF_Public);
-		}
+		SetFlags(RF_Public);
+	}
+	else
+	{
+		ClearFlags(RF_Public);
 	}
 #endif
 }
diff --git a/Engine/Plugins/FX/Niagara/Source/Niagara/Private/NiagaraScript.cpp b/Engine/Plugins/FX/Niagara/Source/Niagara/Private/NiagaraScript.cpp
--- a/Engine/Plugins/FX/Niagara/Source/Niagara/Private/NiagaraScript.cpp
+++ b/Engine/Plugins/FX/Niagara/Source/Niagara/Private/NiagaraScript.cpp
@@ -1587,6 +1587,49 @@ void UNiagaraScript::ComputeVMCompilationId_EmitterShared(FNiagaraVMExecutableDa
 	Id.AdditionalVariables.Sort([](const FNiagaraVariableBase& A, const FNiagaraVariableBase& B) { return A.GetName().LexicalLess(B.GetName()); });
 }
 
+void UNiagaraScript::InvalidateAndEmptyDefaultDataInterfaces(int32 SlackCount, EDataInterfaceInvalidationMode InvalidationMode)
+{
+	DataInterfacesPendingInvalidation.Reserve(DataInterfacesPendingInvalidation.Num() + CachedDefaultDataInterfaces.Num());
+
+	for (FNiagaraScriptDataInterfaceInfo& DefaultDataInterface : CachedDefaultDataInterfaces)
+	{
+		if (DefaultDataInterface.DataInterface)
+		{
+			if (DefaultDataInterface.DataInterface->GetOuter() == this)
+			{
+				DataInterfacesPendingInvalidation.Add(DefaultDataInterface.DataInterface);
+			}
+		}
+	}
+
+	CachedDefaultDataInterfaces.Empty(SlackCount);
+
+	if (InvalidationMode == EDataInterfaceInvalidationMode::ProcessInvaliation)
+	{
+		const ERenameFlags RenameFlags = REN_NonTransactional | REN_DoNotDirty;
+
+		for (UNiagaraDataInterface* DataInterface : DataInterfacesPendingInvalidation)
+		{
+			const FName OrigName = DataInterface->GetFName();
+
+			FNameBuilder InvalidatedName;
+			InvalidatedName << TEXT("Invalidated_");
+			OrigName.AppendString(InvalidatedName);
+
+			FName NewName(InvalidatedName.ToView());
+
+			while (!DataInterface->Rename(*NewName.ToString(), nullptr, (RenameFlags | REN_Test)))
+			{
+				NewName.SetNumber(NewName.GetNumber() + 1);
+			}
+
+			DataInterface->Rename(*NewName.ToString(), nullptr, RenameFlags);
+		}
+
+		DataInterfacesPendingInvalidation.Reset();
+	}
+}
+
 #endif
 
 bool UNiagaraScript::ContainsUsage(ENiagaraScriptUsage InUsage) const
@@ -3064,8 +3107,7 @@ void UNiagaraScript::SetVMCompilationResults(const FNiagaraVMExecutableDataId& I
 	}
 
 	ResolveParameterCollectionReferences();
-
-	CachedDefaultDataInterfaces.Empty(CachedScriptVM.DataInterfaceInfo.Num());
+	InvalidateAndEmptyDefaultDataInterfaces(CachedScriptVM.DataInterfaceInfo.Num(), EDataInterfaceInvalidationMode::ProcessInvaliation);
 	ResolvedDataInterfaces.Empty();
 	for (const FNiagaraScriptDataInterfaceCompileInfo& Info : CachedScriptVM.DataInterfaceInfo)
 	{
@@ -4229,7 +4271,7 @@ bool UNiagaraScript::SynchronizeExecutablesWithCompilation(const UNiagaraScript*
 		CachedScriptVM.OptimizationTask.State = nullptr;
 		CachedParameterCollectionReferences = Script->CachedParameterCollectionReferences;
 		CachedDefaultUObjects = Script->CachedDefaultUObjects;
-		CachedDefaultDataInterfaces.Empty();
+		InvalidateAndEmptyDefaultDataInterfaces(Script->CachedDefaultDataInterfaces.Num(), EDataInterfaceInvalidationMode::ProcessInvaliation);
 		ResolvedDataInterfaces.Empty();
 		for (const FNiagaraScriptDataInterfaceInfo& Info : Script->CachedDefaultDataInterfaces)
 		{
@@ -4270,7 +4312,7 @@ void UNiagaraScript::InvalidateCompileResults(const FString& Reason)
 		GetLastGeneratedVMId().Invalidate();
 	}
 	CachedDefaultUObjects.Reset();
-	CachedDefaultDataInterfaces.Reset();
+	InvalidateAndEmptyDefaultDataInterfaces(0, EDataInterfaceInvalidationMode::QueueInvaliation);
 	ResolvedDataInterfaces.Reset();
 	CachedScriptVM.OptimizationTask.State = nullptr;
 }
