- Update the bone transforms if the number of expected transforms does not match
- This fixes an issue where we don't require current frame data (so do not wait on a tick) and change which SKM the component is using

diff --git a/Engine/Plugins/FX/Niagara/Source/Niagara/Private/NiagaraDataInterfaceSkeletalMesh.cpp b/Engine/Plugins/FX/Niagara/Source/Niagara/Private/NiagaraDataInterfaceSkeletalMesh.cpp
--- a/Engine/Plugins/FX/Niagara/Source/Niagara/Private/NiagaraDataInterfaceSkeletalMesh.cpp
+++ b/Engine/Plugins/FX/Niagara/Source/Niagara/Private/NiagaraDataInterfaceSkeletalMesh.cpp
@@ -337,7 +337,8 @@ void FSkeletalMeshSkinningData::RegisterUser(FSkeletalMeshSkinningDataUsage Usag
 	if (bNeedsDataImmediately)
 	{
 		check(IsInGameThread());
-		if (CurrBoneRefToLocals().Num() == 0 || CurrComponentTransforms().Num() == 0 )
+		const int32 NumExpectedTransforms = SkelMesh->GetRefBasesInvMatrix().Num();
+		if (CurrBoneRefToLocals().Num() != NumExpectedTransforms || CurrComponentTransforms().Num() != NumExpectedTransforms)
 		{
 			UpdateBoneTransforms();
 		}
