- Fix transition error if using a simple counter and it doesn't execute the readback on the final stage

diff --git a/Engine/Plugins/FX/Niagara/Source/Niagara/Private/DataInterface/NiagaraDataInterfaceSimpleCounter.cpp b/Engine/Plugins/FX/Niagara/Source/Niagara/Private/DataInterface/NiagaraDataInterfaceSimpleCounter.cpp
--- a/Engine/Plugins/FX/Niagara/Source/Niagara/Private/DataInterface/NiagaraDataInterfaceSimpleCounter.cpp
+++ b/Engine/Plugins/FX/Niagara/Source/Niagara/Private/DataInterface/NiagaraDataInterfaceSimpleCounter.cpp
@@ -104,14 +104,18 @@ struct FNDISimpleCounterProxy : public FNiagaraDataInterfaceProxyRW
 			if (FNDISimpleCounterInstanceData_RenderThread* InstanceData = PerInstanceData_RenderThread.Find(SystemInstanceID))
 			{
 				//-TODO:RDG: This should be done using the graph builder
-				const FNiagaraGPUInstanceCountManager& CountManager = Context.GetComputeDispatchInterface().GetGPUInstanceCounterManager();
-				FNiagaraGpuReadbackManager* ReadbackManager = Context.GetComputeDispatchInterface().GetGpuReadbackManager();
+				const FNiagaraGpuComputeDispatch& ComputeDispatch = static_cast<const FNiagaraGpuComputeDispatch&>(Context.GetComputeDispatchInterface());
+				const FNiagaraGPUInstanceCountManager& CountManager = ComputeDispatch.GetGPUInstanceCounterManager();
+				FNiagaraGpuReadbackManager* ReadbackManager = ComputeDispatch.GetGpuReadbackManager();
+
+				const ERHIAccess CountRHIAccess = ComputeDispatch.IsExecutingLastDispatchGroup() ? FNiagaraGPUInstanceCountManager::kCountBufferDefaultState : ERHIAccess::UAVCompute;
 
 				AddPass(
 					Context.GetGraphBuilder(),
 					RDG_EVENT_NAME("NiagaraSimpleCounter::PostStage"),
-					[InstanceData, SystemInstanceID, ReadbackManager, CountBuffer=CountManager.GetInstanceCountBuffer().Buffer, Proxy=this](FRHICommandListImmediate& RHICmdList)
+					[InstanceData, SystemInstanceID, ReadbackManager, CountBuffer=CountManager.GetInstanceCountBuffer().Buffer, CountRHIAccess, Proxy=this](FRHICommandListImmediate& RHICmdList)
 					{
+						RHICmdList.Transition(FRHITransitionInfo(CountBuffer, CountRHIAccess, ERHIAccess::CopySrc));
 						ReadbackManager->EnqueueReadback(
 							RHICmdList,
 							CountBuffer,
@@ -137,6 +141,7 @@ struct FNDISimpleCounterProxy : public FNiagaraDataInterfaceProxyRW
 								);
 							}
 						);
+						RHICmdList.Transition(FRHITransitionInfo(CountBuffer, ERHIAccess::CopySrc, CountRHIAccess));
 					}
 				);
 			}
diff --git a/Engine/Plugins/FX/Niagara/Source/Niagara/Private/NiagaraGpuComputeDispatch.cpp b/Engine/Plugins/FX/Niagara/Source/Niagara/Private/NiagaraGpuComputeDispatch.cpp
--- a/Engine/Plugins/FX/Niagara/Source/Niagara/Private/NiagaraGpuComputeDispatch.cpp
+++ b/Engine/Plugins/FX/Niagara/Source/Niagara/Private/NiagaraGpuComputeDispatch.cpp
@@ -1167,6 +1167,7 @@ void FNiagaraGpuComputeDispatch::ExecuteTicks(FRDGBuilder& GraphBuilder, TConstS
 		const bool bIsLastGroup = &DispatchGroup == &DispatchList.DispatchGroups.Last();
 
 		bIsExecutingFirstDispatchGroup = bIsFirstGroup;
+		bIsExecutingLastDispatchGroup = bIsLastGroup;
 
 		// Consume per tick data from the game thread
 		//-TODO: This does not work currently as some senders assume the data will not be deferred processed
diff --git a/Engine/Plugins/FX/Niagara/Source/Niagara/Private/NiagaraGpuComputeDispatch.h b/Engine/Plugins/FX/Niagara/Source/Niagara/Private/NiagaraGpuComputeDispatch.h
--- a/Engine/Plugins/FX/Niagara/Source/Niagara/Private/NiagaraGpuComputeDispatch.h
+++ b/Engine/Plugins/FX/Niagara/Source/Niagara/Private/NiagaraGpuComputeDispatch.h
@@ -182,7 +182,9 @@ public:
 
 	//-TODO: Temporary while the count buffer is not an RDG resource
 	bool IsExecutingFirstDispatchGroup() const { return bIsExecutingFirstDispatchGroup; }
+	bool IsExecutingLastDispatchGroup() const { return bIsExecutingLastDispatchGroup; }
 	bool bIsExecutingFirstDispatchGroup = false;
+	bool bIsExecutingLastDispatchGroup = false;
 
 private:
 	void DumpDebugFrame();
