- Add "fx.Niagara.SystemSimulation.TickBatchMode" to control how we want to generate instance tick tasks
- Default change to always process the last batch inline with the system simulation task
- This can be a small win in some situations and particularly on platforms with lower core count where the additional task spawned might not be processed immediatly resulting in stalls on the GT

diff --git a/Engine/Plugins/FX/Niagara/Source/Niagara/Private/NiagaraSystemSimulation.cpp b/Engine/Plugins/FX/Niagara/Source/Niagara/Private/NiagaraSystemSimulation.cpp
--- a/Engine/Plugins/FX/Niagara/Source/Niagara/Private/NiagaraSystemSimulation.cpp
+++ b/Engine/Plugins/FX/Niagara/Source/Niagara/Private/NiagaraSystemSimulation.cpp
@@ -92,13 +92,24 @@ static FAutoConsoleVariableRef CVarNiagaraSystemSimulationTaskStallTimeout(
 );
 
 static int32 GNiagaraSystemSimulationTickBatchSize = NiagaraSystemTickBatchSize;
-static FAutoConsoleVariableRef CVarParallelSystemInstanceTickBatchSize(
+static FAutoConsoleVariableRef CVarNiagaraSystemSimulationTickBatchSize(
 	TEXT("fx.Niagara.SystemSimulation.TickBatchSize"),
 	GNiagaraSystemSimulationTickBatchSize,
 	TEXT("The number of system instances to process per async task. \n"),
 	ECVF_Default
 );
 
+static int32 GNiagaraSystemSimulationTickBatchMode = 1;
+static FAutoConsoleVariableRef CVarNiagaraSystemSimulationTickBatchMode(
+	TEXT("fx.Niagara.SystemSimulation.TickBatchMode"),
+	GNiagaraSystemSimulationTickBatchMode,
+	TEXT("Controls what we do with each batch of instances in a simulation")
+	TEXT("0 - Spawn tasks for all batches.")
+	TEXT("1 - Spawn tasks for all but the final batch.")
+	TEXT("2 - Run all instance inline with the main simulation task."),
+	ECVF_Default
+);
+
 static int32 GNiagaraSystemSimulationConcurrentGPUTickInit = 1;
 static FAutoConsoleVariableRef CVarNiagaraConcurrentGPUTickInit(
 	TEXT("fx.Niagara.SystemSimulation.ConcurrentGPUTickInit"),
@@ -1006,16 +1017,16 @@ void FNiagaraSystemSimulation::SetInstanceState(FNiagaraSystemInstance* Instance
 	}
 }
 
-void FNiagaraSystemSimulation::AddSystemToTickBatch(FNiagaraSystemInstance* Instance, FNiagaraSystemSimulationTickContext& Context)
+void FNiagaraSystemSimulation::AddSystemToTickBatch(FNiagaraSystemSimulationTickContext& Context, FNiagaraSystemInstance* Instance, bool bIsLastInstance)
 {
 	Context.TickBatch.Add(Instance);
-	if (Context.TickBatch.Num() == GNiagaraSystemSimulationTickBatchSize)
+	if (bIsLastInstance || Context.TickBatch.Num() >= GNiagaraSystemSimulationTickBatchSize)
 	{
-		FlushTickBatch(Context);
+		FlushTickBatch(Context, bIsLastInstance);
 	}
 }
 
-void FNiagaraSystemSimulation::FlushTickBatch(FNiagaraSystemSimulationTickContext& Context)
+void FNiagaraSystemSimulation::FlushTickBatch(FNiagaraSystemSimulationTickContext& Context, bool bIsLastInstance)
 {
 	if (Context.TickBatch.Num() == 0)
 	{
@@ -1025,8 +1036,17 @@ void FNiagaraSystemSimulation::FlushTickBatch(FNiagaraSystemSimulationTickContex
 	// If we are running async create tasks to execute
 	if ( Context.IsRunningAsync() )
 	{
+		// We execute the instance tick immediately if...
+		// - We don't have a system script (stateless path)
+		// - Mode is set to run last batch inlin
+		// - Mode is set to run all in line
+		const bool bExecuteInstanceTickImmediate =
+			(bRunUpdateScript == false) ||
+			((GNiagaraSystemSimulationTickBatchMode == 1) && bIsLastInstance) ||
+			(GNiagaraSystemSimulationTickBatchMode == 2);
+
 		FGraphEventArray FinalizePrereqArray;
-		if (bRunUpdateScript == false)
+		if (bExecuteInstanceTickImmediate)
 		{
 			for (FNiagaraSystemInstance* Inst : Context.TickBatch)
 			{
@@ -1793,7 +1813,8 @@ void FNiagaraSystemSimulation::Tick_Concurrent(FNiagaraSystemSimulationTickConte
 			}
 		}
 
-		for ( int32 i=0; i < Context.Instances.Num(); ++i )
+
+		for ( int32 i=0; i < NumInstances; ++i )
 		{
 			FNiagaraSystemInstance* SystemInstance = Context.Instances[i];
 
@@ -1805,9 +1826,8 @@ void FNiagaraSystemSimulation::Tick_Concurrent(FNiagaraSystemSimulationTickConte
 
 			SystemInstance->TickSystemState();
 
-			AddSystemToTickBatch(SystemInstance, Context);
+			AddSystemToTickBatch(Context, SystemInstance, i == NumInstances - 1);
 		}
-		FlushTickBatch(Context);
 	}
 	else
 	{
@@ -1827,11 +1847,12 @@ void FNiagaraSystemSimulation::Tick_Concurrent(FNiagaraSystemSimulationTickConte
 
 		TransferSystemSimResults(Context);
 
-		for (FNiagaraSystemInstance* Instance : Context.Instances)
+		const int32 NumInstances = Context.Instances.Num();
+		for (int32 i = 0; i < NumInstances; ++i)
 		{
-			AddSystemToTickBatch(Instance, Context);
+			FNiagaraSystemInstance* SystemInstance = Context.Instances[i];
+			AddSystemToTickBatch(Context, SystemInstance, i == NumInstances - 1);
 		}
-		FlushTickBatch(Context);
 	}
 
 	// When not running async we can finalize straight away
diff --git a/Engine/Plugins/FX/Niagara/Source/Niagara/Public/NiagaraSystemSimulation.h b/Engine/Plugins/FX/Niagara/Source/Niagara/Public/NiagaraSystemSimulation.h
--- a/Engine/Plugins/FX/Niagara/Source/Niagara/Public/NiagaraSystemSimulation.h
+++ b/Engine/Plugins/FX/Niagara/Source/Niagara/Public/NiagaraSystemSimulation.h
@@ -343,8 +343,8 @@ protected:
 		TUniquePtr<FNiagaraScriptExecutionContextBase>& ExecContext,
 		FScriptExecutionConstantBufferTable& ConstantBufferTable) const;
 
-	void AddSystemToTickBatch(FNiagaraSystemInstance* Instance, FNiagaraSystemSimulationTickContext& Context);
-	void FlushTickBatch(FNiagaraSystemSimulationTickContext& Context);
+	void AddSystemToTickBatch(FNiagaraSystemSimulationTickContext& Context, FNiagaraSystemInstance* Instance, bool bIsLastInstance);
+	void FlushTickBatch(FNiagaraSystemSimulationTickContext& Context, bool bIsLastInstance);
 
 	void Tick_GameThread_Internal(float DeltaSeconds, const FGraphEventRef& MyCompletionGraphEvent);
 
