Change handling of undo with NiagaraSystemViewModel so that kills active compilations first rather than just invalidating their results after trying to process the stack.  Avoids crashes with the old compilation mode when data (like emitters) gets removed out from under it.

diff --git a/Engine/Plugins/FX/Niagara/Source/NiagaraEditor/Private/ViewModels/NiagaraSystemViewModel.cpp b/Engine/Plugins/FX/Niagara/Source/NiagaraEditor/Private/ViewModels/NiagaraSystemViewModel.cpp
--- a/Engine/Plugins/FX/Niagara/Source/NiagaraEditor/Private/ViewModels/NiagaraSystemViewModel.cpp
+++ b/Engine/Plugins/FX/Niagara/Source/NiagaraEditor/Private/ViewModels/NiagaraSystemViewModel.cpp
@@ -1063,8 +1063,8 @@ void FNiagaraSystemViewModel::AddReferencedObjects(FReferenceCollector& Collecto
 void FNiagaraSystemViewModel::PostUndo(bool bSuccess)
 {
 	// Reset system stack and emitter handle view models to prevent accessing invalid data if they were in the undo operation.
+	System->KillAllActiveCompilations();
 	ResetStack();
-	System->InvalidateActiveCompiles();
 	System->RequestCompile(false);
 }
 
