- Fix static analysis issue

#rnx

diff --git a/Engine/Plugins/FX/Niagara/Source/Niagara/Private/NiagaraDataInterfaceSkeletalMesh.cpp b/Engine/Plugins/FX/Niagara/Source/Niagara/Private/NiagaraDataInterfaceSkeletalMesh.cpp
--- a/Engine/Plugins/FX/Niagara/Source/Niagara/Private/NiagaraDataInterfaceSkeletalMesh.cpp
+++ b/Engine/Plugins/FX/Niagara/Source/Niagara/Private/NiagaraDataInterfaceSkeletalMesh.cpp
@@ -337,7 +337,7 @@ void FSkeletalMeshSkinningData::RegisterUser(FSkeletalMeshSkinningDataUsage Usag
 	if (bNeedsDataImmediately)
 	{
 		check(IsInGameThread());
-		const int32 NumExpectedTransforms = SkelMesh->GetRefBasesInvMatrix().Num();
+		const int32 NumExpectedTransforms = SkelMesh ? SkelMesh->GetRefBasesInvMatrix().Num() : 0;
 		if (CurrBoneRefToLocals().Num() != NumExpectedTransforms || CurrComponentTransforms().Num() != NumExpectedTransforms)
 		{
 			UpdateBoneTransforms();
