- Fix OOB reads from solo simulations
- Updated the reregister context to ensure solo simulations are destroyed rather than just batched
- Invalidated the script execution context when emitters are enabled / disabled / removed, this fixes parameter store bindings potentially reading OOB memory

diff --git a/Engine/Plugins/FX/Niagara/Source/Niagara/Private/NiagaraCommon.cpp b/Engine/Plugins/FX/Niagara/Source/Niagara/Private/NiagaraCommon.cpp
--- a/Engine/Plugins/FX/Niagara/Source/Niagara/Private/NiagaraCommon.cpp
+++ b/Engine/Plugins/FX/Niagara/Source/Niagara/Private/NiagaraCommon.cpp
@@ -88,14 +88,37 @@ FNiagaraSystemUpdateContext::~FNiagaraSystemUpdateContext()
 
 void FNiagaraSystemUpdateContext::CommitUpdate()
 {
-	for (UNiagaraSystem* NiagaraSystem : SystemSimsToDestroy)
+	if (SystemSimsToDestroy.Num() > 0)
 	{
-		if (NiagaraSystem)
+		for (UNiagaraSystem* NiagaraSystem : SystemSimsToDestroy)
 		{
-			FNiagaraWorldManager::DestroyAllSystemSimulations(NiagaraSystem);
+			if (NiagaraSystem)
+			{
+				FNiagaraWorldManager::DestroyAllSystemSimulations(NiagaraSystem);
+			}
+		}
+
+		for (TObjectIterator<UNiagaraComponent> It; It; ++It)
+		{
+			UNiagaraComponent* Comp = *It;
+
+			if ( ::IsValid(Comp) && SystemSimsToDestroy.Contains(Comp->GetAsset()) )
+			{
+				FNiagaraSystemInstanceControllerConstPtr SystemInstanceController = Comp->GetSystemInstanceController();
+				if (SystemInstanceController && SystemInstanceController->HasValidSimulation() && SystemInstanceController->IsSolo())
+				{
+					FNiagaraSystemSimulationPtr SoloSimulation = SystemInstanceController->GetSoloSystemSimulation();
+					if (SoloSimulation.IsValid())
+					{
+						Comp->DestroyInstanceNotComponent();
+						SoloSimulation->Destroy();
+					}
+				}
+			}
 		}
+
+		SystemSimsToDestroy.Empty();
 	}
-	SystemSimsToDestroy.Empty();
 
 	for (UNiagaraComponent* Comp : ComponentsToDestroyInstance)
 	{
diff --git a/Engine/Plugins/FX/Niagara/Source/Niagara/Private/NiagaraEmitterHandle.cpp b/Engine/Plugins/FX/Niagara/Source/Niagara/Private/NiagaraEmitterHandle.cpp
--- a/Engine/Plugins/FX/Niagara/Source/Niagara/Private/NiagaraEmitterHandle.cpp
+++ b/Engine/Plugins/FX/Niagara/Source/Niagara/Private/NiagaraEmitterHandle.cpp
@@ -148,6 +148,7 @@ bool FNiagaraEmitterHandle::SetIsEnabled(bool bInIsEnabled, UNiagaraSystem& InOw
 			FString InvalidateReason = TEXT("Emitter enabled changed.");
 			InOwnerSystem.GetSystemSpawnScript()->InvalidateCompileResults(InvalidateReason);
 			InOwnerSystem.GetSystemUpdateScript()->InvalidateCompileResults(InvalidateReason);
+			InOwnerSystem.RefreshSystemParametersFromEmitter(*this);
 
 			// Clean out the emitter's compile results for cleanliness.
 			FVersionedNiagaraEmitterData* EmitterData = VersionedInstance.GetEmitterData();
diff --git a/Engine/Plugins/FX/Niagara/Source/Niagara/Private/NiagaraSystem.cpp b/Engine/Plugins/FX/Niagara/Source/Niagara/Private/NiagaraSystem.cpp
--- a/Engine/Plugins/FX/Niagara/Source/Niagara/Private/NiagaraSystem.cpp
+++ b/Engine/Plugins/FX/Niagara/Source/Niagara/Private/NiagaraSystem.cpp
@@ -1492,8 +1492,18 @@ void UNiagaraSystem::RefreshSystemParametersFromEmitter(const FNiagaraEmitterHan
 	{
 		if (FVersionedNiagaraEmitterData* EmitterData = EmitterHandle.GetEmitterData())
 		{
-			EmitterData->EmitterSpawnScriptProps.Script->RapidIterationParameters.CopyParametersTo(SystemSpawnScript->RapidIterationParameters, false, FNiagaraParameterStore::EDataInterfaceCopyMethod::None);
-			EmitterData->EmitterUpdateScriptProps.Script->RapidIterationParameters.CopyParametersTo(SystemUpdateScript->RapidIterationParameters, false, FNiagaraParameterStore::EDataInterfaceCopyMethod::None);
+			if (EmitterHandle.GetIsEnabled())
+			{
+				EmitterData->EmitterSpawnScriptProps.Script->RapidIterationParameters.CopyParametersTo(SystemSpawnScript->RapidIterationParameters, false, FNiagaraParameterStore::EDataInterfaceCopyMethod::None);
+				EmitterData->EmitterUpdateScriptProps.Script->RapidIterationParameters.CopyParametersTo(SystemUpdateScript->RapidIterationParameters, false, FNiagaraParameterStore::EDataInterfaceCopyMethod::None);
+			}
+			else
+			{
+				EmitterData->EmitterSpawnScriptProps.Script->RapidIterationParameters.RemoveParameters(SystemSpawnScript->RapidIterationParameters);
+				EmitterData->EmitterUpdateScriptProps.Script->RapidIterationParameters.RemoveParameters(SystemUpdateScript->RapidIterationParameters);
+			}
+
+			ForEachScript([&](UNiagaraScript* Script) { Script->InvalidateExecutionReadyParameterStores(); });
 		}
 	}
 }
@@ -1509,6 +1519,8 @@ void UNiagaraSystem::RemoveSystemParametersForEmitter(const FNiagaraEmitterHandl
 			EmitterData->EmitterSpawnScriptProps.Script->RapidIterationParameters.RemoveParameters(SystemSpawnScript->RapidIterationParameters);
 			EmitterData->EmitterUpdateScriptProps.Script->RapidIterationParameters.RemoveParameters(SystemUpdateScript->RapidIterationParameters);
 		}
+
+		ForEachScript([&](UNiagaraScript* Script) { Script->InvalidateExecutionReadyParameterStores(); });
 	}
 }
 #endif
