- Thread safety changes for AllowAsyncRenderThreadUpdatesEditorGameWorld and AllowAsyncRenderThreadUpdatesEditor
- This is a partial integrate of PR 12926 with minor modifications to Niagara

diff --git a/Engine/Plugins/FX/Niagara/Source/Niagara/Classes/NiagaraPlatformSet.h b/Engine/Plugins/FX/Niagara/Source/Niagara/Classes/NiagaraPlatformSet.h
--- a/Engine/Plugins/FX/Niagara/Source/Niagara/Classes/NiagaraPlatformSet.h
+++ b/Engine/Plugins/FX/Niagara/Source/Niagara/Classes/NiagaraPlatformSet.h
@@ -457,6 +457,8 @@ private:
 	static NIAGARA_API int32 QualityLevelOverride;
 
 #if WITH_EDITOR
+	static NIAGARA_API FCriticalSection CachedPlatformIniCritSec;
+
 	//Cached data read from platform ini files.
 	static NIAGARA_API TMap<FName, FPlatformIniSettings> CachedPlatformIniSettings;
 
diff --git a/Engine/Plugins/FX/Niagara/Source/Niagara/Private/NiagaraPlatformSet.cpp b/Engine/Plugins/FX/Niagara/Source/Niagara/Private/NiagaraPlatformSet.cpp
--- a/Engine/Plugins/FX/Niagara/Source/Niagara/Private/NiagaraPlatformSet.cpp
+++ b/Engine/Plugins/FX/Niagara/Source/Niagara/Private/NiagaraPlatformSet.cpp
@@ -292,6 +292,7 @@ int32 FNiagaraPlatformSet::GetFullQualityLevelMask(int32 NumQualityLevels)
 
 uint32 FNiagaraPlatformSet::LastDirtiedFrame = 0;
 #if WITH_EDITOR
+FCriticalSection FNiagaraPlatformSet::CachedPlatformIniCritSec;
 TMap<const UDeviceProfile*, int32> FNiagaraPlatformSet::CachedQLMasksPerDeviceProfile;
 TMap<FName, FNiagaraPlatformSet::FPlatformIniSettings> FNiagaraPlatformSet::CachedPlatformIniSettings;
 #endif
@@ -467,12 +468,18 @@ bool FNiagaraPlatformSet::CanChangeScalabilityAtRuntime()
 int32 FNiagaraPlatformSet::GetActiveQualityMaskForDeviceProfile(const UDeviceProfile* Profile)
 {
 #if WITH_EDITOR
+	FScopeLock Lock(&CachedPlatformIniCritSec);
 
 	if (int32* CachedQLForProfile = CachedQLMasksPerDeviceProfile.Find(Profile))
 	{
 		return *CachedQLForProfile;
 	}
-	
+
+	if (!ensureMsgf(IsInGameThread(), TEXT("GetActiveQualityMaskForDeviceProfile is attempting to access the config system outside of the game thread, results may be incorrect.")))
+	{
+		return INDEX_NONE;
+	}
+
 	FPlatformIniSettings& PlatformSettings = GetPlatformIniSettings(Profile->DeviceType);
 	int32 QLMask = PlatformSettings.QualityLevelMask;
 
@@ -1053,8 +1060,11 @@ bool FNiagaraPlatformSet::CanConsiderDeviceProfile(const UDeviceProfile* Profile
 void FNiagaraPlatformSet::InvalidateCachedData()
 {
 #if WITH_EDITOR
-	CachedQLMasksPerDeviceProfile.Empty();
-	CachedPlatformIniSettings.Empty();
+	{
+		FScopeLock Lock(&CachedPlatformIniCritSec);
+		CachedQLMasksPerDeviceProfile.Empty();
+		CachedPlatformIniSettings.Empty();
+	}
 	FDeviceProfileValueCache::Empty();
 #endif
 
@@ -1251,6 +1261,8 @@ bool FNiagaraPlatformSet::GatherConflicts(TConstArrayView<const FNiagaraPlatform
 
 FNiagaraPlatformSet::FPlatformIniSettings& FNiagaraPlatformSet::GetPlatformIniSettings(const FString& PlatformName)
 {
+	FScopeLock Lock(&CachedPlatformIniCritSec);
+
 	if (FPlatformIniSettings* CachedSettings = CachedPlatformIniSettings.Find(*PlatformName))
 	{
 		return *CachedSettings;
diff --git a/Engine/Source/Runtime/Engine/Classes/Particles/ParticleSystemComponent.h b/Engine/Source/Runtime/Engine/Classes/Particles/ParticleSystemComponent.h
--- a/Engine/Source/Runtime/Engine/Classes/Particles/ParticleSystemComponent.h
+++ b/Engine/Source/Runtime/Engine/Classes/Particles/ParticleSystemComponent.h
@@ -1369,6 +1369,10 @@ protected:
 	ENGINE_API virtual void OnUnregister()  override;
 	ENGINE_API virtual void SendRenderDynamicData_Concurrent() override;
 public:
+#if WITH_EDITORONLY_DATA
+	ENGINE_API virtual bool RequiresGameThreadEndOfFrameUpdates() const override { return true; }
+	ENGINE_API virtual bool RequiresGameThreadEndOfFrameRecreate() const override { return true; }
+#endif
 	//~ End UActorComponent Interface.
 
 	enum EForceAsyncWorkCompletion
