- Add 'fx.Niagara.SystemInstance.LWCTileUpdateMode' to control what happens when we cross LWC tile boundaries
- Default is to reset the simualtion, optionally can be changed to rebase position attributes
- Does not handle data interface rebasing

diff --git a/Engine/Plugins/FX/Niagara/Shaders/Private/NiagaraLWCHelper.usf b/Engine/Plugins/FX/Niagara/Shaders/Private/NiagaraLWCHelper.usf
new file mode 100644
--- /dev/null
+++ b/Engine/Plugins/FX/Niagara/Shaders/Private/NiagaraLWCHelper.usf
@@ -0,0 +1,35 @@
+// Copyright Epic Games, Inc. All Rights Reserved.
+
+#include "/Engine/Private/Common.ush"
+#include "/Engine/Private/ComputeShaderUtils.ush"
+
+RWBuffer<float>	FloatBuffer;
+uint			FloatBufferStride;
+
+Buffer<int>		CountBuffer;
+uint			CountBufferOffset;
+
+uint			NumPositions;
+DECLARE_SCALAR_ARRAY(uint,	PositionComponentOffsets, MaxPositions);
+float3			TileShift;
+
+[numthreads(THREADGROUP_SIZE, 1, 1)]
+void TileShiftPositions(uint3 GroupId : SV_GroupID, int GroupThreadIndex : SV_GroupIndex)
+{
+	const uint LinearThreadId = GetUnWrappedDispatchThreadId(GroupId, GroupThreadIndex, THREADGROUP_SIZE);
+	const uint NumInstances = uint(CountBuffer[CountBufferOffset]);	
+	if ( LinearThreadId >= NumInstances )
+	{
+		return;
+	}
+	
+	for ( uint iPosition=0; iPosition < NumPositions; ++iPosition )
+	{
+		const uint XOffset = LinearThreadId + (GET_SCALAR_ARRAY_ELEMENT(PositionComponentOffsets, iPosition) * FloatBufferStride);
+		const uint YOffset = XOffset + FloatBufferStride;
+		const uint ZOffset = YOffset + FloatBufferStride;
+		FloatBuffer[XOffset] = FloatBuffer[XOffset] + TileShift.x;
+		FloatBuffer[YOffset] = FloatBuffer[YOffset] + TileShift.y;
+		FloatBuffer[ZOffset] = FloatBuffer[ZOffset] + TileShift.z;
+	}
+}
diff --git a/Engine/Plugins/FX/Niagara/Source/Niagara/Classes/NiagaraSystem.h b/Engine/Plugins/FX/Niagara/Source/Niagara/Classes/NiagaraSystem.h
--- a/Engine/Plugins/FX/Niagara/Source/Niagara/Classes/NiagaraSystem.h
+++ b/Engine/Plugins/FX/Niagara/Source/Niagara/Classes/NiagaraSystem.h
@@ -196,6 +196,28 @@ struct FNiagaraSystemScalabilityOverrides
 	TArray<FNiagaraSystemScalabilityOverride> Overrides;
 };
 
+UENUM()
+enum class ENiagaraLwcTileUpdateMode : uint8
+{
+	/**
+	When a simulation moves greater or equal to the number of tiles specified in fx.LWCTileRecache the simulation will be reset.
+	This means all particles are cleared and we start the system from time 0 again.
+	*/
+	ResetSimulation,
+	/**
+	When a simulation moves greater or equal to the number of tiles specified in fx.LWCTileRecache we will trigger a rebase.
+	A rebase will move all Position type attributes by the difference in tiles.
+	This is not guaranteed to cover everything, such as data inteface, matrix attributes, etc, and is therefore could result in simulation errors.
+	*/
+	Rebase,
+	/**
+	This option is similar to Rebase with the exception that when we cross fx.LWCTileForceReset number of tiles we will reset the simulation.
+	This can provide some additional safety if you were to cross a large number of tiles and move back the same number of tiles,
+	as in situations like the floating point numbers may have lost significant precision resulting in an invalid simulation.
+	*/
+	RebaseOrResetSimulation,
+};
+
 /** A Niagara System contains multiple Niagara Emitters to create various effects.
  * Niagara Systems can be placed in the world, unlike Emitters, and expose User Parameters to configure an effect at runtime.*/
 UCLASS(BlueprintType, MinimalAPI, meta = (LoadBehavior = "LazyOnDemand"))
@@ -496,6 +518,8 @@ public:
 	NIAGARA_API bool UsesCollection(const UNiagaraParameterCollection* Collection)const;
 
 	bool SupportsLargeWorldCoordinates() const { return bSupportLargeWorldCoordinates && bLwcEnabledSettingCached; }
+	ENiagaraLwcTileUpdateMode GetLargeWorldCoordinateTileUpdateMode() const;
+
 	FORCEINLINE bool ShouldDisableExperimentalVM() const { return bDisableExperimentalVM; }
 
 	FORCEINLINE bool IsInitialOwnerVelocityFromActor() const { return bInitialOwnerVelocityFromActor; }
@@ -629,6 +653,9 @@ private:
 	UPROPERTY(EditAnywhere, AdvancedDisplay, Category = "System")
 	uint8 bInitialOwnerVelocityFromActor : 1 = true;
 
+	UPROPERTY(EditAnywhere, Category = Rendering, AdvancedDisplay, meta = (EditCondition = "bSupportLargeWorldCoordinates", EditConditionHides, DisplayAfter=bSupportLargeWorldCoordinates))
+	TOptional<ENiagaraLwcTileUpdateMode> LargeWorldCoordinateTileUpdateMode;
+
 public:
 	/**
 	When enabled this is the default value set on the component.
@@ -641,6 +668,7 @@ public:
 	When enabled this is the default value set on the component.
 	Optionally write this 0-255 value to the stencil buffer in CustomDepth pass (Requires project setting or r.CustomDepth == 3)
 	*/
+
 	UPROPERTY(EditAnywhere, BlueprintReadOnly, AdvancedDisplay, Category="Rendering", meta=(DisplayName="Default CustomDepthStencil Value", editcondition="bOverrideCustomDepthStencilWriteMask", UIMin = "0", UIMax = "255"))
 	int32 CustomDepthStencilValue = 0;
 
diff --git a/Engine/Plugins/FX/Niagara/Source/Niagara/Private/NiagaraLWCHelper.cpp b/Engine/Plugins/FX/Niagara/Source/Niagara/Private/NiagaraLWCHelper.cpp
new file mode 100644
--- /dev/null
+++ b/Engine/Plugins/FX/Niagara/Source/Niagara/Private/NiagaraLWCHelper.cpp
@@ -0,0 +1,209 @@
+// Copyright Epic Games, Inc. All Rights Reserved.
+
+#include "NiagaraLWCHelper.h"
+
+#include "NiagaraComputeExecutionContext.h"
+#include "NiagaraEmitterInstanceImpl.h"
+#include "NiagaraDataSet.h"
+#include "NiagaraGpuComputeDispatchInterface.h"
+#include "NiagaraLWCTileShiftPositions.h"
+
+namespace NiagaraLWCHelper
+{
+using FGPUEmitterArray = TArray<FNiagaraComputeExecutionContext*, TInlineAllocator<8>>;
+
+void RebaseGPUEmitters(FRHICommandListImmediate& RHICmdList, const FVector3f& TileShift, const FGPUEmitterArray& GPUEmitters, FNiagaraGpuComputeDispatchInterface* ComputeInterface)
+{
+	TArray<FRHITransitionInfo, TInlineAllocator<8>>			TransitionsBefore;
+	TArray<FRHITransitionInfo, TInlineAllocator<8>>			TransitionsAfter;
+	TArray<FNiagaraLWCTileShiftPositionsCS::FParameters>	ComputeJobs;
+
+	for (FNiagaraComputeExecutionContext* GPUContext : GPUEmitters)
+	{
+		FNiagaraDataSet* DataSet = GPUContext->MainDataSet;
+		FNiagaraDataBuffer* CurrentBuffer = DataSet->GetCurrentData();
+		if (CurrentBuffer == nullptr)
+		{
+			return;
+		}
+
+		const uint32 CountBufferOffset = CurrentBuffer->GetGPUInstanceCountBufferOffset();
+		if (CountBufferOffset == INDEX_NONE)
+		{
+			return;
+		}
+
+		const uint32 NumInstances = CurrentBuffer->GetNumInstances();
+		if (NumInstances == 0)
+		{
+			return;
+		}
+
+		FRHIUnorderedAccessView* FloatBufferUAV = CurrentBuffer->GetGPUBufferFloat().UAV;
+
+		FNiagaraLWCTileShiftPositionsCS::FParameters ShaderParameters;
+		ShaderParameters.FloatBuffer		= FloatBufferUAV;
+		ShaderParameters.FloatBufferStride	= CurrentBuffer->GetFloatStride() / sizeof(float);
+		ShaderParameters.NumInstances		= NumInstances;
+		ShaderParameters.CountBuffer		= ComputeInterface->GetGPUInstanceCounterManager().GetInstanceCountBuffer().SRV;
+		ShaderParameters.CountBufferOffset	= CountBufferOffset;
+		ShaderParameters.TileShift			= TileShift;
+
+		const TConstArrayView<FNiagaraVariableLayoutInfo> VariableLayouts(DataSet->GetCompiledData().VariableLayouts);
+		const TConstArrayView<FNiagaraVariableBase> Variables(DataSet->GetCompiledData().Variables);
+		const int32 NumVariables = FMath::Min(VariableLayouts.Num(), Variables.Num());
+		uint32 NumPositions = 0;
+		bool bNeedsTransitions = false;
+
+		for (int32 iVariable = 0; iVariable < NumVariables; ++iVariable)
+		{
+			if (Variables[iVariable].GetType() != FNiagaraTypeDefinition::GetPositionDef())
+			{
+				continue;
+			}
+
+			if (ensure(VariableLayouts[iVariable].GetNumFloatComponents() == 3))
+			{
+				const uint32 FloatOffset = VariableLayouts[iVariable].GetFloatComponentStart();
+				GET_SCALAR_ARRAY_ELEMENT(ShaderParameters.PositionComponentOffsets, NumPositions) = FloatOffset;
+				++NumPositions;
+				if (NumPositions >= FNiagaraLWCTileShiftPositionsCS::MaxPositions)
+				{
+					ShaderParameters.NumPositions = NumPositions;
+					bNeedsTransitions = true;
+					ComputeJobs.Add(ShaderParameters);
+					NumPositions = 0;
+				}
+			}
+		}
+
+		if (NumPositions > 0)
+		{
+			ShaderParameters.NumPositions = NumPositions;
+			bNeedsTransitions = true;
+			ComputeJobs.Add(ShaderParameters);
+		}
+
+		if (bNeedsTransitions)
+		{
+			TransitionsBefore.Emplace(FloatBufferUAV, ERHIAccess::SRVMask, ERHIAccess::UAVCompute);
+			TransitionsAfter.Emplace(FloatBufferUAV, ERHIAccess::UAVCompute, ERHIAccess::SRVMask);
+		}
+	}
+
+	if (ComputeJobs.Num() > 0)
+	{
+		RHICmdList.Transition(TransitionsBefore);
+		for (const FNiagaraLWCTileShiftPositionsCS::FParameters& Parameters : ComputeJobs)
+		{
+			FNiagaraLWCTileShiftPositionsCS::Execute(RHICmdList, Parameters);
+		}
+		RHICmdList.Transition(TransitionsAfter);
+	}
+}
+
+void RebaseEmitters(const FVector3f& TileShift, TArrayView<FNiagaraEmitterInstanceRef> Emitters, FNiagaraGpuComputeDispatchInterface* ComputeInterface)
+{
+	FGPUEmitterArray GPUEmittersToRebase;
+
+	for (const FNiagaraEmitterInstanceRef& EmitterRef : Emitters)
+	{
+		FNiagaraEmitterInstanceImpl* EmitterInstance = EmitterRef->AsStateful();
+		if (EmitterInstance == nullptr || EmitterInstance->IsLocalSpace())
+		{
+			return;
+		}
+
+		// GPU data to rebase
+		if (FNiagaraComputeExecutionContext* GPUContext = EmitterInstance->GetGPUContext())
+		{
+			if (ComputeInterface && EmitterInstance->GetNumParticles() > 0)
+			{
+				GPUEmittersToRebase.Add(GPUContext);
+			}
+		}
+		// CPU data to rebase
+		else
+		{
+			RebaseDataSet(TileShift, EmitterInstance->GetParticleData());
+		}
+	}
+
+	if (GPUEmittersToRebase.Num() > 0)
+	{
+		ENQUEUE_RENDER_COMMAND(RebaseGPUEmitter)
+		(
+			[TileShift, GPUEmittersToRebase, ComputeInterface](FRHICommandListImmediate& RHICmdList)
+			{
+				RebaseGPUEmitters(RHICmdList, TileShift, GPUEmittersToRebase, ComputeInterface);
+			}
+		);
+	}
+}
+
+void RebaseDataSet(const FVector3f& TileShift, FNiagaraDataSet& DataSet, uint32 iInstance)
+{
+	FNiagaraDataBuffer* CurrentData = DataSet.GetCurrentData();
+	if (CurrentData == nullptr || iInstance < CurrentData->GetNumInstances())
+	{
+		return;
+	}
+
+	const TConstArrayView<FNiagaraVariableLayoutInfo> VariableLayouts(DataSet.GetCompiledData().VariableLayouts);
+	const TConstArrayView<FNiagaraVariableBase> Variables(DataSet.GetCompiledData().Variables);
+	const int32 NumVariables = FMath::Min(VariableLayouts.Num(), Variables.Num());
+	for (int32 iVariable = 0; iVariable < NumVariables; ++iVariable)
+	{
+		if (Variables[iVariable].GetType() != FNiagaraTypeDefinition::GetPositionDef())
+		{
+			continue;
+		}
+
+		if (ensure(VariableLayouts[iVariable].GetNumFloatComponents() == 3))
+		{
+			const uint32 FloatOffset = VariableLayouts[iVariable].GetFloatComponentStart();
+			float* XComponent = reinterpret_cast<float*>(CurrentData->GetComponentPtrFloat(FloatOffset + 0));
+			float* YComponent = reinterpret_cast<float*>(CurrentData->GetComponentPtrFloat(FloatOffset + 1));
+			float* ZComponent = reinterpret_cast<float*>(CurrentData->GetComponentPtrFloat(FloatOffset + 2));
+			XComponent[iInstance] += TileShift.X;
+			YComponent[iInstance] += TileShift.Y;
+			ZComponent[iInstance] += TileShift.Z;
+		}
+	}
+}
+
+void RebaseDataSet(const FVector3f& TileShift, FNiagaraDataSet& DataSet)
+{
+	FNiagaraDataBuffer* CurrentData = DataSet.GetCurrentData();
+	if (CurrentData == nullptr || CurrentData->GetNumInstances() == 0)
+	{
+		return;
+	}
+
+	const TConstArrayView<FNiagaraVariableLayoutInfo> VariableLayouts(DataSet.GetCompiledData().VariableLayouts);
+	const TConstArrayView<FNiagaraVariableBase> Variables(DataSet.GetCompiledData().Variables);
+	const int32 NumVariables = FMath::Min(VariableLayouts.Num(), Variables.Num());
+	for (int32 iVariable = 0; iVariable < NumVariables; ++iVariable)
+	{
+		if (Variables[iVariable].GetType() != FNiagaraTypeDefinition::GetPositionDef())
+		{
+			continue;
+		}
+
+		if (ensure(VariableLayouts[iVariable].GetNumFloatComponents() == 3))
+		{
+			const uint32 FloatOffset = VariableLayouts[iVariable].GetFloatComponentStart();
+			float* XComponent = reinterpret_cast<float*>(CurrentData->GetComponentPtrFloat(FloatOffset + 0));
+			float* YComponent = reinterpret_cast<float*>(CurrentData->GetComponentPtrFloat(FloatOffset + 1));
+			float* ZComponent = reinterpret_cast<float*>(CurrentData->GetComponentPtrFloat(FloatOffset + 2));
+			for (uint32 iInstance = 0; iInstance < CurrentData->GetNumInstances(); ++iInstance)
+			{
+				XComponent[iInstance] += TileShift.X;
+				YComponent[iInstance] += TileShift.Y;
+				ZComponent[iInstance] += TileShift.Z;
+			}
+		}
+	}
+}
+
+} //namespace NiagaraLWCHelper
diff --git a/Engine/Plugins/FX/Niagara/Source/Niagara/Private/NiagaraLWCHelper.h b/Engine/Plugins/FX/Niagara/Source/Niagara/Private/NiagaraLWCHelper.h
new file mode 100644
--- /dev/null
+++ b/Engine/Plugins/FX/Niagara/Source/Niagara/Private/NiagaraLWCHelper.h
@@ -0,0 +1,15 @@
+// Copyright Epic Games, Inc. All Rights Reserved.
+
+#pragma once
+
+#include "NiagaraEmitterInstance.h"
+
+class FNiagaraDataSet;
+class FNiagaraGpuComputeDispatchInterface;
+
+namespace NiagaraLWCHelper
+{
+	extern void RebaseEmitters(const FVector3f& TileShift, TArrayView<FNiagaraEmitterInstanceRef> Emitters, FNiagaraGpuComputeDispatchInterface* ComputeInterface);
+	extern void RebaseDataSet(const FVector3f& TileShift, FNiagaraDataSet& DataSet, uint32 iInstance);
+	extern void RebaseDataSet(const FVector3f& TileShift, FNiagaraDataSet& DataSet);
+}
diff --git a/Engine/Plugins/FX/Niagara/Source/Niagara/Private/NiagaraSystem.cpp b/Engine/Plugins/FX/Niagara/Source/Niagara/Private/NiagaraSystem.cpp
--- a/Engine/Plugins/FX/Niagara/Source/Niagara/Private/NiagaraSystem.cpp
+++ b/Engine/Plugins/FX/Niagara/Source/Niagara/Private/NiagaraSystem.cpp
@@ -608,6 +608,11 @@ bool UNiagaraSystem::UsesCollection(const UNiagaraParameterCollection* Collectio
 	return false;
 }
 
+ENiagaraLwcTileUpdateMode UNiagaraSystem::GetLargeWorldCoordinateTileUpdateMode() const
+{
+	return LargeWorldCoordinateTileUpdateMode.Get(GetDefault<UNiagaraSettings>()->LargeWorldCoordinateTileUpdateMode);
+}
+
 #if 0
 namespace FNiagaraSystemParameterUtilities
 {
diff --git a/Engine/Plugins/FX/Niagara/Source/Niagara/Private/NiagaraSystemInstance.cpp b/Engine/Plugins/FX/Niagara/Source/Niagara/Private/NiagaraSystemInstance.cpp
--- a/Engine/Plugins/FX/Niagara/Source/Niagara/Private/NiagaraSystemInstance.cpp
+++ b/Engine/Plugins/FX/Niagara/Source/Niagara/Private/NiagaraSystemInstance.cpp
@@ -10,12 +10,14 @@
 #include "NiagaraEmitterInstance.h"
 #include "NiagaraEmitterInstanceImpl.h"
 #include "NiagaraGPUSystemTick.h"
+#include "NiagaraLWCHelper.h"
 #include "NiagaraSceneComponentUtils.h"
 #include "NiagaraStats.h"
 #include "NiagaraParameterCollection.h"
 #include "NiagaraSystemSimulation.h"
 #include "NiagaraWorldManager.h"
 #include "NiagaraComponent.h"
+#include "NiagaraSettings.h"
 #include "NiagaraSimCache.h"
 #include "NiagaraGpuComputeDebug.h"
 #include "NiagaraGpuComputeDispatchInterface.h"
@@ -1300,6 +1302,75 @@ FTransform FNiagaraSystemInstance::GetLWCSimToWorld(bool bLocalSpaceEmitter) con
 	return SimToWorld;
 }
 
+void FNiagaraSystemInstance::UpdateLWCTile(const FVector3f& InLWCTile, FNiagaraDataSet& EmitterDataSet)
+{
+	const ENiagaraLwcTileUpdateMode TileUpdateMode = System->GetLargeWorldCoordinateTileUpdateMode();
+	const FVector3f TileDifference(LWCTile - InLWCTile);
+
+	// Should we reset the simulation
+	bool bResetSimulation = false;
+	switch (TileUpdateMode)
+	{
+		case ENiagaraLwcTileUpdateMode::ResetSimulation:
+			bResetSimulation = true;
+			break;
+
+		case ENiagaraLwcTileUpdateMode::RebaseOrResetSimulation:
+			bResetSimulation = TileDifference.GetAbsMax() >= GetDefault<UNiagaraSettings>()->LargeWorldCoordinateMaxTilesBeforeReset;
+			break;
+	}
+
+	// Should we reset the system simulation
+	if ( bResetSimulation )
+	{
+		if (TickCount > 1)
+		{
+			UE_LOG(LogNiagara, Warning, TEXT("NiagaraComponent(%s - %s) required LWC tile recache and was reset."), *GetFullNameSafe(GetAttachComponent()), *GetFullNameSafe(System));
+		}
+		Reset(FNiagaraSystemInstance::EResetMode::ResetAll);
+		return;
+	}
+
+	// Rebase internal variables
+	const FVector LWCTileShift = FVector(TileDifference) * UE_LWC_RENDER_TILE_SIZE;
+	const FVector3f TileShift(LWCTileShift);
+
+	WorldTransform.AddToTranslation(LWCTileShift);
+	if (PreviousLocation.IsSet())
+	{
+		PreviousLocation = PreviousLocation.GetValue() + LWCTileShift;
+	}
+
+	for (FNiagaraOwnerParameters& Parameters : OwnerParameters)
+	{
+		Parameters.EngineLocalToWorld			= Parameters.EngineLocalToWorld.ConcatTranslation(TileShift);
+		Parameters.EngineWorldToLocal			= Parameters.EngineLocalToWorld.Inverse();
+		Parameters.EngineLocalToWorldTransposed	= Parameters.EngineLocalToWorld.GetTransposed();
+		Parameters.EngineWorldToLocalTransposed	= Parameters.EngineWorldToLocal.GetTransposed();
+		Parameters.EngineLocalToWorldNoScale	= Parameters.EngineLocalToWorldNoScale.ConcatTranslation(TileShift);
+		Parameters.EngineWorldToLocalNoScale	= Parameters.EngineLocalToWorldNoScale.Inverse();
+		Parameters.EnginePosition				+= TileShift;
+		Parameters.EngineLWCTile				= InLWCTile;
+
+	}
+
+	GatheredInstanceParameters.ComponentTrans.AddToTranslation(LWCTileShift);
+
+	// Rebase data sets
+	NiagaraLWCHelper::RebaseDataSet(TileShift, EmitterDataSet, SystemInstanceIndex);
+	NiagaraLWCHelper::RebaseEmitters(TileShift, Emitters, ComputeDispatchInterface);
+
+	//-TODO: Rebase data interface data
+
+	// Set the new tile and invalidate the render state to ensure the proxy is updated with the latest tile
+	LWCTile = InLWCTile;
+
+	if (USceneComponent* Component = GetAttachComponent())
+	{
+		Component->MarkRenderStateDirty();
+	}
+}
+
 FNiagaraWorldManager* FNiagaraSystemInstance::GetWorldManager() const
 {
 	check(World);
diff --git a/Engine/Plugins/FX/Niagara/Source/Niagara/Private/NiagaraSystemInstanceController.cpp b/Engine/Plugins/FX/Niagara/Source/Niagara/Private/NiagaraSystemInstanceController.cpp
--- a/Engine/Plugins/FX/Niagara/Source/Niagara/Private/NiagaraSystemInstanceController.cpp
+++ b/Engine/Plugins/FX/Niagara/Source/Niagara/Private/NiagaraSystemInstanceController.cpp
@@ -162,8 +162,16 @@ void FNiagaraSystemInstanceController::GetMaterialStreamingInfo(FNiagaraMaterial
 		return;
 	}
 
+	if (OverrideParameters)
+	{
+		OverrideParameters->Tick();
+	}
+	SystemInstance->GetInstanceParameters().Tick();
+
 	for (const FNiagaraEmitterInstanceRef& EmitterInst : SystemInstance->GetEmitters())
 	{
+		EmitterInst->GetRendererBoundVariables().Tick();
+
 		EmitterInst->ForEachEnabledRenderer(
 			[&](UNiagaraRendererProperties* Properties)
 			{
diff --git a/Engine/Plugins/FX/Niagara/Source/Niagara/Private/NiagaraSystemSimulation.cpp b/Engine/Plugins/FX/Niagara/Source/Niagara/Private/NiagaraSystemSimulation.cpp
--- a/Engine/Plugins/FX/Niagara/Source/Niagara/Private/NiagaraSystemSimulation.cpp
+++ b/Engine/Plugins/FX/Niagara/Source/Niagara/Private/NiagaraSystemSimulation.cpp
@@ -19,6 +19,7 @@
 #include "NiagaraGPUSystemTick.h"
 #include "NiagaraCrashReporterHandler.h"
 
+#include "Misc/LargeWorldRenderPosition.h"
 
 // Niagara Async Ticking Sequence
 // - NiagaraSystemSimulation::Tick_GameThread
@@ -1284,14 +1285,11 @@ void FNiagaraSystemSimulation::Tick_GameThread_Internal(float DeltaSeconds, cons
 				{
 					if ( Instance->GetRequestedExecutionState() == ENiagaraExecutionState::Active && Instance->GetActualExecutionState() == ENiagaraExecutionState::Active)
 					{
-						if (UFXSystemComponent::RequiresLWCTileRecache(Instance->GetLWCTile(), SceneComponent->GetComponentLocation()))
+						const FVector ComponentLocation = SceneComponent->GetComponentLocation();
+						if (UFXSystemComponent::RequiresLWCTileRecache(Instance->GetLWCTile(), ComponentLocation))
 						{
-							//-OPT: For safety we reset everything, but if everything is local space we may not need to, or we could rebase.
-							if (Instance->GetTickCount() > 1)
-							{
-								UE_LOG(LogNiagara, Warning, TEXT("NiagaraComponent(%s - %s) required LWC tile recache and was reset."), *GetFullNameSafe(SceneComponent), *GetFullNameSafe(System));
-							}
-							Instance->Reset(FNiagaraSystemInstance::EResetMode::ResetAll);
+							const FVector3f ActorTile = FLargeWorldRenderScalar::GetTileFor(ComponentLocation);
+							Instance->UpdateLWCTile(ActorTile, MainDataSet);
 						}
 					}
 				}
diff --git a/Engine/Plugins/FX/Niagara/Source/Niagara/Public/NiagaraSettings.h b/Engine/Plugins/FX/Niagara/Source/Niagara/Public/NiagaraSettings.h
--- a/Engine/Plugins/FX/Niagara/Source/Niagara/Public/NiagaraSettings.h
+++ b/Engine/Plugins/FX/Niagara/Source/Niagara/Public/NiagaraSettings.h
@@ -149,6 +149,14 @@ class UNiagaraSettings : public UDeveloperSettings
 	UPROPERTY(config, EditAnywhere, Category = Niagara, meta = ( ConfigRestartRequired = true ))
 	bool bSystemsSupportLargeWorldCoordinates = true;
 
+	/** Controls how we respond to a system instance crossing large world coordinate tiles. Can be overriden on individual assets. */
+	UPROPERTY(config, EditAnywhere, Category = Niagara, meta = (EditCondition = "bSystemsSupportLargeWorldCoordinates", EditConditionHides))
+	ENiagaraLwcTileUpdateMode LargeWorldCoordinateTileUpdateMode = ENiagaraLwcTileUpdateMode::ResetSimulation;
+
+	/** The maximum tiles we can cross before forcing a simulation to reset, if the tile update mode is set to RebaseOrResetSimulation. */
+	UPROPERTY(config, EditAnywhere, Category = Niagara, meta = (EditCondition = "bSystemsSupportLargeWorldCoordinates", EditConditionHides))
+	uint32 LargeWorldCoordinateMaxTilesBeforeReset = 4;
+
 	/**
 	 If set to true, types like positions and vectors cannot be assigned to each other without an explicit conversion step.
 	 If false, type checks are loosened and some types can be implicitly converted into each other.
diff --git a/Engine/Plugins/FX/Niagara/Source/Niagara/Public/NiagaraSystemInstance.h b/Engine/Plugins/FX/Niagara/Source/Niagara/Public/NiagaraSystemInstance.h
--- a/Engine/Plugins/FX/Niagara/Source/Niagara/Public/NiagaraSystemInstance.h
+++ b/Engine/Plugins/FX/Niagara/Source/Niagara/Public/NiagaraSystemInstance.h
@@ -160,6 +160,7 @@ public:
 	FORCEINLINE FNiagaraParameterStore& GetInstanceParameters() { return InstanceParameters; }
 	NIAGARA_API FNiagaraLWCConverter GetLWCConverter(bool bLocalSpaceEmitter = false) const;
 	NIAGARA_API FTransform GetLWCSimToWorld(bool bLocalSpaceEmitter = false) const;
+	void UpdateLWCTile(const FVector3f& InLWCTile, FNiagaraDataSet& EmitterDataSet);
 
 	FORCEINLINE uint32 GetParameterIndex(bool PreviousFrame = false) const
 	{
diff --git a/Engine/Plugins/FX/Niagara/Source/NiagaraShader/Internal/NiagaraLWCTileShiftPositions.h b/Engine/Plugins/FX/Niagara/Source/NiagaraShader/Internal/NiagaraLWCTileShiftPositions.h
new file mode 100644
--- /dev/null
+++ b/Engine/Plugins/FX/Niagara/Source/NiagaraShader/Internal/NiagaraLWCTileShiftPositions.h
@@ -0,0 +1,34 @@
+// Copyright Epic Games, Inc. All Rights Reserved.
+
+#pragma once
+
+#include "GlobalShader.h"
+#include "RHI.h"
+#include "RHIUtilities.h"
+#include "RenderGraphBuilder.h"
+#include "RenderGraphUtils.h"
+
+class FNiagaraLWCTileShiftPositionsCS : public FGlobalShader
+{
+	DECLARE_GLOBAL_SHADER(FNiagaraLWCTileShiftPositionsCS);
+	SHADER_USE_PARAMETER_STRUCT(FNiagaraLWCTileShiftPositionsCS, FGlobalShader);
+
+	static constexpr uint32 ThreadGroupSize = 64;
+	static constexpr uint32 MaxPositions = 8;
+
+	BEGIN_SHADER_PARAMETER_STRUCT(FParameters, NIAGARASHADER_API)
+		SHADER_PARAMETER_UAV(RWBuffer<float>,	FloatBuffer)
+		SHADER_PARAMETER(uint32,				FloatBufferStride)
+		SHADER_PARAMETER(uint32,				NumInstances)
+
+		SHADER_PARAMETER_SRV(Buffer<int>,		CountBuffer)
+		SHADER_PARAMETER(uint32,				CountBufferOffset)
+
+		SHADER_PARAMETER(uint32,				NumPositions)
+		SHADER_PARAMETER_SCALAR_ARRAY(uint32,	PositionComponentOffsets, [MaxPositions])
+		SHADER_PARAMETER(FVector3f,				TileShift)
+	END_SHADER_PARAMETER_STRUCT()
+
+	static void ModifyCompilationEnvironment(const FGlobalShaderPermutationParameters& Parameters, FShaderCompilerEnvironment& OutEnvironment);
+	NIAGARASHADER_API static void Execute(FRHIComputeCommandList& RHICmdList, const FParameters& Parameters);
+};
diff --git a/Engine/Plugins/FX/Niagara/Source/NiagaraShader/Private/NiagaraLWCTileShiftPositions.cpp b/Engine/Plugins/FX/Niagara/Source/NiagaraShader/Private/NiagaraLWCTileShiftPositions.cpp
new file mode 100644
--- /dev/null
+++ b/Engine/Plugins/FX/Niagara/Source/NiagaraShader/Private/NiagaraLWCTileShiftPositions.cpp
@@ -0,0 +1,23 @@
+// Copyright Epic Games, Inc. All Rights Reserved.
+
+#include "NiagaraLWCTileShiftPositions.h"
+
+#include "RenderGraphUtils.h"
+
+IMPLEMENT_GLOBAL_SHADER(FNiagaraLWCTileShiftPositionsCS, "/Plugin/FX/Niagara/Private/NiagaraLWCHelper.usf", "TileShiftPositions", SF_Compute);
+
+void FNiagaraLWCTileShiftPositionsCS::ModifyCompilationEnvironment(const FGlobalShaderPermutationParameters& Parameters, FShaderCompilerEnvironment& OutEnvironment)
+{
+	FGlobalShader::ModifyCompilationEnvironment(Parameters, OutEnvironment);
+	OutEnvironment.SetDefine(TEXT("THREADGROUP_SIZE"), ThreadGroupSize);
+	OutEnvironment.SetDefine(TEXT("MaxPositions"), MaxPositions);
+}
+
+void FNiagaraLWCTileShiftPositionsCS::Execute(FRHIComputeCommandList& RHICmdList, const FParameters& Parameters)
+{
+	TShaderMapRef<FNiagaraLWCTileShiftPositionsCS> ComputeShader(GetGlobalShaderMap(GMaxRHIFeatureLevel));
+
+	const uint32 NumThreadGroups = FMath::DivideAndRoundUp<uint32>(Parameters.NumInstances, ThreadGroupSize);
+	const FIntVector NumWrappedThreadGroups = FComputeShaderUtils::GetGroupCountWrapped(NumThreadGroups);
+	FComputeShaderUtils::Dispatch(RHICmdList, ComputeShader, Parameters, NumWrappedThreadGroups);
+}
