Removed unused code from previous feature iteration

#rnx

diff --git a/Engine/Plugins/FX/Niagara/Source/NiagaraEditor/Private/Widgets/DataChannel/NiagaraDataChannelWizard.cpp b/Engine/Plugins/FX/Niagara/Source/NiagaraEditor/Private/Widgets/DataChannel/NiagaraDataChannelWizard.cpp
--- a/Engine/Plugins/FX/Niagara/Source/NiagaraEditor/Private/Widgets/DataChannel/NiagaraDataChannelWizard.cpp
+++ b/Engine/Plugins/FX/Niagara/Source/NiagaraEditor/Private/Widgets/DataChannel/NiagaraDataChannelWizard.cpp
@@ -959,16 +959,6 @@ TSharedRef<FModuleWizardModel> DataChannel::CreateSpawnNDCModuleWizardModel()
 					GraphSchema->TryCreateConnection(SpawnFunction->GetOutputPin(0), OutputNode->GetInputPin(0));
 					Utilities::SetDefaultBinding(Graph, DIPin->PinName, FName("Emitter.SpawnDataChannel"));
 
-					// add scalability inputs
-					UEdGraphPin* EmitterScalePin = nullptr;
-					UEdGraphPin* SpawnScalePin = nullptr;
-					if (AssetPage->Data->bModifySpawnCountByScalability)
-					{
-						EmitterScalePin = Utilities::AddReadParameterPin(FNiagaraTypeDefinition::GetFloatDef(), FName("Transient.ScalabilityEmitterSpawnCountScale"), MapGetNode);
-						Utilities::SetDefaultValue(Graph, EmitterScalePin->PinName, FNiagaraTypeDefinition::GetFloatDef(), 1.0f);
-						SpawnScalePin = Utilities::AddReadParameterPin(FNiagaraTypeDefinition::GetFloatDef(), FName("Engine.Emitter.SpawnCountScale"), MapGetNode);
-					}
-
 					// Add module inputs
 					if (UEdGraphPin* EnableInput = SpawnFunction->FindPin(FName("Enable"), EGPD_Input))
 					{
