- Fix use of unsafe raw object pointer

#rnx
#tests Use Niagara, the repro is unknown.

diff --git a/Engine/Plugins/FX/Niagara/Source/NiagaraEditor/Private/ViewModels/Stack/NiagaraStackGraphUtilities.cpp b/Engine/Plugins/FX/Niagara/Source/NiagaraEditor/Private/ViewModels/Stack/NiagaraStackGraphUtilities.cpp
--- a/Engine/Plugins/FX/Niagara/Source/NiagaraEditor/Private/ViewModels/Stack/NiagaraStackGraphUtilities.cpp
+++ b/Engine/Plugins/FX/Niagara/Source/NiagaraEditor/Private/ViewModels/Stack/NiagaraStackGraphUtilities.cpp
@@ -4292,9 +4292,10 @@ FGuid FNiagaraStackGraphUtilities::GetScriptVariableIdForLinkedModuleParameterHa
 
 bool FNiagaraStackGraphUtilities::DependencyUtilities::DoesStackModuleProvideDependency(const FNiagaraStackModuleData& StackModuleData, const FNiagaraModuleDependency& SourceModuleRequiredDependency, const UNiagaraNodeOutput& SourceOutputNode)
 {
-	if (StackModuleData.ModuleNode != nullptr && StackModuleData.ModuleNode->FunctionScript != nullptr)
+	UNiagaraNodeFunctionCall* ModuleNode = StackModuleData.WeakModuleNode.Get();
+	if (ModuleNode != nullptr && ModuleNode->FunctionScript != nullptr)
 	{
-		FVersionedNiagaraScriptData* ScriptData = StackModuleData.ModuleNode->FunctionScript->GetScriptData(StackModuleData.ModuleNode->SelectedScriptVersion);
+		FVersionedNiagaraScriptData* ScriptData = ModuleNode->FunctionScript->GetScriptData(ModuleNode->SelectedScriptVersion);
 		if (ScriptData && ScriptData->ProvidedDependencies.Contains(SourceModuleRequiredDependency.Id))
 		{
 			if (SourceModuleRequiredDependency.ScriptConstraint == ENiagaraModuleDependencyScriptConstraint::AllScripts)
@@ -4304,7 +4305,7 @@ bool FNiagaraStackGraphUtilities::DependencyUtilities::DoesStackModuleProvideDep
 
 			if (SourceModuleRequiredDependency.ScriptConstraint == ENiagaraModuleDependencyScriptConstraint::SameScript)
 			{
-				UNiagaraNodeOutput* OutputNode = FNiagaraStackGraphUtilities::GetEmitterOutputNodeForStackNode(*StackModuleData.ModuleNode);
+				UNiagaraNodeOutput* OutputNode = FNiagaraStackGraphUtilities::GetEmitterOutputNodeForStackNode(*ModuleNode);
 				return OutputNode != nullptr && UNiagaraScript::IsEquivalentUsage(OutputNode->GetUsage(), SourceOutputNode.GetUsage()) && OutputNode->GetUsageId() == SourceOutputNode.GetUsageId();
 			}
 		}
@@ -4530,7 +4531,9 @@ int32 FNiagaraStackGraphUtilities::DependencyUtilities::FindBestIndexForModuleIn
 		{
 			return *OutputNodePtr;
 		}
-		UNiagaraNodeOutput* OutputNode = FNiagaraStackGraphUtilities::GetEmitterOutputNodeForStackNode(*StackModuleData.ModuleNode);
+		UNiagaraNodeFunctionCall* ModuleNode = StackModuleData.WeakModuleNode.Get();
+		check(ModuleNode != nullptr);
+		UNiagaraNodeOutput* OutputNode = FNiagaraStackGraphUtilities::GetEmitterOutputNodeForStackNode(*ModuleNode);
 		ScriptUsageToOutputNode.Add(StackModuleData.Usage, OutputNode);
 		return OutputNode;
 	};
@@ -4585,17 +4588,22 @@ int32 FNiagaraStackGraphUtilities::DependencyUtilities::FindBestIndexForModuleIn
 	bool bProvidedDependencyUpperBoundFound = false;
 
 	const TArray<FName>& NewModuleScriptProvidedDependencies = ScriptData->ProvidedDependencies;
-	auto GetStackModuleDataRequiredDependenciesBeingProvided = [&NewModuleScriptProvidedDependencies](const FNiagaraStackModuleData& StackModuleData)->const TArray<FNiagaraModuleDependency> /*OutDependencies*/ {
-		TArray<FNiagaraModuleDependency> OutDependencies;
-		for (FNiagaraModuleDependency& RequiredDependency : StackModuleData.ModuleNode->GetScriptData()->RequiredDependencies)
+	auto GetStackModuleDataRequiredDependenciesBeingProvided =
+		[&NewModuleScriptProvidedDependencies](const FNiagaraStackModuleData& StackModuleData)->const TArray<FNiagaraModuleDependency> /*OutDependencies*/
 		{
-			if (NewModuleScriptProvidedDependencies.Contains(RequiredDependency.Id))
+			TArray<FNiagaraModuleDependency> OutDependencies;
+			if (UNiagaraNodeFunctionCall* ModuleNode = StackModuleData.WeakModuleNode.Get())
 			{
-				OutDependencies.Add(RequiredDependency);
+				for (FNiagaraModuleDependency& RequiredDependency : ModuleNode->GetScriptData()->RequiredDependencies)
+				{
+					if (NewModuleScriptProvidedDependencies.Contains(RequiredDependency.Id))
+					{
+						OutDependencies.Add(RequiredDependency);
+					}
+				}
 			}
-		}
-		return OutDependencies;
-	};
+			return OutDependencies;
+		};
 
 	for (const FNiagaraStackModuleData& CurrentStackModuleData : StackModuleData)
 	{
diff --git a/Engine/Plugins/FX/Niagara/Source/NiagaraEditor/Private/ViewModels/Stack/NiagaraStackModuleItem.cpp b/Engine/Plugins/FX/Niagara/Source/NiagaraEditor/Private/ViewModels/Stack/NiagaraStackModuleItem.cpp
--- a/Engine/Plugins/FX/Niagara/Source/NiagaraEditor/Private/ViewModels/Stack/NiagaraStackModuleItem.cpp
+++ b/Engine/Plugins/FX/Niagara/Source/NiagaraEditor/Private/ViewModels/Stack/NiagaraStackModuleItem.cpp
@@ -347,8 +347,9 @@ namespace NiagaraStackModuleItemIssues {
 				break;
 			}
 
-			if (CurrentStackModuleData.ModuleNode->GetScriptData() != nullptr &&
-				CurrentStackModuleData.ModuleNode->GetScriptData()->RequiredDependencies.ContainsByPredicate(
+			UNiagaraNodeFunctionCall* ModuleNode = CurrentStackModuleData.WeakModuleNode.Get();
+			if (ModuleNode && ModuleNode->GetScriptData() != nullptr &&
+				ModuleNode->GetScriptData()->RequiredDependencies.ContainsByPredicate(
 					[&RequiredDependency](const FNiagaraModuleDependency& CurrentRequiredDependency) { return CurrentRequiredDependency.Id == RequiredDependency.Id; }))
 			{
 				LastModuleDataIndexRequiringDependency = StackModuleDataIndex;
@@ -381,8 +382,11 @@ namespace NiagaraStackModuleItemIssues {
 		if (UNiagaraScript::ContainsEquivilentUsage(CompatibleUsages, LastDependentModuleData.Usage))
 		{
 			// If the dependency provider is compatible with the last dependent usage it can be added directly before or after the last dependent.
-			OutTargetOutputNode = FNiagaraStackGraphUtilities::GetEmitterOutputNodeForStackNode(*LastDependentModuleData.ModuleNode);
-			OutTargetIndex = LastDependentModuleData.Index + (RequiredDependency.Type == ENiagaraModuleDependencyType::PostDependency ? 1 : 0);
+			if ( UNiagaraNodeFunctionCall* ModuleNode = LastDependentModuleData.WeakModuleNode.Get() )
+			{
+				OutTargetOutputNode = FNiagaraStackGraphUtilities::GetEmitterOutputNodeForStackNode(*ModuleNode);
+				OutTargetIndex = LastDependentModuleData.Index + (RequiredDependency.Type == ENiagaraModuleDependencyType::PostDependency ? 1 : 0);
+			}
 		}
 		else
 		{
@@ -448,7 +452,7 @@ namespace NiagaraStackModuleItemIssues {
 		}
 
 		const TArray<FNiagaraStackModuleData>& StackModuleData = DependentSystemViewModel->GetStackModuleDataByEmitterHandleId(DependentEmitterHandleId);
-		int32 DependentModuleIndex = StackModuleData.IndexOfByPredicate([DependentModuleNode](const FNiagaraStackModuleData& StackModuleDataItem) { return StackModuleDataItem.ModuleNode == DependentModuleNode; });
+		int32 DependentModuleIndex = StackModuleData.IndexOfByPredicate([DependentModuleNode](const FNiagaraStackModuleData& StackModuleDataItem) { return StackModuleDataItem.WeakModuleNode.Get() == DependentModuleNode; });
 		
 		UNiagaraNodeOutput* TargetOutputNode = nullptr;
 		TOptional<int32> TargetIndex;
@@ -615,10 +619,14 @@ namespace NiagaraStackModuleItemIssues {
 
 			if (TargetIndex != INDEX_NONE)
 			{
-				FText FixDescription = FText::Format(LOCTEXT("MoveDependentFixDescriptionFormat", "Move module {0} {1} {2}"), FText::FromString(DependentModuleNode.GetFunctionName()), LocationText, FText::FromString(CurrentStackModuleData.ModuleNode->GetFunctionName()));
-				OutFixes.Add(UNiagaraStackEntry::FStackIssueFix(FixDescription, UNiagaraStackEntry::FStackIssueFixDelegate::CreateStatic(&MoveModuleToFixDependencyIssue,
-					TWeakObjectPtr<UNiagaraSystem>(&DependentSystemViewModel->GetSystem()), DependentEmitterHandleId, TWeakObjectPtr<UNiagaraScript>(&DependentScript), TWeakObjectPtr<UNiagaraNodeFunctionCall>(&DependentModuleNode),
-					CurrentStackModuleData.Usage, CurrentStackModuleData.UsageId, TargetIndex)));
+				UNiagaraNodeFunctionCall* ModuleNode = CurrentStackModuleData.WeakModuleNode.Get();
+				if (ensure(ModuleNode))
+				{
+					FText FixDescription = FText::Format(LOCTEXT("MoveDependentFixDescriptionFormat", "Move module {0} {1} {2}"), FText::FromString(DependentModuleNode.GetFunctionName()), LocationText, FText::FromString(ModuleNode->GetFunctionName()));
+					OutFixes.Add(UNiagaraStackEntry::FStackIssueFix(FixDescription, UNiagaraStackEntry::FStackIssueFixDelegate::CreateStatic(&MoveModuleToFixDependencyIssue,
+						TWeakObjectPtr<UNiagaraSystem>(&DependentSystemViewModel->GetSystem()), DependentEmitterHandleId, TWeakObjectPtr<UNiagaraScript>(&DependentScript), TWeakObjectPtr<UNiagaraNodeFunctionCall>(&DependentModuleNode),
+						CurrentStackModuleData.Usage, CurrentStackModuleData.UsageId, TargetIndex)));
+				}
 			}
 		}
 	}
@@ -675,7 +683,7 @@ namespace NiagaraStackModuleItemIssues {
 		for (int32 WrongVersionDependencyProviderIndex : WrongVersionDependencyProviderIndices)
 		{
 			const FNiagaraStackModuleData& CurrentStackModuleData = StackModuleData[WrongVersionDependencyProviderIndex];
-			UNiagaraNodeFunctionCall* ModuleNode = CurrentStackModuleData.ModuleNode;
+			UNiagaraNodeFunctionCall* ModuleNode = CurrentStackModuleData.WeakModuleNode.Get();
 			
 			if (FindStackModuleItem(SourceSystemViewModel, SourceEmitterHandleId, ModuleNode) == nullptr)
 			{
@@ -730,17 +738,26 @@ namespace NiagaraStackModuleItemIssues {
 		for (int32 DisabledDependencyProviderIndex : DisabledDependencyProviderIndices)
 		{
 			const FNiagaraStackModuleData& CurrentStackModuleData = StackModuleData[DisabledDependencyProviderIndex];
-			FText FixDescription = FText::Format(LOCTEXT("EnableDependencyProviderFixDescriptionFormat", "Enable module {0} which provides the dependency."),
-				FText::FromString(CurrentStackModuleData.ModuleNode->GetFunctionName()));
-			OutFixes.Add(UNiagaraStackEntry::FStackIssueFix(FixDescription, UNiagaraStackEntry::FStackIssueFixDelegate::CreateStatic(&EnableModuleToFixDependencyIssue,
-				TWeakObjectPtr<UNiagaraNodeFunctionCall>(CurrentStackModuleData.ModuleNode))));
+			UNiagaraNodeFunctionCall* ModuleNode = CurrentStackModuleData.WeakModuleNode.Get();
+			if (ensure(ModuleNode))
+			{
+				FText FixDescription = FText::Format(LOCTEXT("EnableDependencyProviderFixDescriptionFormat", "Enable module {0} which provides the dependency."),
+					FText::FromString(ModuleNode->GetFunctionName()));
+				OutFixes.Add(UNiagaraStackEntry::FStackIssueFix(FixDescription, UNiagaraStackEntry::FStackIssueFixDelegate::CreateStatic(&EnableModuleToFixDependencyIssue,
+					TWeakObjectPtr<UNiagaraNodeFunctionCall>(ModuleNode))));
+			}
 		}
 	}
 
 	bool IsCorrectDependencyVersion(const FNiagaraStackModuleData& StackModuleData, const FNiagaraModuleDependency& SourceModuleRequiredDependency)
 	{
-		FVersionedNiagaraScriptData* ScriptData = StackModuleData.ModuleNode->FunctionScript->GetScriptData(StackModuleData.ModuleNode->SelectedScriptVersion);
-		return SourceModuleRequiredDependency.IsVersionAllowed(ScriptData->Version);
+		UNiagaraNodeFunctionCall* ModuleNode = StackModuleData.WeakModuleNode.Get();
+		if (ensure(ModuleNode))
+		{
+			FVersionedNiagaraScriptData* ScriptData = ModuleNode->FunctionScript->GetScriptData(ModuleNode->SelectedScriptVersion);
+			return SourceModuleRequiredDependency.IsVersionAllowed(ScriptData->Version);
+		}
+		return false;
 	}
 
 	void GenerateDependencyIssues(
@@ -759,7 +776,7 @@ namespace NiagaraStackModuleItemIssues {
 		}
 
 		FVersionedNiagaraScriptData* ScriptData = SourceModuleNode.FunctionScript->GetScriptData(SourceModuleNode.SelectedScriptVersion);
-		int32 ModuleIndex = SourceStackModuleData.IndexOfByPredicate([&SourceModuleNode](const FNiagaraStackModuleData& ModuleData) { return ModuleData.ModuleNode == &SourceModuleNode; });
+		int32 ModuleIndex = SourceStackModuleData.IndexOfByPredicate([&SourceModuleNode](const FNiagaraStackModuleData& ModuleData) { return ModuleData.WeakModuleNode.Get() == &SourceModuleNode; });
 		if (ensureMsgf(ModuleIndex != INDEX_NONE, TEXT("In system %s, module %s (%s) did not exist in the stack module data."),
 			*SourceSystemViewModel->GetSystem().GetPathName(), *SourceModuleNode.GetFunctionName(), *SourceModuleNode.GetName()))
 		{
@@ -787,11 +804,12 @@ namespace NiagaraStackModuleItemIssues {
 				for (int32 DependencyProviderIndex : DependencyProviderIndices)
 				{
 					const FNiagaraStackModuleData& DependencyProviderData = SourceStackModuleData[DependencyProviderIndex];
+					UNiagaraNodeFunctionCall* ModuleNode = DependencyProviderData.WeakModuleNode.Get();
 					bool bCorrectOrder =
 						(SourceRequiredDependency.Type == ENiagaraModuleDependencyType::PreDependency && DependencyProviderIndex < ModuleIndex) ||
 						(SourceRequiredDependency.Type == ENiagaraModuleDependencyType::PostDependency && DependencyProviderIndex > ModuleIndex);
 					bool bCorrectVersion = IsCorrectDependencyVersion(DependencyProviderData, SourceRequiredDependency);
-					bool bEnabled = DependencyProviderData.ModuleNode->GetDesiredEnabledState() == ENodeEnabledState::Enabled;
+					bool bEnabled = ModuleNode && ModuleNode->GetDesiredEnabledState() == ENodeEnabledState::Enabled;
 					bool bUsageIsSupported = UNiagaraScript::ContainsEquivilentUsage(SupportedUsages, DependencyProviderData.Usage);
 
 					if (bEnabled && bCorrectOrder && bCorrectVersion)
diff --git a/Engine/Plugins/FX/Niagara/Source/NiagaraEditor/Public/ViewModels/NiagaraSystemViewModel.h b/Engine/Plugins/FX/Niagara/Source/NiagaraEditor/Public/ViewModels/NiagaraSystemViewModel.h
--- a/Engine/Plugins/FX/Niagara/Source/NiagaraEditor/Public/ViewModels/NiagaraSystemViewModel.h
+++ b/Engine/Plugins/FX/Niagara/Source/NiagaraEditor/Public/ViewModels/NiagaraSystemViewModel.h
@@ -103,7 +103,7 @@ struct FNiagaraSystemViewModelOptions
 
 struct FNiagaraStackModuleData
 {
-	UNiagaraNodeFunctionCall* ModuleNode;
+	TWeakObjectPtr<UNiagaraNodeFunctionCall> WeakModuleNode;
 	ENiagaraScriptUsage Usage;
 	FGuid UsageId;
 	int32 Index;
