- Fix RemoveEmitterHandlesById not calling RemoveSystemParametersForEmitter
- This could result in dangling parameters being used (OOB memory read)

diff --git a/Engine/Plugins/FX/Niagara/Source/Niagara/Private/NiagaraSystem.cpp b/Engine/Plugins/FX/Niagara/Source/Niagara/Private/NiagaraSystem.cpp
--- a/Engine/Plugins/FX/Niagara/Source/Niagara/Private/NiagaraSystem.cpp
+++ b/Engine/Plugins/FX/Niagara/Source/Niagara/Private/NiagaraSystem.cpp
@@ -2850,11 +2850,16 @@ void UNiagaraSystem::RemoveEmitterHandle(const FNiagaraEmitterHandle& EmitterHan
 
 void UNiagaraSystem::RemoveEmitterHandlesById(const TSet<FGuid>& HandlesToRemove)
 {
-	auto RemovePredicate = [&](const FNiagaraEmitterHandle& EmitterHandle)
+	for (auto It=EmitterHandles.CreateIterator(); It; ++It)
 	{
-		return HandlesToRemove.Contains(EmitterHandle.GetId());
-	};
-	EmitterHandles.RemoveAll(RemovePredicate);
+		const FNiagaraEmitterHandle& EmitterHandle = *It;
+		if (!HandlesToRemove.Contains(EmitterHandle.GetId()))
+		{
+			continue;
+		}
+		RemoveSystemParametersForEmitter(EmitterHandle);
+		It.RemoveCurrent();
+	}
 
 	InitEmitterCompiledData();
 	InvalidateCachedData();
