- Add customization for stateless emitters
- Added an edit condition to disable fixed bounds when the system has fixed bounds enabled.  This matches stateful.

#rnx

diff --git a/Engine/Plugins/FX/Niagara/Source/NiagaraEditor/Private/Stateless/NiagaraStatelessEmitterDetails.cpp b/Engine/Plugins/FX/Niagara/Source/NiagaraEditor/Private/Stateless/NiagaraStatelessEmitterDetails.cpp
new file mode 100644
--- /dev/null
+++ b/Engine/Plugins/FX/Niagara/Source/NiagaraEditor/Private/Stateless/NiagaraStatelessEmitterDetails.cpp
@@ -0,0 +1,69 @@
+// Copyright Epic Games, Inc. All Rights Reserved.
+
+#include "Stateless/NiagaraStatelessEmitterDetails.h"
+#include "Stateless/NiagaraStatelessEmitter.h"
+#include "NiagaraSystem.h"
+
+#include "DetailLayoutBuilder.h"
+#include "DetailCategoryBuilder.h"
+//#include "IDetailChildrenBuilder.h"
+//#include "IDetailGroup.h"
+//#include "DetailWidgetRow.h"
+
+#define LOCTEXT_NAMESPACE "NiagaraStatelessEmitter"
+
+TSharedRef<IDetailCustomization> FNiagaraStatelessEmitterDetails::MakeInstance()
+{
+	return MakeShared<FNiagaraStatelessEmitterDetails>();
+}
+
+void FNiagaraStatelessEmitterDetails::CustomizeDetails(IDetailLayoutBuilder& DetailBuilder)
+{
+	UNiagaraStatelessEmitter* Emitter = nullptr;
+	{
+		TArray<TWeakObjectPtr<UObject>> ObjectsBeingCustomized;
+		DetailBuilder.GetObjectsBeingCustomized(ObjectsBeingCustomized);
+		Emitter = ObjectsBeingCustomized.Num() == 1 ? Cast<UNiagaraStatelessEmitter>(ObjectsBeingCustomized[0]) : nullptr;
+	}
+	if (Emitter == nullptr)
+	{
+		return;
+	}
+	WeakEmitter = Emitter;
+
+	static const FName NAME_EmitterProperties("Emitter Properties");
+	static const FName NAME_FixedBounds("FixedBounds");
+
+	IDetailCategoryBuilder& DetailCategory = DetailBuilder.EditCategory(NAME_EmitterProperties);
+
+	TArray<TSharedRef<IPropertyHandle>> CategoryProperties;
+	DetailCategory.GetDefaultProperties(CategoryProperties);
+
+	for (TSharedRef<IPropertyHandle>& PropertyHandle : CategoryProperties)
+	{
+		const FName PropertyName = PropertyHandle->GetProperty() ? PropertyHandle->GetProperty()->GetFName() : NAME_None;
+		if (PropertyName == NAME_FixedBounds)
+		{
+			IDetailPropertyRow& PropertyRow = DetailCategory.AddProperty(PropertyHandle);
+			PropertyRow.IsEnabled(TAttribute<bool>::CreateSP(this, &FNiagaraStatelessEmitterDetails::GetFixedBoundsEnabled));
+		}
+		else
+		{
+			DetailCategory.AddProperty(PropertyHandle);
+		}
+	}
+}
+
+bool FNiagaraStatelessEmitterDetails::GetFixedBoundsEnabled() const
+{
+	if (const UNiagaraStatelessEmitter* Emitter = WeakEmitter.Get())
+	{
+		if (const UNiagaraSystem* System = Emitter->GetTypedOuter<UNiagaraSystem>())
+		{
+			return System->bFixedBounds == false;
+		}
+	}
+	return false;
+}
+
+#undef LOCTEXT_NAMESPACE
diff --git a/Engine/Plugins/FX/Niagara/Source/NiagaraEditor/Private/Stateless/NiagaraStatelessEmitterDetails.h b/Engine/Plugins/FX/Niagara/Source/NiagaraEditor/Private/Stateless/NiagaraStatelessEmitterDetails.h
new file mode 100644
--- /dev/null
+++ b/Engine/Plugins/FX/Niagara/Source/NiagaraEditor/Private/Stateless/NiagaraStatelessEmitterDetails.h
@@ -0,0 +1,21 @@
+// Copyright Epic Games, Inc. All Rights Reserved.
+
+#pragma once
+
+#include "IDetailCustomization.h"
+#include "UObject/WeakObjectPtr.h"
+
+class IDetailLayoutBuilder;
+class UNiagaraStatelessEmitter;
+
+class FNiagaraStatelessEmitterDetails : public IDetailCustomization
+{
+public:
+	static TSharedRef<IDetailCustomization> MakeInstance();
+	virtual void CustomizeDetails(IDetailLayoutBuilder& DetailBuilder) override;
+
+	bool GetFixedBoundsEnabled() const;
+
+protected:
+	TWeakObjectPtr<UNiagaraStatelessEmitter> WeakEmitter;
+};
diff --git a/Engine/Plugins/FX/Niagara/Source/NiagaraEditor/Private/ViewModels/Stack/NiagaraStackStatelessEmitterGroup.cpp b/Engine/Plugins/FX/Niagara/Source/NiagaraEditor/Private/ViewModels/Stack/NiagaraStackStatelessEmitterGroup.cpp
--- a/Engine/Plugins/FX/Niagara/Source/NiagaraEditor/Private/ViewModels/Stack/NiagaraStackStatelessEmitterGroup.cpp
+++ b/Engine/Plugins/FX/Niagara/Source/NiagaraEditor/Private/ViewModels/Stack/NiagaraStackStatelessEmitterGroup.cpp
@@ -4,6 +4,7 @@
 
 #include "IDetailTreeNode.h"
 #include "PropertyHandle.h"
+#include "Stateless/NiagaraStatelessEmitterDetails.h"
 #include "Stateless/NiagaraEmitterStatePropertyCustomization.h"
 #include "Stateless/NiagaraDistributionPropertyCustomization.h"
 #include "Stateless/NiagaraDistributionIntPropertyCustomization.h"
@@ -130,6 +131,7 @@ void UNiagaraStackStatelessEmitterObjectItem::RefreshChildrenInternal(const TArr
 			bool bHideTopLevelCategories = false;
 			StatelessEmitterStackObject->Initialize(CreateDefaultChildRequiredData(), StatelessEmitter, bIsTopLevelObject, bHideTopLevelCategories, GetStackEditorDataKey());
 			StatelessEmitterStackObject->SetOnFilterDetailNodes(OnFilterDetailNodes, UNiagaraStackObject::EDetailNodeFilterMode::FilterAllNodes);
+			StatelessEmitterStackObject->RegisterInstancedCustomPropertyLayout(UNiagaraStatelessEmitter::StaticClass(), FOnGetDetailCustomizationInstance::CreateStatic(&FNiagaraStatelessEmitterDetails::MakeInstance));
 			StatelessEmitterStackObject->RegisterInstancedCustomPropertyTypeLayout(FNiagaraEmitterStateData::StaticStruct()->GetFName(), FOnGetPropertyTypeCustomizationInstance::CreateStatic(&FNiagaraEmitterStatePropertyCustomization::MakeInstance));
 			StatelessEmitterStackObject->RegisterInstancedCustomPropertyTypeLayout(FNiagaraDistributionFloat::StaticStruct()->GetFName(), FOnGetPropertyTypeCustomizationInstance::CreateStatic(&FNiagaraDistributionPropertyCustomization::MakeFloatInstance));
 			StatelessEmitterStackObject->RegisterInstancedCustomPropertyTypeLayout(FNiagaraDistributionVector2::StaticStruct()->GetFName(), FOnGetPropertyTypeCustomizationInstance::CreateStatic(&FNiagaraDistributionPropertyCustomization::MakeVector2Instance));
