- Light renderer support for stateless
- This will essentially run the stateless simulation on the CPU (always) like a regular emitter
- Most other CPU only renderer types should be able to run via this path
- When the light renderer is updated we can likely remove this path

diff --git a/Engine/Plugins/FX/Niagara/Source/Niagara/Internal/Stateless/NiagaraStatelessEmitterData.h b/Engine/Plugins/FX/Niagara/Source/Niagara/Internal/Stateless/NiagaraStatelessEmitterData.h
--- a/Engine/Plugins/FX/Niagara/Source/Niagara/Internal/Stateless/NiagaraStatelessEmitterData.h
+++ b/Engine/Plugins/FX/Niagara/Source/Niagara/Internal/Stateless/NiagaraStatelessEmitterData.h
@@ -38,6 +38,7 @@ struct FNiagaraStatelessEmitterData
 
 	bool											bCanEverExecute = false;
 	bool											bDeterministic = false;
+	ENiagaraSimTarget								SimTarget = ENiagaraSimTarget::GPUComputeSim;
 	ENiagaraStatelessFeatureMask					FeatureMask = ENiagaraStatelessFeatureMask::None;
 	int32											RandomSeed = 0;
 	FNiagaraStatelessRangeFloat						LifetimeRange = FNiagaraStatelessRangeFloat(0.0f, 0.0f);
diff --git a/Engine/Plugins/FX/Niagara/Source/Niagara/Internal/Stateless/NiagaraStatelessEmitterInstance.h b/Engine/Plugins/FX/Niagara/Source/Niagara/Internal/Stateless/NiagaraStatelessEmitterInstance.h
--- a/Engine/Plugins/FX/Niagara/Source/Niagara/Internal/Stateless/NiagaraStatelessEmitterInstance.h
+++ b/Engine/Plugins/FX/Niagara/Source/Niagara/Internal/Stateless/NiagaraStatelessEmitterInstance.h
@@ -20,6 +20,8 @@ namespace NiagaraStateless
 		}
 	};
 
+	using FShaderParametersPtr = TUniquePtr<NiagaraStateless::FCommonShaderParameters, FShaderParametersDeletor>;
+
 	class FEmitterInstance_RT : public INiagaraComputeDataBufferInterface
 	{
 	public:
@@ -33,7 +35,7 @@ namespace NiagaraStateless
 		float													DeltaTime = 0.0f;		//-OPT: We should be able to pull this from the view which is important for zero tick path
 		ENiagaraExecutionState									ExecutionState = ENiagaraExecutionState::Active;
 		TArray<FNiagaraStatelessRuntimeSpawnInfo>				SpawnInfos;
-		TUniquePtr<NiagaraStateless::FCommonShaderParameters, FShaderParametersDeletor>	ShaderParameters;
+		FShaderParametersPtr									ShaderParameters;
 
 		mutable bool											bBindingBufferDirty = true;
 		mutable TArray<uint8>									BindingBufferData;
@@ -84,15 +86,13 @@ public:
 
 	const FNiagaraStatelessEmitterData* GetEmitterData() const { check(EmitterData.IsValid()); return EmitterData.Get(); }
 
-	const NiagaraStateless::FEmitterInstance_RT* GetRenderInstance() const { return RenderThreadDataPtr.Get(); }
-
 private:
 	void InitEmitterData();
 
 	void InitEmitterState();
 	void TickEmitterState();
 	void CalculateBounds();
-	void SendRenderData();
+	void UpdateSimulationData(float DeltaSeconds);
 
 	void InitSpawnInfos(float InitializationAge);
 	void InitSpawnInfosForLoop(float InitializationAge);
@@ -137,4 +137,6 @@ private:
 	float										CurrentLoopAgeEnd = 0.0f;
 
 	TUniquePtr<NiagaraStateless::FEmitterInstance_RT>	RenderThreadDataPtr;
+
+	NiagaraStateless::FShaderParametersPtr			ShaderParameters;
 };
diff --git a/Engine/Plugins/FX/Niagara/Source/Niagara/Private/Stateless/NiagaraStatelessEmitter.cpp b/Engine/Plugins/FX/Niagara/Source/Niagara/Private/Stateless/NiagaraStatelessEmitter.cpp
--- a/Engine/Plugins/FX/Niagara/Source/Niagara/Private/Stateless/NiagaraStatelessEmitter.cpp
+++ b/Engine/Plugins/FX/Niagara/Source/Niagara/Private/Stateless/NiagaraStatelessEmitter.cpp
@@ -455,8 +455,18 @@ void UNiagaraStatelessEmitter::CacheFromCompiledData()
 			if (Renderer && Renderer->GetIsEnabled())
 			{
 				Renderer->PopulateRequiredBindings(StatelessEmitterData->RendererBindings);
+				if (Renderer->IsSimTargetSupported(ENiagaraSimTarget::GPUComputeSim) == false)
+				{
+					StatelessEmitterData->SimTarget = ENiagaraSimTarget::CPUSim;
+				}
 			}
 		}
+
+		// If the simulation target was modified to CPU then we might not be able to execute
+		if (StatelessEmitterData->SimTarget == ENiagaraSimTarget::CPUSim && !EnumHasAnyFlags(StatelessEmitterData->FeatureMask, ENiagaraStatelessFeatureMask::ExecuteCPU))
+		{
+			StatelessEmitterData->bCanEverExecute = false;
+		}
 	}
 }
 
diff --git a/Engine/Plugins/FX/Niagara/Source/Niagara/Private/Stateless/NiagaraStatelessEmitterInstance.cpp b/Engine/Plugins/FX/Niagara/Source/Niagara/Private/Stateless/NiagaraStatelessEmitterInstance.cpp
--- a/Engine/Plugins/FX/Niagara/Source/Niagara/Private/Stateless/NiagaraStatelessEmitterInstance.cpp
+++ b/Engine/Plugins/FX/Niagara/Source/Niagara/Private/Stateless/NiagaraStatelessEmitterInstance.cpp
@@ -5,20 +5,13 @@
 #include "Stateless/NiagaraStatelessEmitter.h"
 #include "Stateless/NiagaraStatelessEmitterData.h"
 #include "Stateless/NiagaraStatelessExpression.h"
+#include "Stateless/NiagaraStatelessParticleSimContext.h"
 #include "NiagaraGpuComputeDispatchInterface.h"
 #include "NiagaraSystem.h"
 #include "NiagaraSystemInstance.h"
 
 #include "Shader.h"
 
-static bool bPerfTest = false;
-static FAutoConsoleVariableRef CVarTest(
-	TEXT("fx.NiagaraStateless.PerfTest"),
-	bPerfTest,
-	TEXT(""),
-	ECVF_Scalability | ECVF_Default
-);
-
 namespace NiagaraStatelessEmitterInstancePrivate
 {
 	static const float DefaultLoopDuration		= 0.001f;
@@ -132,6 +125,8 @@ void FNiagaraStatelessEmitterInstance::Init(int32 InEmitterIndex)
 	}
 
 	// Pull out information
+	SimTarget = EmitterData->SimTarget;
+
 	RandomSeed = EmitterData->RandomSeed + ParentSystemInstance->GetRandomSeedOffset();
 	if (!EmitterData->bDeterministic)
 	{
@@ -149,23 +144,31 @@ void FNiagaraStatelessEmitterInstance::Init(int32 InEmitterIndex)
 	EmitterTransforms.InitializeTransforms(IsLocalSpace(), FTransform3f(ParentSystemInstance->GetWorldTransform()));
 
 	// Allocate and fill shader parameters
-	RenderThreadDataPtr.Reset(new NiagaraStateless::FEmitterInstance_RT());
-	RenderThreadDataPtr->EmitterData = EmitterData;
-	RenderThreadDataPtr->RandomSeed = RandomSeed;
-	RenderThreadDataPtr->Age = 0.0f;
-	RenderThreadDataPtr->DeltaTime = 0.0f;
-	RenderThreadDataPtr->ExecutionState = ENiagaraExecutionState::Active;
-	RenderThreadDataPtr->ShaderParameters.Reset(WeakStatelessEmitter->AllocateShaderParameters(EmitterTransforms, RendererBindings));
-	RenderThreadDataPtr->ShaderParameters->Common_RandomSeed = RandomSeed;
-
-	ENQUEUE_RENDER_COMMAND(FInitStatelessEmitter)(
-		[RenderThreadData=RenderThreadDataPtr.Get(), ComputeInterface = ParentSystemInstance->GetComputeDispatchInterface()](FRHICommandListImmediate& RHICmdList)
-		{
-			RenderThreadData->ComputeManager = &ComputeInterface->GetOrCreateDataManager<FNiagaraStatelessComputeManager>();
-		}
-	);
+	if (SimTarget == ENiagaraSimTarget::CPUSim)
+	{
+		ShaderParameters.Reset(WeakStatelessEmitter->AllocateShaderParameters(EmitterTransforms, RendererBindings));
+		ShaderParameters->Common_RandomSeed = RandomSeed;
+	}
+	else
+	{
+		RenderThreadDataPtr.Reset(new NiagaraStateless::FEmitterInstance_RT());
+		RenderThreadDataPtr->EmitterData = EmitterData;
+		RenderThreadDataPtr->RandomSeed = RandomSeed;
+		RenderThreadDataPtr->Age = 0.0f;
+		RenderThreadDataPtr->DeltaTime = 0.0f;
+		RenderThreadDataPtr->ExecutionState = ENiagaraExecutionState::Active;
+		RenderThreadDataPtr->ShaderParameters.Reset(WeakStatelessEmitter->AllocateShaderParameters(EmitterTransforms, RendererBindings));
+		RenderThreadDataPtr->ShaderParameters->Common_RandomSeed = RandomSeed;
+
+		ENQUEUE_RENDER_COMMAND(FInitStatelessEmitter)(
+			[RenderThreadData=RenderThreadDataPtr.Get(), ComputeInterface = ParentSystemInstance->GetComputeDispatchInterface()](FRHICommandListImmediate& RHICmdList)
+			{
+				RenderThreadData->ComputeManager = &ComputeInterface->GetOrCreateDataManager<FNiagaraStatelessComputeManager>();
+			}
+		);
 
-	GPUDataBufferInterfaces = RenderThreadDataPtr.Get();
+		GPUDataBufferInterfaces = RenderThreadDataPtr.Get();
+	}
 
 	bNeedsEmitterStateInit = true;
 }
@@ -208,21 +211,32 @@ void FNiagaraStatelessEmitterInstance::ResetSimulation(bool bKillExisting)
 	InternalExecutionState = ENiagaraExecutionState::Active;
 	ExecutionState = InternalExecutionState;
 	ScalabilityState = ENiagaraExecutionStateManagement::Awaken;
-	if (NiagaraStateless::FEmitterInstance_RT* RenderThreadData = RenderThreadDataPtr.Get())
+
+	if (SimTarget == ENiagaraSimTarget::CPUSim)
 	{
-		ENQUEUE_RENDER_COMMAND(UpdateStatelessAge)(
-			[RenderThreadData, RandomSeed_RT=RandomSeed](FRHICommandListImmediate& RHICmdList)
-			{
-				RenderThreadData->Age = 0.0f;
-				RenderThreadData->DeltaTime = 0.0f;
-				RenderThreadData->ExecutionState = ENiagaraExecutionState::Active;
-				RenderThreadData->RandomSeed = RandomSeed_RT;
-				if (RenderThreadData->ShaderParameters.IsValid())
+		if (ShaderParameters.IsValid())
+		{
+			ShaderParameters->Common_RandomSeed = RandomSeed;
+		}
+	}
+	else
+	{
+		if (NiagaraStateless::FEmitterInstance_RT* RenderThreadData = RenderThreadDataPtr.Get())
+		{
+			ENQUEUE_RENDER_COMMAND(UpdateStatelessAge)(
+				[RenderThreadData, RandomSeed_RT=RandomSeed](FRHICommandListImmediate& RHICmdList)
 				{
-					RenderThreadData->ShaderParameters->Common_RandomSeed = RandomSeed_RT;
+					RenderThreadData->Age = 0.0f;
+					RenderThreadData->DeltaTime = 0.0f;
+					RenderThreadData->ExecutionState = ENiagaraExecutionState::Active;
+					RenderThreadData->RandomSeed = RandomSeed_RT;
+					if (RenderThreadData->ShaderParameters.IsValid())
+					{
+						RenderThreadData->ShaderParameters->Common_RandomSeed = RandomSeed_RT;
+					}
 				}
-			}
-		);
+			);
+		}
 	}
 }
 
@@ -321,7 +335,7 @@ void FNiagaraStatelessEmitterInstance::Tick(float DeltaSeconds)
 	TickSpawnInfos();
 	TickEmitterState();
 	CalculateBounds();
-	SendRenderData();
+	UpdateSimulationData(DeltaSeconds);
 }
 
 void FNiagaraStatelessEmitterInstance::InitEmitterState()
@@ -519,33 +533,13 @@ void FNiagaraStatelessEmitterInstance::CalculateBounds()
 	}
 }
 
-void FNiagaraStatelessEmitterInstance::SendRenderData()
+void FNiagaraStatelessEmitterInstance::UpdateSimulationData(float DeltaSeconds)
 {
-	NiagaraStateless::FEmitterInstance_RT* RenderThreadData = RenderThreadDataPtr.Get();
-	if (RenderThreadData == nullptr)
-	{
-		return;
-	}
-
-	struct FDataForRenderThread
-	{
-		float Age = 0.0f;
-		ENiagaraExecutionState ExecutionState = ENiagaraExecutionState::Disabled;
-
-		NiagaraStateless::FCommonShaderParameters* ShaderParameters = nullptr;
-
-		bool bHasBindingBufferData = false;
-		TArray<uint8> BindingBufferData;
-
-		bool bHasSpawnInfoData = false;
-		TArray<FNiagaraStatelessRuntimeSpawnInfo>	SpawnInfos;
-	};
-
-	FDataForRenderThread DataForRenderThread;
-	DataForRenderThread.Age				= Age;
-	DataForRenderThread.ExecutionState	= InternalExecutionState;
+	using namespace NiagaraStateless;
 
+	// Update parameter data (if needed)
 	bool bNeedsShaderParametersUpdate = false;
+	bool bNeedsParametersUpdate = false;
 	if (RendererBindings.GetParametersDirty())
 	{
 		for (const TPair<int32, FInstancedStruct>& ExpressionMapping : EmitterData->Expressions)
@@ -559,57 +553,96 @@ void FNiagaraStatelessEmitterInstance::SendRenderData()
 		if (EmitterData->bModulesHaveRendererBindings)
 		{
 			bNeedsShaderParametersUpdate = true;
-			DataForRenderThread.bHasBindingBufferData = true;
-			DataForRenderThread.BindingBufferData = RendererBindings.GetParameterDataArray();
-			check((DataForRenderThread.BindingBufferData.Num() % sizeof(uint32)) == 0);
+			bNeedsParametersUpdate = true;
 		}
 	}
 
+	// Update transforms (if needed)
 	const FTransform3f ParentTransform = FTransform3f(ParentSystemInstance->GetWorldTransform());
 	bNeedsShaderParametersUpdate |= EmitterTransforms.UpdateTransforms(ParentTransform);
 
-	if (bPerfTest && !bNeedsShaderParametersUpdate && !bSpawnInfosDirty)
+	// If CPU simulation execute immediately
+	if (SimTarget == ENiagaraSimTarget::CPUSim)
 	{
-		return;
-	}
+		if (bNeedsShaderParametersUpdate)
+		{
+			ShaderParameters.Reset(WeakStatelessEmitter->AllocateShaderParameters(EmitterTransforms, RendererBindings));
+			ShaderParameters->Common_RandomSeed = RandomSeed;
+		}
 
-	if (bNeedsShaderParametersUpdate)
-	{
-		DataForRenderThread.ShaderParameters = WeakStatelessEmitter->AllocateShaderParameters(EmitterTransforms, RendererBindings);
-		DataForRenderThread.ShaderParameters->Common_RandomSeed = RandomSeed;
-	}
+		FNiagaraDataBuffer& DataBuffer = ParticleDataSet->BeginSimulate();
 
-	if (bSpawnInfosDirty)
-	{
-		DataForRenderThread.bHasSpawnInfoData = true;
-		DataForRenderThread.SpawnInfos = SpawnInfos;
-		bSpawnInfosDirty = false;
+		FParticleSimulationContext ParticleSimulation(EmitterData.Get(), ShaderParameters.Get(), RendererBindings.GetParameterDataArray());
+		ParticleSimulation.Simulate(RandomSeed, Age, DeltaSeconds, SpawnInfos, &DataBuffer);
+
+		ParticleDataSet->EndSimulate();
 	}
+	// If GPU simulation send data to the RT
+	else if ( FEmitterInstance_RT* RenderThreadData = RenderThreadDataPtr.Get() )
+	{
+		struct FDataForRenderThread
+		{
+			float Age = 0.0f;
+			ENiagaraExecutionState ExecutionState = ENiagaraExecutionState::Disabled;
+
+			FCommonShaderParameters* ShaderParameters = nullptr;
+
+			bool bHasBindingBufferData = false;
+			TArray<uint8> BindingBufferData;
+
+			bool bHasSpawnInfoData = false;
+			TArray<FNiagaraStatelessRuntimeSpawnInfo>	SpawnInfos;
+		};
+
+		FDataForRenderThread DataForRenderThread;
+		DataForRenderThread.Age				= Age;
+		DataForRenderThread.ExecutionState	= InternalExecutionState;
+
+		if (bNeedsParametersUpdate)
+		{
+			DataForRenderThread.bHasBindingBufferData = true;
+			DataForRenderThread.BindingBufferData = RendererBindings.GetParameterDataArray();
+			check((DataForRenderThread.BindingBufferData.Num() % sizeof(uint32)) == 0);
+		}
 
-	ENQUEUE_RENDER_COMMAND(UpdateStatelessAge)(
-		[RenderThreadData, EmitterData=MoveTemp(DataForRenderThread)](FRHICommandListImmediate& RHICmdList) mutable
+		if (bNeedsShaderParametersUpdate)
 		{
-			RenderThreadData->DeltaTime			= FMath::Max(EmitterData.Age - RenderThreadData->Age, 0.0f);
-			RenderThreadData->Age				= EmitterData.Age;
-			RenderThreadData->ExecutionState	= EmitterData.ExecutionState;
+			DataForRenderThread.ShaderParameters = WeakStatelessEmitter->AllocateShaderParameters(EmitterTransforms, RendererBindings);
+			DataForRenderThread.ShaderParameters->Common_RandomSeed = RandomSeed;
+		}
 
-			if (EmitterData.ShaderParameters)
-			{
-				RenderThreadData->ShaderParameters.Reset(EmitterData.ShaderParameters);
-			}
+		if (bSpawnInfosDirty)
+		{
+			DataForRenderThread.bHasSpawnInfoData = true;
+			DataForRenderThread.SpawnInfos = SpawnInfos;
+			bSpawnInfosDirty = false;
+		}
 
-			if (EmitterData.bHasBindingBufferData)
+		ENQUEUE_RENDER_COMMAND(UpdateStatelessAge)(
+			[RenderThreadData, EmitterData=MoveTemp(DataForRenderThread)](FRHICommandListImmediate& RHICmdList) mutable
 			{
-				RenderThreadData->bBindingBufferDirty = true;
-				RenderThreadData->BindingBufferData = MoveTemp(EmitterData.BindingBufferData);
-			}
+				RenderThreadData->DeltaTime			= FMath::Max(EmitterData.Age - RenderThreadData->Age, 0.0f);
+				RenderThreadData->Age				= EmitterData.Age;
+				RenderThreadData->ExecutionState	= EmitterData.ExecutionState;
 
-			if (EmitterData.bHasSpawnInfoData)
-			{
-				RenderThreadData->SpawnInfos = MoveTemp(EmitterData.SpawnInfos);
+				if (EmitterData.ShaderParameters)
+				{
+					RenderThreadData->ShaderParameters.Reset(EmitterData.ShaderParameters);
+				}
+
+				if (EmitterData.bHasBindingBufferData)
+				{
+					RenderThreadData->bBindingBufferDirty = true;
+					RenderThreadData->BindingBufferData = MoveTemp(EmitterData.BindingBufferData);
+				}
+
+				if (EmitterData.bHasSpawnInfoData)
+				{
+					RenderThreadData->SpawnInfos = MoveTemp(EmitterData.SpawnInfos);
+				}
 			}
-		}
-	);
+		);
+	}
 }
 
 void FNiagaraStatelessEmitterInstance::InitSpawnInfos(float InitializationAge)
@@ -924,18 +957,29 @@ void FNiagaraStatelessEmitterInstance::CaptureForDebugging(FNiagaraDataBuffer* D
 		return;
 	}
 
-	ENQUEUE_RENDER_COMMAND(CaptureStatelessForDebugging)(
-		[RenderThreadData=RenderThreadDataPtr.Get(), DataBuffer](FRHICommandListImmediate& RHICmdList)
+	// Capture the data based on Sim Target
+	if (SimTarget == ENiagaraSimTarget::CPUSim)
+	{
+		if (FNiagaraDataBuffer* CurrentDataBuffer = ParticleDataSet->GetCurrentData())
 		{
-			// No compute manager then we can not do anything
-			if ( !RenderThreadData->ComputeManager )
+			CurrentDataBuffer->CopyTo(*DataBuffer, 0, 0, INDEX_NONE);
+		}
+	}
+	else
+	{
+		ENQUEUE_RENDER_COMMAND(CaptureStatelessForDebugging)(
+			[RenderThreadData=RenderThreadDataPtr.Get(), DataBuffer](FRHICommandListImmediate& RHICmdList)
 			{
-				return;
-			}
+				// No compute manager then we can not do anything
+				if ( !RenderThreadData->ComputeManager )
+				{
+					return;
+				}
 
-			RenderThreadData->ComputeManager->GenerateDataBufferForDebugging(RHICmdList, DataBuffer, RenderThreadData);
-		}
-	);
+				RenderThreadData->ComputeManager->GenerateDataBufferForDebugging(RHICmdList, DataBuffer, RenderThreadData);
+			}
+		);
 
-	FlushRenderingCommands();
+		FlushRenderingCommands();
+	}
 }
diff --git a/Engine/Plugins/FX/Niagara/Source/Niagara/Public/NiagaraLightRendererProperties.h b/Engine/Plugins/FX/Niagara/Source/Niagara/Public/NiagaraLightRendererProperties.h
--- a/Engine/Plugins/FX/Niagara/Source/Niagara/Public/NiagaraLightRendererProperties.h
+++ b/Engine/Plugins/FX/Niagara/Source/Niagara/Public/NiagaraLightRendererProperties.h
@@ -12,7 +12,7 @@
 class FNiagaraEmitterInstance;
 class SWidget;
 
-UCLASS(editinlinenew, MinimalAPI, meta = (DisplayName = "Light Renderer"))
+UCLASS(editinlinenew, MinimalAPI, meta = (DisplayName = "Light Renderer", SupportsStateless))
 class UNiagaraLightRendererProperties : public UNiagaraRendererProperties
 {
 public:
