set num cells has optional enable binding

diff --git a/Engine/Plugins/FX/Niagara/Source/Niagara/Private/NiagaraDataInterfaceGrid2DCollection.cpp b/Engine/Plugins/FX/Niagara/Source/Niagara/Private/NiagaraDataInterfaceGrid2DCollection.cpp
--- a/Engine/Plugins/FX/Niagara/Source/Niagara/Private/NiagaraDataInterfaceGrid2DCollection.cpp
+++ b/Engine/Plugins/FX/Niagara/Source/Niagara/Private/NiagaraDataInterfaceGrid2DCollection.cpp
@@ -123,6 +123,7 @@ struct FNiagaraGridCollection2DDIFunctionVersion
 	{
 		InitialVersion = 0,
 		AddReadMemberToFunctions = 1,
+		AddExecuteToSetNumCells = 2,
 
 		VersionPlusOne,
 		LatestVersion = VersionPlusOne - 1
@@ -401,6 +402,21 @@ bool UNiagaraDataInterfaceGrid2DCollection::UpgradeFunctionCall(FNiagaraFunction
 		}
 	}
 
+	if ( FunctionSignature.FunctionVersion < FNiagaraGridCollection2DDIFunctionVersion::AddExecuteToSetNumCells )
+	{
+		TArray<FNiagaraFunctionSignature> Sigs;
+		GetFunctionsInternal(Sigs);
+
+		for (const FNiagaraFunctionSignature& Sig : Sigs)
+		{
+			if (Sig.Name == FunctionSignature.Name && FunctionSignature.Name == SetNumCellsFunctionName)
+			{
+				FunctionSignature.Inputs = Sig.Inputs;
+				return true;
+			}
+		}
+	}
+
 	return false;
 }
 #endif
@@ -418,6 +434,7 @@ void UNiagaraDataInterfaceGrid2DCollection::GetFunctionsInternal(TArray<FNiagara
 		Sig.Inputs.Add(FNiagaraVariable(FNiagaraTypeDefinition(GetClass()), TEXT("Grid")));
 		Sig.Inputs.Add(FNiagaraVariable(FNiagaraTypeDefinition::GetIntDef(), TEXT("NumCellsX")));
 		Sig.Inputs.Add(FNiagaraVariable(FNiagaraTypeDefinition::GetIntDef(), TEXT("NumCellsY")));
+		Sig.Inputs.Add_GetRef(FNiagaraVariable(FNiagaraTypeDefinition::GetBoolDef(), TEXT("Execute"))).SetValue(true);
 		Sig.Outputs.Add(FNiagaraVariable(FNiagaraTypeDefinition::GetBoolDef(), TEXT("Success")));
 
 		Sig.ModuleUsageBitmask = ENiagaraScriptUsageMask::Emitter | ENiagaraScriptUsageMask::System;
@@ -1049,7 +1066,7 @@ void UNiagaraDataInterfaceGrid2DCollection::GetVMExternalFunction(const FVMExter
 	}
 	else if (BindingInfo.Name == SetNumCellsFunctionName)
 	{
-		check(BindingInfo.GetNumInputs() == 3 && BindingInfo.GetNumOutputs() == 1);
+		check(BindingInfo.GetNumInputs() == 4 && BindingInfo.GetNumOutputs() == 1);
 		NDI_FUNC_BINDER(UNiagaraDataInterfaceGrid2DCollection, VMSetNumCells)::Bind(this, OutFunc);
 	}
 
@@ -2472,13 +2489,15 @@ void UNiagaraDataInterfaceGrid2DCollection::VMSetNumCells(FVectorVMExternalFunct
 	VectorVM::FUserPtrHandler<FGrid2DCollectionRWInstanceData_GameThread> InstData(Context);
 	VectorVM::FExternalFuncInputHandler<int> InNumCellsX(Context);
 	VectorVM::FExternalFuncInputHandler<int> InNumCellsY(Context);
+	VectorVM::FExternalFuncInputHandler<FNiagaraBool> InExecute(Context);
 	VectorVM::FExternalFuncRegisterHandler<FNiagaraBool> OutSuccess(Context);
 
 	for (int32 InstanceIdx = 0; InstanceIdx < Context.GetNumInstances(); ++InstanceIdx)
 	{
 		int NewNumCellsX = InNumCellsX.GetAndAdvance();
 		int NewNumCellsY = InNumCellsY.GetAndAdvance();
-		bool bSuccess = (InstData.Get() != nullptr && Context.GetNumInstances() == 1 && NumCellsX >= 0 && NumCellsY >= 0);
+		bool bExecute = InExecute.GetAndAdvance();
+		bool bSuccess = (InstData.Get() != nullptr && Context.GetNumInstances() == 1 && NumCellsX >= 0 && NumCellsY >= 0) && bExecute;
 		*OutSuccess.GetDestAndAdvance() = bSuccess;
 		if (bSuccess)
 		{
@@ -2674,6 +2693,22 @@ void FGrid2DCollectionRWInstanceData_RenderThread::BeginSimulate(FRDGBuilder& Gr
 		Buffers.Emplace(DestinationData);
 
 		const FRDGTextureDesc TextureDesc = FRDGTextureDesc::Create2DArray(NumCells, PixelFormat.GetValue(), FClearValueBinding::Black, ETextureCreateFlags::ShaderResource | ETextureCreateFlags::UAV, uint16(NumAttributes));
+
+	#if 0
+		// Note: Invalid as the render target pool holds onto a pointer and can exist beyond the lifetime of the grid
+		TStringBuilder<128> StringBuilder;
+		SourceDIName.AppendString(StringBuilder);
+		StringBuilder.Append("_");
+		for (FName CurrName : Vars)
+		{
+			CurrName.AppendString(StringBuilder);
+			StringBuilder.Append("_");
+		}
+		StringBuilder.Append("Grid2DCollection");		
+		const TCHAR *GridTextureName = StringBuilder.ToString();
+		UE_LOG(LogNiagara, Warning, TEXT("%s"), *StringBuilder);		
+	#endif
+
 		DestinationData->Initialize(GraphBuilder, TEXT("Grid2D::GridTexture"), TextureDesc);
 
 		// This destination buffer will sometimes have old data in it.  Force it to clear.
