- Fix for Cutout + UV Flips when not in SubUV mode

diff --git a/Engine/Plugins/FX/Niagara/Shaders/Private/NiagaraSpriteVertexFactory.ush b/Engine/Plugins/FX/Niagara/Shaders/Private/NiagaraSpriteVertexFactory.ush
--- a/Engine/Plugins/FX/Niagara/Shaders/Private/NiagaraSpriteVertexFactory.ush
+++ b/Engine/Plugins/FX/Niagara/Shaders/Private/NiagaraSpriteVertexFactory.ush
@@ -673,11 +673,12 @@ float3x3 CalcTangentBasis(FVertexFactoryIntermediates Intermediates)
 void ComputeBillboardUVs(FVertexFactoryInput Input, float2 ParticleSize, float SubImageIndex, float2 UVScale, out float2 UVForPosition, out float2 UVForTexturing, out float2 UVForTexturingUnflipped)
 {
 #if FEATURE_LEVEL >= FEATURE_LEVEL_SM4
+	const uint NumCutOutVertices = GetCutoutNumVertices();
+
 	BRANCH
-	if (GetCutoutUseSubImage())
+	if (NumCutOutVertices > 0)
 	{
 		// Avoid uint divide which is extremely slow on GCN
-		const uint NumCutOutVertices = GetCutoutNumVertices();
 		const uint VertexIndexMask = NumCutOutVertices == 8 ? 7 : 3; // BVC_FourVertices || BVC_EightVertices
 		uint CutoutVertexIndex = Input.VertexId & VertexIndexMask;
 		const float NumFrames = NiagaraSpriteVF.SubImageSize.x * NiagaraSpriteVF.SubImageSize.y;
