- Add stack note about update script not running when interpolated spawning is disabled

diff --git a/Engine/Plugins/FX/Niagara/Source/Niagara/Classes/NiagaraEmitter.h b/Engine/Plugins/FX/Niagara/Source/Niagara/Classes/NiagaraEmitter.h
--- a/Engine/Plugins/FX/Niagara/Source/Niagara/Classes/NiagaraEmitter.h
+++ b/Engine/Plugins/FX/Niagara/Source/Niagara/Classes/NiagaraEmitter.h
@@ -282,7 +282,11 @@ struct FVersionedNiagaraEmitterData
 	UPROPERTY(EditAnywhere, BlueprintReadWrite, Category = "Emitter", meta = (EditCondition = "bDeterminism", EditConditionHides="true"))
 	int32 RandomSeed = 0;
 
-	/** When enabled, this will spawn using interpolated parameter values and perform a partial update at spawn time. This adds significant additional cost for spawning but will produce much smoother spawning for high spawn rates, erratic frame rates and fast moving emitters. */
+	/**
+	When enabled, this will spawn using interpolated parameter values and perform a partial update at spawn time.
+	This adds additional cost for spawning but will produce much smoother results for high spawn rates, erratic frame rates and fast moving emitters.
+	NOTE: When not disabled only the spawn script will run for newly spawned particles.  We do not run the update script.
+	*/
 	UPROPERTY(EditAnywhere, BlueprintReadWrite, Category = "Emitter")
 	uint32 bInterpolatedSpawning : 1;
 
diff --git a/Engine/Plugins/FX/Niagara/Source/NiagaraEditor/Private/ViewModels/Stack/NiagaraStackEmitterSettingsGroup.cpp b/Engine/Plugins/FX/Niagara/Source/NiagaraEditor/Private/ViewModels/Stack/NiagaraStackEmitterSettingsGroup.cpp
--- a/Engine/Plugins/FX/Niagara/Source/NiagaraEditor/Private/ViewModels/Stack/NiagaraStackEmitterSettingsGroup.cpp
+++ b/Engine/Plugins/FX/Niagara/Source/NiagaraEditor/Private/ViewModels/Stack/NiagaraStackEmitterSettingsGroup.cpp
@@ -207,6 +207,37 @@ void UNiagaraStackEmitterPropertiesItem::RefreshIssues(TArray<FStackIssue>& NewI
 			);
 		}
 
+		if (ActualEmitterData->bInterpolatedSpawning == false)
+		{
+			const FText InterpolatedSpawningDisabledInfoFix = LOCTEXT("InterpolatedSpawningDisabledInfoFix", "Enable interpolated spawning");
+			FVersionedNiagaraEmitterWeakPtr WeakEmitter = GetEmitterViewModel()->GetEmitter().ToWeakPtr();
+
+			NewIssues.Emplace(
+				EStackIssueSeverity::Info,
+				LOCTEXT("InterpolatedSpawningDisabledInfo", "Interpolated spawning disabled, particle update will not run on the first frame."),
+				LOCTEXT("InterpolatedSpawningDisabledDetailed", "When interpolated spawning is disabled only the spawn script runs on the first frame."),
+				GetStackEditorDataKey(),
+				true,
+				FStackIssueFix(
+					InterpolatedSpawningDisabledInfoFix,
+					FStackIssueFixDelegate::CreateLambda(
+						[WeakEmitter, InterpolatedSpawningDisabledInfoFix]()
+						{
+							FVersionedNiagaraEmitter PinnedEmitter = WeakEmitter.ResolveWeakPtr();
+							if (PinnedEmitter.Emitter)
+							{
+								FScopedTransaction ScopedTransaction(InterpolatedSpawningDisabledInfoFix);
+
+								PinnedEmitter.Emitter->Modify();
+								PinnedEmitter.GetEmitterData()->bInterpolatedSpawning = true;
+								UNiagaraSystem::RequestCompileForEmitter(PinnedEmitter);
+							}
+						}
+					)
+				)
+			);
+		}
+
 		UNiagaraSystem& System = GetSystemViewModel()->GetSystem();
 		TWeakPtr<FNiagaraSystemViewModel> WeakSysViewModel = GetSystemViewModel();
 		if (System.NeedsWarmup())
