Added horizontal/vertical curve flip in Niagara lightweight emitter curve editor

diff --git a/Engine/Plugins/FX/Niagara/Source/NiagaraEditor/Private/Widgets/SNiagaraDistributionCurveEditor.cpp b/Engine/Plugins/FX/Niagara/Source/NiagaraEditor/Private/Widgets/SNiagaraDistributionCurveEditor.cpp
--- a/Engine/Plugins/FX/Niagara/Source/NiagaraEditor/Private/Widgets/SNiagaraDistributionCurveEditor.cpp
+++ b/Engine/Plugins/FX/Niagara/Source/NiagaraEditor/Private/Widgets/SNiagaraDistributionCurveEditor.cpp
@@ -3,6 +3,7 @@
 #include "Widgets/SNiagaraDistributionCurveEditor.h"
 
 #include "NiagaraEditorStyle.h"
+#include "ScopedTransaction.h"
 #include "SEnumCombo.h"
 #include "Widgets/INiagaraDistributionAdapter.h"
 #include "Widgets/Input/SButton.h"
@@ -292,10 +293,12 @@ public:
 		AllottedGeometryCache = AllottedGeometry;
 	}
 
+	void RequestWidgetUpdate() { bNeedsWidgetUpdate = true; }
+
 private:
 	void UpdateKeyWidgets()
 	{
-		if (KeyWidgetData.Num() != EditCurve->GetNumKeys())
+		if (KeyWidgetData.Num() != EditCurve->GetNumKeys() || bNeedsWidgetUpdate)
 		{
 			BorderOverlay->ClearChildren();
 			KeyWidgetData.Empty();
@@ -316,6 +319,7 @@ private:
 					.OnMoved(this, &SNiagaraCompactCurveKeySelector::KeyHandleMoved, KeyWidgetDataItem.Handle)
 				];
 			}
+			bNeedsWidgetUpdate = false;
 		}
 
 		float TimeMin;
@@ -377,6 +381,8 @@ private:
 	TSharedPtr<SOverlay> BorderOverlay;
 	int32 CachedCurveStateSerialNumber = INDEX_NONE;
 	TArray<FKeyWidgetData> KeyWidgetData;
+
+	bool bNeedsWidgetUpdate = false;
 };
 
 void SNiagaraDistributionCurveEditor::Construct(const FArguments& InArgs, TSharedRef<INiagaraDistributionAdapter> InDistributionAdapter, int32 InChannelIndex)
@@ -468,6 +474,41 @@ void SNiagaraDistributionCurveEditor::Construct(const FArguments& InArgs, TShare
 						]
 					]
 				]
+				+ SVerticalBox::Slot()
+				.AutoHeight()
+				[
+					SNew(SHorizontalBox)
+						+ SHorizontalBox::Slot()
+						.AutoWidth()
+						.Padding(0, 0, 2, 0)
+						[
+							SNew(SButton)
+								.ButtonStyle(FAppStyle::Get(), "SimpleButton")
+								.OnClicked(this, &SNiagaraDistributionCurveEditor::FlipCurve, ECurveFlipDirection::Horizontal)
+								.ToolTipText(LOCTEXT("FlipHorizontalToolTip", "Flip curve horizontally"))
+								.ContentPadding(0.0f)
+								[
+									SNew(SImage)
+										.Image(FAppStyle::GetBrush("Icons.FlipHorizontal"))
+										.ColorAndOpacity(FSlateColor::UseForeground())
+								]
+						]
+						+ SHorizontalBox::Slot()
+						.AutoWidth()
+						.Padding(0, 0, 2, 0)
+						[
+							SNew(SButton)
+								.ButtonStyle(FAppStyle::Get(), "SimpleButton")
+								.OnClicked(this, &SNiagaraDistributionCurveEditor::FlipCurve, ECurveFlipDirection::Vertical)
+								.ToolTipText(LOCTEXT("FlipVerticalToolTip", "Flip curve vertically"))
+								.ContentPadding(0.0f)
+								[
+									SNew(SImage)
+										.Image(FAppStyle::GetBrush("Icons.FlipVertical"))
+										.ColorAndOpacity(FSlateColor::UseForeground())
+								]
+						]
+				]
 			]
 			// Curve view and key selector.
 			+ SHorizontalBox::Slot()
@@ -529,7 +570,7 @@ void SNiagaraDistributionCurveEditor::Construct(const FArguments& InArgs, TShare
 				+ SGridPanel::Slot(0, 1)
 				.Padding(0, 0, 0, 3)
 				[
-					SNew(SNiagaraCompactCurveKeySelector, &EditCurve)
+					SAssignNew(KeySelectorWidget, SNiagaraCompactCurveKeySelector, &EditCurve)
 					.HandleSize(FVector2f(KeyHandleWidth, KeyHandleHeight))
 					.SelectedKeyHandle(this, &SNiagaraDistributionCurveEditor::GetSelectedKeyHandle)
 					.CurveStateSerialNumber(this, &SNiagaraDistributionCurveEditor::GetCurveStateSerialNumber)
@@ -808,6 +849,90 @@ FReply SNiagaraDistributionCurveEditor::DeleteKeyButtonClicked()
 	return FReply::Unhandled();
 }
 
+FReply SNiagaraDistributionCurveEditor::FlipCurve(ECurveFlipDirection Direction)
+{
+	if (EditCurve.GetNumKeys() <= 0)
+	{
+		return FReply::Unhandled();
+	}
+
+	FScopedTransaction Transaction(LOCTEXT("FlipCurve", "Flip Curve"));
+
+	int32 NumKeys = EditCurve.GetNumKeys();
+	int32 NewSelectedKeyIndex = SelectedKeyIndex;
+
+	// Intialize the range
+	float RangeMin = FLT_MAX;
+	float RangeMax = -FLT_MAX;
+
+	// If flipping horizontally
+	if (Direction == ECurveFlipDirection::Horizontal)
+	{
+		// Get the range
+		float MinTime, MaxTime;
+		EditCurve.GetTimeRange(MinTime, MaxTime);
+		RangeMin = FMath::Min(MinTime, RangeMin);
+		RangeMax = FMath::Max(MaxTime, RangeMax);
+
+		TArray<TPair<float, float>> TimeValuePairs;
+		TArray<ERichCurveInterpMode> InterpModes;
+		TimeValuePairs.SetNum(NumKeys);
+		InterpModes.SetNum(NumKeys);
+
+		// Get the (flipped x, original y) value of all points
+		FKeyHandle KeyHandle = EditCurve.GetFirstKeyHandle();
+		for (int32 Index = 0; Index <= NumKeys - 1; ++Index)
+		{
+			FRichCurveKey Key = EditCurve.GetKey(KeyHandle);
+			TimeValuePairs[Index].Key = RangeMax - Key.Time + RangeMin;
+			TimeValuePairs[Index].Value = Key.Value;
+			InterpModes[Index] = EditCurve.GetKeyInterpMode(KeyHandle);
+			KeyHandle = EditCurve.GetNextKey(KeyHandle);
+		}
+
+		// Delete all existing keys
+		EditCurve.Reset();
+
+		// Add the flipped points back in
+		FKeyHandle NewKeyHandle;
+		for (int32 Index = 0; Index <= NumKeys - 1; ++Index)
+		{
+			NewKeyHandle = EditCurve.AddKey(TimeValuePairs[Index].Key, TimeValuePairs[Index].Value);
+			EditCurve.SetKeyInterpMode(NewKeyHandle, InterpModes[Index]);
+			//NOTE: key tangents are not handled right now since distribution curve editor currently don't support tangent editing from user
+		}
+
+		NewSelectedKeyIndex = NumKeys - 1 - SelectedKeyIndex;
+		if (KeySelectorWidget) 
+		{
+			KeySelectorWidget->RequestWidgetUpdate();
+		}
+	}
+
+	// If flipping vertically
+	if (Direction == ECurveFlipDirection::Vertical)
+	{
+		// Get the range
+		float MinVal, MaxVal;
+		EditCurve.GetValueRange(MinVal, MaxVal);
+		RangeMin = FMath::Min(MinVal, RangeMin);
+		RangeMax = FMath::Max(MaxVal, RangeMax);
+
+		// Flip the y value of all keys
+		FKeyHandle KeyHandle = EditCurve.GetFirstKeyHandle();
+		for (int32 Index = 0; Index <= NumKeys - 1; ++Index)
+		{
+			EditCurve.SetKeyValue(KeyHandle, RangeMax - EditCurve.GetKey(KeyHandle).Value + RangeMin);
+			KeyHandle = EditCurve.GetNextKey(KeyHandle);
+		}
+	}
+
+	DistributionAdapter->SetCurveValue(ChannelIndex, EditCurve);
+	UpdateSelectedKeyIndex(NewSelectedKeyIndex);
+	CurveStateSerialNumber++;
+	return FReply::Handled();
+}
+
 void SNiagaraDistributionCurveEditor::OnKeySelectorKeySelected(FKeyHandle InSelectedKeyHandle)
 {
 	int32 KeyIndex = 0;
diff --git a/Engine/Plugins/FX/Niagara/Source/NiagaraEditor/Private/Widgets/SNiagaraDistributionCurveEditor.h b/Engine/Plugins/FX/Niagara/Source/NiagaraEditor/Private/Widgets/SNiagaraDistributionCurveEditor.h
--- a/Engine/Plugins/FX/Niagara/Source/NiagaraEditor/Private/Widgets/SNiagaraDistributionCurveEditor.h
+++ b/Engine/Plugins/FX/Niagara/Source/NiagaraEditor/Private/Widgets/SNiagaraDistributionCurveEditor.h
@@ -7,6 +7,7 @@
 #include "Widgets/SCompoundWidget.h"
 
 class INiagaraDistributionAdapter;
+class SNiagaraCompactCurveKeySelector;
 
 class SNiagaraDistributionCurveEditor : public SCompoundWidget
 {
@@ -20,6 +21,14 @@ public:
 
 	void Construct(const FArguments& InArgs, TSharedRef<INiagaraDistributionAdapter> InDistributionAdapter, int32 InChannelIndex);
 
+private:
+	/** Curve flip direction types **/
+	enum class ECurveFlipDirection : uint8
+	{
+		Horizontal,
+		Vertical
+	};
+
 private:
 	TOptional<int32> GetSelectedKeyIndex() const;
 	TOptional<int32> GetMaxKeyIndex() const;
@@ -45,6 +54,7 @@ private:
 
 	FReply AddKeyButtonClicked();
 	FReply DeleteKeyButtonClicked();
+	FReply FlipCurve(ECurveFlipDirection Direction);
 
 	FKeyHandle GetSelectedKeyHandle() const { return SelectedKeyHandle; }
 	int32 GetCurveStateSerialNumber() const { return CurveStateSerialNumber; }
@@ -60,6 +70,7 @@ private:
 	void UpdateCachedValues() const;
 
 private:
+	TSharedPtr<SNiagaraCompactCurveKeySelector> KeySelectorWidget;
 	TSharedPtr<INiagaraDistributionAdapter> DistributionAdapter;
 	int32 ChannelIndex;
 
