- Enable Niagara Editor Performance Stats for new VM
- The new VM doesn't breakdown per module and will only display information for each stack group

diff --git a/Engine/Plugins/FX/Niagara/Source/Niagara/Classes/NiagaraSystem.h b/Engine/Plugins/FX/Niagara/Source/Niagara/Classes/NiagaraSystem.h
--- a/Engine/Plugins/FX/Niagara/Source/Niagara/Classes/NiagaraSystem.h
+++ b/Engine/Plugins/FX/Niagara/Source/Niagara/Classes/NiagaraSystem.h
@@ -41,6 +41,16 @@ struct FNiagaraGraphCachedDataBase;
 class UNiagaraDataChannel;
 class ITargetPlatform;
 
+enum class ENiagaraPerformanceStateScopeMode : int
+{
+	/** No stats exist and we should not show the option. */
+	None,
+	/** Only basic (i.e. whole script) stats exist. */
+	Basic,
+	/** Detailed (i.e. per module) stats exist. */
+	Detailed,
+};
+
 USTRUCT()
 struct FNiagaraEmitterCompiledData
 {
@@ -546,8 +556,8 @@ public:
 	FORCEINLINE void SetBakeOutRapidIterationOnCook(bool bBakeOut) { bBakeOutRapidIteration = bBakeOut; bBakeOutRapidIterationOnCook = bBakeOut; }
 	FORCEINLINE void SetTrimAttributesOnCook(bool bTrim) { bTrimAttributes = bTrim; bTrimAttributesOnCook = bTrim; }
 
-	/** returns true if the system, and it's underlying scripts support stat scopes and the collection of the timings */
-	NIAGARA_API bool SupportsStatScopedPerformanceMode() const;
+	/** Returns information on the availability of performance stat scopes */
+	NIAGARA_API ENiagaraPerformanceStateScopeMode GetPerformanceStateScopeMode() const;
 
 protected:
 	/** When enable constant values are baked into the scripts while editing the system, this will increase iteration times but improve performance. */
diff --git a/Engine/Plugins/FX/Niagara/Source/Niagara/Private/NiagaraScriptExecutionContext.cpp b/Engine/Plugins/FX/Niagara/Source/Niagara/Private/NiagaraScriptExecutionContext.cpp
--- a/Engine/Plugins/FX/Niagara/Source/Niagara/Private/NiagaraScriptExecutionContext.cpp
+++ b/Engine/Plugins/FX/Niagara/Source/Niagara/Private/NiagaraScriptExecutionContext.cpp
@@ -278,6 +278,9 @@ bool FNiagaraScriptExecutionContextBase::Execute(class FNiagaraSystemInstance* I
 bool FNiagaraScriptExecutionContextBase::ExecuteInternal_Experimental(uint32 NumInstances, const FScriptExecutionConstantBufferTable& ConstantBufferTable)
 {
 	TRACE_CPUPROFILER_EVENT_SCOPE(VectorVM_Experimental);
+#if STATS
+	const uint64 ExecutionStartCycles = FPlatformTime::Cycles64();
+#endif
 	FVectorVMExecContext ExecCtx;
 	ExecCtx.VVMState               = VectorVMState;
 	ExecCtx.DataSets               = DataSetMetaTable;
@@ -292,6 +295,14 @@ bool FNiagaraScriptExecutionContextBase::ExecuteInternal_Experimental(uint32 Num
 	{
 		ExecVectorVMState(&ExecCtx, nullptr, nullptr);
 	}
+
+#if STATS
+	// We fill them all out as this makes sure the UI can pickup the appropriate scope
+	for (FStatScopeData& StatScope : StatScopeData)
+	{
+		StatScope.ExecutionCycleCount += FPlatformTime::Cycles64() - ExecutionStartCycles;
+	}
+#endif
 	return true;
 }
 
diff --git a/Engine/Plugins/FX/Niagara/Source/Niagara/Private/NiagaraSystem.cpp b/Engine/Plugins/FX/Niagara/Source/Niagara/Private/NiagaraSystem.cpp
--- a/Engine/Plugins/FX/Niagara/Source/Niagara/Private/NiagaraSystem.cpp
+++ b/Engine/Plugins/FX/Niagara/Source/Niagara/Private/NiagaraSystem.cpp
@@ -1034,10 +1034,14 @@ void UNiagaraSystem::RecomputeExecutionOrderForDataInterface(class UNiagaraDataI
 	}
 }
 
-bool UNiagaraSystem::SupportsStatScopedPerformanceMode() const
+ENiagaraPerformanceStateScopeMode UNiagaraSystem::GetPerformanceStateScopeMode() const
 {
 	const UNiagaraSettings* Settings = GetDefault<UNiagaraSettings>();
-	return bDisableExperimentalVM || !Settings || !Settings->bExperimentalVMEnabled;
+	if (bDisableExperimentalVM || !Settings || !Settings->bExperimentalVMEnabled)
+	{
+		return ENiagaraPerformanceStateScopeMode::Detailed;
+	}
+	return ENiagaraPerformanceStateScopeMode::Basic;
 }
 
 #endif
diff --git a/Engine/Plugins/FX/Niagara/Source/NiagaraEditor/Private/Toolkits/NiagaraSystemToolkit.cpp b/Engine/Plugins/FX/Niagara/Source/NiagaraEditor/Private/Toolkits/NiagaraSystemToolkit.cpp
--- a/Engine/Plugins/FX/Niagara/Source/NiagaraEditor/Private/Toolkits/NiagaraSystemToolkit.cpp
+++ b/Engine/Plugins/FX/Niagara/Source/NiagaraEditor/Private/Toolkits/NiagaraSystemToolkit.cpp
@@ -1009,7 +1009,7 @@ void FNiagaraSystemToolkit::OnToggleBoundsSetFixedBounds_System()
 bool FNiagaraSystemToolkit::SupportsStatPerformance() const
 {
 #if STATS
-	return SystemViewModel->GetSystem().SupportsStatScopedPerformanceMode();
+	return SystemViewModel->GetSystem().GetPerformanceStateScopeMode() != ENiagaraPerformanceStateScopeMode::None;
 #else
 	return false;
 #endif
diff --git a/Engine/Plugins/FX/Niagara/Source/NiagaraEditor/Private/Toolkits/SystemToolkitModes/NiagaraSystemToolkitMode_Default.cpp b/Engine/Plugins/FX/Niagara/Source/NiagaraEditor/Private/Toolkits/SystemToolkitModes/NiagaraSystemToolkitMode_Default.cpp
--- a/Engine/Plugins/FX/Niagara/Source/NiagaraEditor/Private/Toolkits/SystemToolkitModes/NiagaraSystemToolkitMode_Default.cpp
+++ b/Engine/Plugins/FX/Niagara/Source/NiagaraEditor/Private/Toolkits/SystemToolkitModes/NiagaraSystemToolkitMode_Default.cpp
@@ -175,7 +175,7 @@ void FNiagaraSystemToolkitMode_Default::ExtendToolbar()
 			{
 				ToolbarBuilder.AddToolBarButton(FNiagaraEditorCommands::Get().ToggleStatPerformance, NAME_None,
                     LOCTEXT("NiagaraShowPerformance", "Performance"),
-                    LOCTEXT("NiagaraShowPerformanceTooltip", "Show runtime performance for particle scripts. Not supported when optimized VM byte code setting is enabled."),
+                    LOCTEXT("NiagaraShowPerformanceTooltip", "Show runtime performance for particle scripts."),
                     FSlateIcon(FAppStyle::Get().GetStyleSetName(), "AssetEditor.ToggleStats"));
 				ToolbarBuilder.AddComboButton(
                     FUIAction(),
diff --git a/Engine/Plugins/FX/Niagara/Source/NiagaraEditor/Private/ViewModels/NiagaraSystemViewModel.cpp b/Engine/Plugins/FX/Niagara/Source/NiagaraEditor/Private/ViewModels/NiagaraSystemViewModel.cpp
--- a/Engine/Plugins/FX/Niagara/Source/NiagaraEditor/Private/ViewModels/NiagaraSystemViewModel.cpp
+++ b/Engine/Plugins/FX/Niagara/Source/NiagaraEditor/Private/ViewModels/NiagaraSystemViewModel.cpp
@@ -3235,9 +3235,9 @@ void FNiagaraSystemViewModel::UpdateSystemFixedBounds()
 	}
 }
 
-bool FNiagaraSystemViewModel::SupportsPerformanceMode() const
+ENiagaraPerformanceStateScopeMode FNiagaraSystemViewModel::GetPerformanceStateScopeMode() const
 {
-	return System->SupportsStatScopedPerformanceMode();
+	return System->GetPerformanceStateScopeMode();
 }
 
 void FNiagaraSystemViewModel::ClearSystemStats()
diff --git a/Engine/Plugins/FX/Niagara/Source/NiagaraEditor/Private/ViewModels/Stack/NiagaraStackEntry.cpp b/Engine/Plugins/FX/Niagara/Source/NiagaraEditor/Private/ViewModels/Stack/NiagaraStackEntry.cpp
--- a/Engine/Plugins/FX/Niagara/Source/NiagaraEditor/Private/ViewModels/Stack/NiagaraStackEntry.cpp
+++ b/Engine/Plugins/FX/Niagara/Source/NiagaraEditor/Private/ViewModels/Stack/NiagaraStackEntry.cpp
@@ -476,6 +476,11 @@ TSharedPtr<FNiagaraEmitterViewModel> UNiagaraStackEntry::GetEmitterViewModel() c
 	return EmitterViewModel.Pin();
 }
 
+TSharedPtr<FNiagaraSystemViewModel> UNiagaraStackEntry::GetSystemViewModelPtr() const
+{
+	return SystemViewModel.Pin();
+}
+
 UNiagaraStackNote* UNiagaraStackEntry::GetStackNote()
 {
 	return StackNote;
diff --git a/Engine/Plugins/FX/Niagara/Source/NiagaraEditor/Public/ViewModels/NiagaraSystemViewModel.h b/Engine/Plugins/FX/Niagara/Source/NiagaraEditor/Public/ViewModels/NiagaraSystemViewModel.h
--- a/Engine/Plugins/FX/Niagara/Source/NiagaraEditor/Public/ViewModels/NiagaraSystemViewModel.h
+++ b/Engine/Plugins/FX/Niagara/Source/NiagaraEditor/Public/ViewModels/NiagaraSystemViewModel.h
@@ -53,7 +53,7 @@ class UNiagaraSystemEditorData;
 class UNiagaraSystemEditorDocumentsViewModel;
 class UNiagaraSystemSelectionViewModel;
 class UNiagaraUserParametersHierarchyViewModel;
-
+enum class ENiagaraPerformanceStateScopeMode : int;
 
 /** Defines different editing modes for this system view model. */
 enum class ENiagaraSystemViewModelEditMode
@@ -339,8 +339,8 @@ public:
 	/** Updates the current system's fixed bounds with its current dynamic bounds. */
 	void UpdateSystemFixedBounds();
 
-	/** Returns whether the current system is in a state that supports PerformanceMode/stats collection. */
-	NIAGARAEDITOR_API bool SupportsPerformanceMode() const;
+	/** Returns what level of performance stats we have available in the editor. */
+	NIAGARAEDITOR_API ENiagaraPerformanceStateScopeMode GetPerformanceStateScopeMode() const;
 
 	/** Clear the captures stats for the current system. */
 	void ClearSystemStats();
diff --git a/Engine/Plugins/FX/Niagara/Source/NiagaraEditor/Public/ViewModels/Stack/NiagaraStackEntry.h b/Engine/Plugins/FX/Niagara/Source/NiagaraEditor/Public/ViewModels/Stack/NiagaraStackEntry.h
--- a/Engine/Plugins/FX/Niagara/Source/NiagaraEditor/Public/ViewModels/Stack/NiagaraStackEntry.h
+++ b/Engine/Plugins/FX/Niagara/Source/NiagaraEditor/Public/ViewModels/Stack/NiagaraStackEntry.h
@@ -451,6 +451,9 @@ public:
 	NIAGARAEDITOR_API TSharedRef<FNiagaraSystemViewModel> GetSystemViewModel() const;
 	NIAGARAEDITOR_API TSharedPtr<FNiagaraEmitterViewModel> GetEmitterViewModel() const;
 
+	/** Same as GetSystemViewModel but can return a nullptr. */
+	NIAGARAEDITOR_API TSharedPtr<FNiagaraSystemViewModel> GetSystemViewModelPtr() const;
+
 	template<typename ChildType, typename PredicateType>
 	static ChildType* FindCurrentChildOfTypeByPredicate(const TArray<UNiagaraStackEntry*>& CurrentChildren, PredicateType Predicate)
 	{
diff --git a/Engine/Plugins/FX/Niagara/Source/NiagaraEditorWidgets/Private/SNiagaraOverviewStack.cpp b/Engine/Plugins/FX/Niagara/Source/NiagaraEditorWidgets/Private/SNiagaraOverviewStack.cpp
--- a/Engine/Plugins/FX/Niagara/Source/NiagaraEditorWidgets/Private/SNiagaraOverviewStack.cpp
+++ b/Engine/Plugins/FX/Niagara/Source/NiagaraEditorWidgets/Private/SNiagaraOverviewStack.cpp
@@ -1130,7 +1130,7 @@ TSharedRef<ITableRow> SNiagaraOverviewStack::OnGenerateRowForEntry(UNiagaraStack
 			.AutoWidth()
 			.Padding(3, 0, 0, 0)
 			[
-				SNew(SNiagaraStackRowPerfWidget, Item)
+				SNew(SNiagaraStackRowPerfWidget, Item, ENiagaraPerformanceStateScopeMode::Detailed)
 			]
 			+ SHorizontalBox::Slot()
 			.VAlign(VAlign_Center)
@@ -1157,7 +1157,7 @@ TSharedRef<ITableRow> SNiagaraOverviewStack::OnGenerateRowForEntry(UNiagaraStack
 			.AutoWidth()
 			.Padding(3, 0, 0, 0)
 			[
-				SNew(SNiagaraStackRowPerfWidget, Item)
+				SNew(SNiagaraStackRowPerfWidget, Item, ENiagaraPerformanceStateScopeMode::Detailed)
 			]
 			+ SHorizontalBox::Slot()
 			.AutoWidth()
@@ -1300,7 +1300,7 @@ TSharedRef<ITableRow> SNiagaraOverviewStack::OnGenerateRowForEntry(UNiagaraStack
 			.AutoWidth()
 			.Padding(6, 0, 5, 0)
 			[
-				SNew(SNiagaraStackRowPerfWidget, Item)
+				SNew(SNiagaraStackRowPerfWidget, Item, ENiagaraPerformanceStateScopeMode::Basic)
 			];
 
 		// Inheritance Icon
diff --git a/Engine/Plugins/FX/Niagara/Source/NiagaraEditorWidgets/Private/Stack/SNiagaraStackRowPerfWidget.cpp b/Engine/Plugins/FX/Niagara/Source/NiagaraEditorWidgets/Private/Stack/SNiagaraStackRowPerfWidget.cpp
--- a/Engine/Plugins/FX/Niagara/Source/NiagaraEditorWidgets/Private/Stack/SNiagaraStackRowPerfWidget.cpp
+++ b/Engine/Plugins/FX/Niagara/Source/NiagaraEditorWidgets/Private/Stack/SNiagaraStackRowPerfWidget.cpp
@@ -24,9 +24,11 @@
 
 IConsoleVariable* SNiagaraStackRowPerfWidget::StatEnabledVar = IConsoleManager::Get().FindConsoleVariable(TEXT("vm.DetailedVMScriptStats"));
 
-void SNiagaraStackRowPerfWidget::Construct(const FArguments& InArgs, UNiagaraStackEntry* InStackEntry)
+void SNiagaraStackRowPerfWidget::Construct(const FArguments& InArgs, UNiagaraStackEntry* InStackEntry, ENiagaraPerformanceStateScopeMode InStatScopeMode)
 {
 	StackEntry = InStackEntry;
+	StatScopeMode = InStatScopeMode;
+
 	SetToolTipText(CreateTooltipText());
 	SetVisibility(TAttribute<EVisibility>::CreateSP(this, &SNiagaraStackRowPerfWidget::IsVisible));
 	
@@ -79,6 +81,7 @@ void SNiagaraStackRowPerfWidget::Construct(const FArguments& InArgs, UNiagaraSta
                     SNew(SBox)
                     .HeightOverride(2)
                     .WidthOverride(this, &SNiagaraStackRowPerfWidget::GetPlaceholderBrushWidth)
+					.Visibility(this, &SNiagaraStackRowPerfWidget::GetPerfBarVisibility)
                     [
                         SNew(SColorBlock)
                         .Color(this, &SNiagaraStackRowPerfWidget::GetPlaceholderBrushColor)
@@ -90,6 +93,7 @@ void SNiagaraStackRowPerfWidget::Construct(const FArguments& InArgs, UNiagaraSta
                     SNew(SBox)
                     .HeightOverride(2)
                     .WidthOverride(this, &SNiagaraStackRowPerfWidget::GetVisualizationBrushWidth)
+					.Visibility(this, &SNiagaraStackRowPerfWidget::GetPerfBarVisibility)
                     [
                         SNew(SColorBlock)
                         .Color(this, &SNiagaraStackRowPerfWidget::GetVisualizationBrushColor)
@@ -152,6 +156,11 @@ float SNiagaraStackRowPerfWidget::GetFullBarWidth() const
 	return IsGroupHeaderEntry() ? 54 : 40;
 }
 
+EVisibility SNiagaraStackRowPerfWidget::GetPerfBarVisibility() const
+{
+	return StatScopeMode == ENiagaraPerformanceStateScopeMode::Detailed ? EVisibility::Visible : EVisibility::Hidden;
+}
+
 FOptionalSize SNiagaraStackRowPerfWidget::GetVisualizationBrushWidth() const
 {
 	if (IsGroupHeaderEntry())
@@ -232,7 +241,7 @@ bool SNiagaraStackRowPerfWidget::HasPerformanceData() const
 		return false;
 	}
 
-	bool IsPerfCaptureEnabled = StatEnabledVar && StatEnabledVar->GetBool() && StackEntry->GetSystemViewModel()->SupportsPerformanceMode();
+	bool IsPerfCaptureEnabled = StatEnabledVar && StatEnabledVar->GetBool() && (StackEntry->GetSystemViewModel()->GetPerformanceStateScopeMode() >= StatScopeMode);
 	if (IsPerfCaptureEnabled && (IsSystemStack() || IsEmitterStack()))
 	{
 		return (IsGroupHeaderEntry() || IsModuleEntry()) && (StackEntry->GetExecutionSubcategoryName() != UNiagaraStackEntry::FExecutionSubcategoryNames::Settings);
@@ -282,8 +291,8 @@ FText SNiagaraStackRowPerfWidget::GetPerformanceDisplayText() const
 
 	FNumberFormattingOptions TimeFormatOptions;
 	TimeFormatOptions.MinimumIntegralDigits = 1;
-	TimeFormatOptions.MinimumFractionalDigits = 2;
-	TimeFormatOptions.MaximumFractionalDigits = 2;
+	TimeFormatOptions.MinimumFractionalDigits = 3;
+	TimeFormatOptions.MaximumFractionalDigits = 3;
 	TimeFormatOptions.UseGrouping = false;
 	
 	if (IsGroupHeaderEntry())
@@ -456,16 +465,14 @@ ENiagaraScriptUsage SNiagaraStackRowPerfWidget::GetUsage() const
 
 ENiagaraStatEvaluationType SNiagaraStackRowPerfWidget::GetEvaluationType() const
 {
-	return StackEntry.IsValid()
-		       ? StackEntry->GetSystemViewModel()->StatEvaluationType
-		       : ENiagaraStatEvaluationType::Average;
+	TSharedPtr<FNiagaraSystemViewModel> ViewModel = StackEntry.IsValid() ? StackEntry->GetSystemViewModelPtr() : nullptr;
+	return ViewModel.IsValid() ? ViewModel->StatEvaluationType : ENiagaraStatEvaluationType::Average;
 }
 
 ENiagaraStatDisplayMode SNiagaraStackRowPerfWidget::GetDisplayMode() const
 {
-	return StackEntry.IsValid()
-               ? StackEntry->GetSystemViewModel()->StatDisplayMode
-               : ENiagaraStatDisplayMode::Percent;
+	TSharedPtr<FNiagaraSystemViewModel> ViewModel = StackEntry.IsValid() ? StackEntry->GetSystemViewModelPtr() : nullptr;
+	return ViewModel.IsValid() ? ViewModel->StatDisplayMode : ENiagaraStatDisplayMode::Percent;
 }
 
 bool SNiagaraStackRowPerfWidget::IsInterpolatedSpawnEnabled() const
diff --git a/Engine/Plugins/FX/Niagara/Source/NiagaraEditorWidgets/Private/Stack/SNiagaraStackRowPerfWidget.h b/Engine/Plugins/FX/Niagara/Source/NiagaraEditorWidgets/Private/Stack/SNiagaraStackRowPerfWidget.h
--- a/Engine/Plugins/FX/Niagara/Source/NiagaraEditorWidgets/Private/Stack/SNiagaraStackRowPerfWidget.h
+++ b/Engine/Plugins/FX/Niagara/Source/NiagaraEditorWidgets/Private/Stack/SNiagaraStackRowPerfWidget.h
@@ -8,6 +8,7 @@
 #include "Widgets/DeclarativeSyntaxSupport.h"
 
 class UNiagaraStackEntry;
+enum class ENiagaraPerformanceStateScopeMode : int;
 
 class SNiagaraStackRowPerfWidget: public SCompoundWidget
 {
@@ -15,12 +16,13 @@ public:
 	SLATE_BEGIN_ARGS(SNiagaraStackRowPerfWidget) {}
 	SLATE_END_ARGS()
 	
-	void Construct(const FArguments& InArgs, UNiagaraStackEntry* InStackEntry);
+	void Construct(const FArguments& InArgs, UNiagaraStackEntry* InStackEntry, ENiagaraPerformanceStateScopeMode InStatScopeMode);
 
 	virtual void Tick(const FGeometry& AllottedGeometry, const double InCurrentTime, const float InDeltaTime) override;
 	
 private:
 	float GetFullBarWidth() const;
+	EVisibility GetPerfBarVisibility() const;
 	FOptionalSize GetVisualizationBrushWidth() const;
 	FOptionalSize GetPlaceholderBrushWidth() const;
 	FLinearColor GetVisualizationBrushColor() const;
@@ -53,6 +55,7 @@ private:
 	mutable TOptional<FText> IconToolTipCache;
 	TWeakObjectPtr<UNiagaraStackEntry> StackEntry;
 	static IConsoleVariable* StatEnabledVar;
+	ENiagaraPerformanceStateScopeMode StatScopeMode;
 
 	float GroupOverallTime = 0;
 	float StackEntryTime = 0;
