- Fix transitions for count buffer with arrays

#rnx

diff --git a/Engine/Plugins/FX/Niagara/Source/Niagara/Classes/NiagaraDataInterfaceArrayImpl.h b/Engine/Plugins/FX/Niagara/Source/Niagara/Classes/NiagaraDataInterfaceArrayImpl.h
--- a/Engine/Plugins/FX/Niagara/Source/Niagara/Classes/NiagaraDataInterfaceArrayImpl.h
+++ b/Engine/Plugins/FX/Niagara/Source/Niagara/Classes/NiagaraDataInterfaceArrayImpl.h
@@ -202,7 +202,7 @@ struct FNiagaraDataInterfaceArrayImplInternal
 	}
 	NIAGARA_API static bool IsRWFunction(const FName FunctionName);
 
-	NIAGARA_API static ERHIAccess GetCountBufferDefaultState();
+	NIAGARA_API static ERHIAccess GetCountBufferRHIAccess(const FNiagaraGpuComputeDispatchInterface& ComputeInterface);
 };
 
 //////////////////////////////////////////////////////////////////////////
@@ -441,14 +441,15 @@ struct FNDIArrayProxyImpl : public INDIArrayProxyBase
 			{InstanceData_RT->ArrayBuffer, 0, InstanceData_RT->ArrayNumBytes},		//-TODO: Technically last element is default for RW buffers
 		};
 
+		const ERHIAccess CountRHIAccess = FNiagaraDataInterfaceArrayImplInternal::GetCountBufferRHIAccess(Context.GetComputeDispatchInterface());
 		const FRHITransitionInfo TransitionsBefore[] =
 		{
-			FRHITransitionInfo(CountManager.GetInstanceCountBuffer().UAV, ERHIAccess::UAVCompute, ERHIAccess::CopySrc),
+			FRHITransitionInfo(CountManager.GetInstanceCountBuffer().UAV, CountRHIAccess, ERHIAccess::CopySrc),
 			FRHITransitionInfo(InstanceData_RT->ArrayBuffer, ERHIAccess::UAVCompute, ERHIAccess::CopySrc),
 		};
 		const FRHITransitionInfo TransitionsAfter[] =
 		{
-			FRHITransitionInfo(CountManager.GetInstanceCountBuffer().UAV, ERHIAccess::CopySrc, ERHIAccess::UAVCompute),
+			FRHITransitionInfo(CountManager.GetInstanceCountBuffer().UAV, ERHIAccess::CopySrc, CountRHIAccess),
 			FRHITransitionInfo(InstanceData_RT->ArrayBuffer, ERHIAccess::CopySrc, ERHIAccess::UAVCompute),
 		};
 
diff --git a/Engine/Plugins/FX/Niagara/Source/Niagara/Private/NiagaraDataInterfaceArrayImpl.cpp b/Engine/Plugins/FX/Niagara/Source/Niagara/Private/NiagaraDataInterfaceArrayImpl.cpp
--- a/Engine/Plugins/FX/Niagara/Source/Niagara/Private/NiagaraDataInterfaceArrayImpl.cpp
+++ b/Engine/Plugins/FX/Niagara/Source/Niagara/Private/NiagaraDataInterfaceArrayImpl.cpp
@@ -1,7 +1,7 @@
 // Copyright Epic Games, Inc. All Rights Reserved.
 
 #include "NiagaraDataInterfaceArrayImpl.h"
-
+#include "NiagaraGpuComputeDispatch.h"
 #include "NiagaraStats.h"
 
 const TCHAR* FNiagaraDataInterfaceArrayImplInternal::HLSLReadTemplateFile = TEXT("/Plugin/FX/Niagara/Private/NiagaraDataInterfaceArrayTemplate.ush");
@@ -220,9 +220,10 @@ bool FNiagaraDataInterfaceArrayImplInternal::IsRWFunction(const FName FunctionNa
 	return RWFunctions.Contains(FunctionName);
 }
 
-ERHIAccess FNiagaraDataInterfaceArrayImplInternal::GetCountBufferDefaultState()
+ERHIAccess FNiagaraDataInterfaceArrayImplInternal::GetCountBufferRHIAccess(const FNiagaraGpuComputeDispatchInterface& InComputeInterface)
 {
-	return FNiagaraGPUInstanceCountManager::kCountBufferDefaultState;
+	const FNiagaraGpuComputeDispatch& ComputeInterface = static_cast<const FNiagaraGpuComputeDispatch&>(InComputeInterface);
+	return ComputeInterface.IsExecutingLastDispatchGroup() ? FNiagaraGPUInstanceCountManager::kCountBufferDefaultState : ERHIAccess::UAVCompute;
 }
 
 FNDIArrayInstanceData_RenderThreadBase::~FNDIArrayInstanceData_RenderThreadBase()
