- Disable option to create stateless emitters if the class is not allowed

#rnx

diff --git a/Engine/Plugins/FX/Niagara/Source/NiagaraEditorWidgets/Private/SNiagaraOverviewGraph.cpp b/Engine/Plugins/FX/Niagara/Source/NiagaraEditorWidgets/Private/SNiagaraOverviewGraph.cpp
--- a/Engine/Plugins/FX/Niagara/Source/NiagaraEditorWidgets/Private/SNiagaraOverviewGraph.cpp
+++ b/Engine/Plugins/FX/Niagara/Source/NiagaraEditorWidgets/Private/SNiagaraOverviewGraph.cpp
@@ -1,6 +1,8 @@
 // Copyright Epic Games, Inc. All Rights Reserved.
 
 #include "SNiagaraOverviewGraph.h"
+#include "ClassViewerModule.h"
+#include "ClassViewerFilter.h"
 #include "EdGraph/EdGraph.h"
 #include "EdGraph/EdGraphNode.h"
 #include "EdGraphSchema_Niagara.h"
@@ -24,6 +26,7 @@
 #include "SNiagaraStack.h"
 #include "ScopedTransaction.h"
 #include "Interfaces/IMainFrameModule.h"
+#include "Stateless/NiagaraStatelessEmitter.h"
 #include "Templates/SharedPointer.h"
 #include "ViewModels/NiagaraEmitterHandleViewModel.h"
 #include "ViewModels/NiagaraOverviewGraphViewModel.h"
@@ -244,9 +247,13 @@ FActionMenuContent SNiagaraOverviewGraph::OnCreateGraphActionMenu(UEdGraph* InGr
 				FExecuteAction::CreateSP(this, &SNiagaraOverviewGraph::OnCreateEmptyEmitter));
 			MenuBuilder.AddMenuEntry(
 				LOCTEXT("AddStatelessEmitter", "Add lightweight emitter"),
-				LOCTEXT("AddStatelessEmitterToolTip", "Adds a lightweight emitter without any modules or renderers."),
+				GetCreateStatelessEmitterTooltip(),
 				FSlateIcon(),
-				FExecuteAction::CreateSP(this, &SNiagaraOverviewGraph::OnCreateStatelessEmitter));
+				FUIAction(
+					FExecuteAction::CreateSP(this, &SNiagaraOverviewGraph::OnCreateStatelessEmitter),
+					FCanExecuteAction::CreateSP(this, &SNiagaraOverviewGraph::CanCreateStatelessEmitter)
+				)
+			);
 			MenuBuilder.AddMenuEntry(
 				LOCTEXT("CommentsLabel", "Add Comment"),
 				LOCTEXT("AddCommentBoxToolTip", "Add a comment box"),
@@ -286,6 +293,25 @@ void SNiagaraOverviewGraph::OnCreateEmptyEmitter()
 	ViewModel->GetSystemViewModel()->AddMinimalEmitter();
 }
 
+bool SNiagaraOverviewGraph::CanCreateStatelessEmitter() const
+{
+	FClassViewerModule* ClassViewerModule = FModuleManager::GetModulePtr<FClassViewerModule>("ClassViewer");
+	const TSharedPtr<IClassViewerFilter> GlobalClassFilter = ClassViewerModule ? ClassViewerModule->GetGlobalClassViewerFilter() : TSharedPtr<IClassViewerFilter>();
+	if (GlobalClassFilter.IsValid())
+	{
+		return GlobalClassFilter->IsClassAllowed(FClassViewerInitializationOptions(), UNiagaraStatelessEmitter::StaticClass(), ClassViewerModule->CreateFilterFuncs());
+	}
+
+	return true;
+}
+
+FText SNiagaraOverviewGraph::GetCreateStatelessEmitterTooltip() const
+{
+	return CanCreateStatelessEmitter() ?
+		LOCTEXT("AddStatelessEmitterToolTip", "Adds a lightweight emitter without any modules or renderers.") :
+		LOCTEXT("AddStatelessEmitterNotAllowedToolTip", "Lightweight emitters are not available due to project restrictions.");
+}
+
 void SNiagaraOverviewGraph::OnCreateStatelessEmitter()
 {
 	ViewModel->GetSystemViewModel()->AddStatelessEmitter();
diff --git a/Engine/Plugins/FX/Niagara/Source/NiagaraEditorWidgets/Private/SNiagaraOverviewGraph.h b/Engine/Plugins/FX/Niagara/Source/NiagaraEditorWidgets/Private/SNiagaraOverviewGraph.h
--- a/Engine/Plugins/FX/Niagara/Source/NiagaraEditorWidgets/Private/SNiagaraOverviewGraph.h
+++ b/Engine/Plugins/FX/Niagara/Source/NiagaraEditorWidgets/Private/SNiagaraOverviewGraph.h
@@ -36,6 +36,8 @@ private:
 	FActionMenuContent OnCreateGraphActionMenu(UEdGraph* InGraph, const FVector2D& InNodePosition, const TArray<UEdGraphPin*>& InDraggedPins, bool bAutoExpand, SGraphEditor::FActionMenuClosed InOnMenuClosed);
 
 	void OnCreateEmptyEmitter();
+	bool CanCreateStatelessEmitter() const;
+	FText GetCreateStatelessEmitterTooltip() const;
 	void OnCreateStatelessEmitter();
 	void OnCreateComment();
 	void OnClearIsolated();
