- Downgrade check to ensure
- We still have some paths occasionally hitting this but the repro is not clear.  This will make sure we don't take down the editor.

#rnx

diff --git a/Engine/Plugins/FX/Niagara/Source/Niagara/Public/NiagaraParameterStore.h b/Engine/Plugins/FX/Niagara/Source/Niagara/Public/NiagaraParameterStore.h
--- a/Engine/Plugins/FX/Niagara/Source/Niagara/Public/NiagaraParameterStore.h
+++ b/Engine/Plugins/FX/Niagara/Source/Niagara/Public/NiagaraParameterStore.h
@@ -446,14 +446,14 @@ public:
 
 	FORCEINLINE const uint8* GetParameterData(int32 Offset, const FNiagaraTypeDefinition& TypeDef) const
 	{
-		check(Offset >= 0 && (Offset + TypeDef.GetSize()) <= ParameterData.Num());
+		ensureMsgf(Offset >= 0 && (Offset + TypeDef.GetSize()) <= ParameterData.Num(), TEXT("OOB immutable access offset(%d) size(%d) store size(%d)"), Offset, TypeDef.GetSize(), ParameterData.Num());
 		return ParameterData.GetData() + Offset;
 	}
 
 	/** Returns a none const pointer to the parameter data, any modifications to the data will not be recorded and automatically pushed to other parameter stores, use with caution. */
 	FORCEINLINE uint8* GetMutableParameterData(int32 Offset, const FNiagaraTypeDefinition& TypeDef)
 	{
-		check(Offset >= 0 && (Offset + TypeDef.GetSize()) <= ParameterData.Num());
+		ensureMsgf(Offset >= 0 && (Offset + TypeDef.GetSize()) <= ParameterData.Num(), TEXT("OOB mutable access offset(%d) size(%d) store size(%d)"), Offset, TypeDef.GetSize(), ParameterData.Num());
 		return ParameterData.GetData() + Offset;
 	}
 
