- Fix mismatch between computed renderer draw order and generating the render data
- Stateless didn't exactly follow stateful in this regard so when combined with scalability we got a mismatch

diff --git a/Engine/Plugins/FX/Niagara/Source/Niagara/Internal/Stateless/NiagaraStatelessEmitter.h b/Engine/Plugins/FX/Niagara/Source/Niagara/Internal/Stateless/NiagaraStatelessEmitter.h
--- a/Engine/Plugins/FX/Niagara/Source/Niagara/Internal/Stateless/NiagaraStatelessEmitter.h
+++ b/Engine/Plugins/FX/Niagara/Source/Niagara/Internal/Stateless/NiagaraStatelessEmitter.h
@@ -65,6 +65,7 @@ public:
 
 	void CacheFromCompiledData();
 protected:
+	ENiagaraSimTarget ComputeSimTarget() const;
 	void BuildCompiledDataSet();
 	void ResolveScalabilitySettings();
 
@@ -72,9 +73,15 @@ public:
 	template<typename TAction>
 	void ForEachEnabledRenderer(TAction Func) const
 	{
+		if (RendererProperties.Num() == 0)
+		{
+			return;
+		}
+
+		const ENiagaraSimTarget SimTarget = ComputeSimTarget();
 		for (UNiagaraRendererProperties* Renderer : RendererProperties)
 		{
-			if (Renderer && Renderer->GetIsEnabled())
+			if (Renderer && Renderer->GetIsEnabled() && Renderer->IsSimTargetSupported(SimTarget))
 			{
 				Func(Renderer);
 			}
diff --git a/Engine/Plugins/FX/Niagara/Source/Niagara/Private/Stateless/NiagaraStatelessEmitter.cpp b/Engine/Plugins/FX/Niagara/Source/Niagara/Private/Stateless/NiagaraStatelessEmitter.cpp
--- a/Engine/Plugins/FX/Niagara/Source/Niagara/Private/Stateless/NiagaraStatelessEmitter.cpp
+++ b/Engine/Plugins/FX/Niagara/Source/Niagara/Private/Stateless/NiagaraStatelessEmitter.cpp
@@ -364,6 +364,9 @@ void UNiagaraStatelessEmitter::CacheFromCompiledData()
 	// Resolve scalability settings
 	ResolveScalabilitySettings();
 
+	// Compute the simulation target, this must always be done to ensure we filter renderers correctly
+	StatelessEmitterData->SimTarget = ComputeSimTarget();
+
 	// Build buffers that are shared across all instances
 	//-OPT: We should be able to build and serialize this data as part of the UNiagaraStatelessEmitter, potentially all of this data even since it's immutable and does not change at runtime
 	if (StatelessEmitterData->bCanEverExecute)
@@ -455,10 +458,6 @@ void UNiagaraStatelessEmitter::CacheFromCompiledData()
 			if (Renderer && Renderer->GetIsEnabled())
 			{
 				Renderer->PopulateRequiredBindings(StatelessEmitterData->RendererBindings);
-				if (Renderer->IsSimTargetSupported(ENiagaraSimTarget::GPUComputeSim) == false)
-				{
-					StatelessEmitterData->SimTarget = ENiagaraSimTarget::CPUSim;
-				}
 			}
 		}
 
@@ -470,6 +469,21 @@ void UNiagaraStatelessEmitter::CacheFromCompiledData()
 	}
 }
 
+ENiagaraSimTarget UNiagaraStatelessEmitter::ComputeSimTarget() const
+{
+	for (UNiagaraRendererProperties* Renderer : RendererProperties)
+	{
+		if (Renderer && Renderer->GetIsEnabled())
+		{
+			if (Renderer->IsSimTargetSupported(ENiagaraSimTarget::GPUComputeSim) == false)
+			{
+				return ENiagaraSimTarget::CPUSim;
+			}
+		}
+	}
+	return ENiagaraSimTarget::GPUComputeSim;
+}
+
 void UNiagaraStatelessEmitter::BuildCompiledDataSet()
 {
 #if WITH_EDITORONLY_DATA
