- When using performance mode disable compile for edit to ensure performance measurements is more accurate
- One measurement I made was ~40% delta between edit mode and none edit mode
- There is a separate option to enable edit mode for profiling

#tests Profile a system in the editor, toggle the profile in edit mode button on/off, exit editor, reopen and ensure the state is set correctly.

diff --git a/Engine/Plugins/FX/Niagara/Source/NiagaraEditor/Private/NiagaraEditorCommands.cpp b/Engine/Plugins/FX/Niagara/Source/NiagaraEditor/Private/NiagaraEditorCommands.cpp
--- a/Engine/Plugins/FX/Niagara/Source/NiagaraEditor/Private/NiagaraEditorCommands.cpp
+++ b/Engine/Plugins/FX/Niagara/Source/NiagaraEditor/Private/NiagaraEditorCommands.cpp
@@ -33,6 +33,7 @@ void FNiagaraEditorCommands::RegisterCommands()
 	UI_COMMAND(ToggleUnlockToChanges, "Lock/Unlock To Changes", "Toggles whether or not changes in the source asset get pulled into this asset automatically.", EUserInterfaceActionType::Button, FInputChord());
 	UI_COMMAND(ToggleStatPerformance, "Performance", "Show runtime performance for particle scripts. Not supported for optimized VM byte code.", EUserInterfaceActionType::ToggleButton, FInputChord());
 	UI_COMMAND(ToggleStatPerformanceGPU, "GPU Profiling", "Measure gpu runtime cost of scripts. Can be expensive for simulation stages with lots of iterations.", EUserInterfaceActionType::ToggleButton, FInputChord());
+	UI_COMMAND(ToggleProfileInEditMode, "Profile In Edit Mode", "When enabled and profiling is active choose how we compile the system.  Edit mode is worse performance but won't recompile when tweaking parameters.", EUserInterfaceActionType::ToggleButton, FInputChord());
 	UI_COMMAND(ClearStatPerformance, "Clear Stats", "Clear all existing stat captures.", EUserInterfaceActionType::Button, FInputChord());
 	UI_COMMAND(ToggleStatPerformanceTypeAvg, "Display Average", "Displays the average of captured stats.", EUserInterfaceActionType::ToggleButton, FInputChord());
 	UI_COMMAND(ToggleStatPerformanceTypeMax, "Display Maximum", "Displays the maximum of captured stats.", EUserInterfaceActionType::ToggleButton, FInputChord());
diff --git a/Engine/Plugins/FX/Niagara/Source/NiagaraEditor/Private/Toolkits/NiagaraSystemToolkit.cpp b/Engine/Plugins/FX/Niagara/Source/NiagaraEditor/Private/Toolkits/NiagaraSystemToolkit.cpp
--- a/Engine/Plugins/FX/Niagara/Source/NiagaraEditor/Private/Toolkits/NiagaraSystemToolkit.cpp
+++ b/Engine/Plugins/FX/Niagara/Source/NiagaraEditor/Private/Toolkits/NiagaraSystemToolkit.cpp
@@ -147,6 +147,7 @@ void FNiagaraSystemToolkit::InitializeWithSystem(const EToolkitMode::Type Mode,
 	SystemOptions.EditMode = ENiagaraSystemViewModelEditMode::SystemAsset;
 	SystemOptions.OnGetSequencerAddMenuContent.BindSP(this, &FNiagaraSystemToolkit::GetSequencerAddMenuContent);
 	SystemOptions.MessageLogGuid = InSystem.GetAssetGuid();
+	SystemOptions.bCompileForEdit = IsStatPerformanceChecked() ? IsProfileInEditMode() : true;
 
 	SystemViewModel = MakeShared<FNiagaraSystemViewModel>();
 	SystemViewModel->Initialize(*System, SystemOptions);
@@ -433,6 +434,11 @@ void FNiagaraSystemToolkit::SetupCommands()
 	GetToolkitCommands()->MapAction(
         FNiagaraEditorCommands::Get().ClearStatPerformance,
         FExecuteAction::CreateSP(this, &FNiagaraSystemToolkit::ClearStatPerformance));
+	GetToolkitCommands()->MapAction(
+		FNiagaraEditorCommands::Get().ToggleProfileInEditMode,
+		FExecuteAction::CreateSP(this, &FNiagaraSystemToolkit::ToggleProfileInEditMode),
+		FCanExecuteAction(),
+		FIsActionChecked::CreateSP(this, &FNiagaraSystemToolkit::IsProfileInEditMode));
 	GetToolkitCommands()->MapAction(
         FNiagaraEditorCommands::Get().ToggleStatPerformanceGPU,
         FExecuteAction::CreateSP(this, &FNiagaraSystemToolkit::ToggleStatPerformanceGPU),
@@ -719,6 +725,7 @@ TSharedRef<SWidget> FNiagaraSystemToolkit::GenerateStatConfigMenuContent(TShared
 	FMenuBuilder MenuBuilder(bShouldCloseWindowAfterMenuSelection, InCommandList);
 
 	MenuBuilder.AddMenuEntry(FNiagaraEditorCommands::Get().ClearStatPerformance);
+	MenuBuilder.AddMenuEntry(FNiagaraEditorCommands::Get().ToggleProfileInEditMode);
 	MenuBuilder.AddMenuSeparator();
 	MenuBuilder.AddMenuEntry(FNiagaraEditorCommands::Get().ToggleStatPerformanceGPU);
 	MenuBuilder.AddMenuSeparator();
@@ -1026,10 +1033,16 @@ void FNiagaraSystemToolkit::ClearStatPerformance()
 
 void FNiagaraSystemToolkit::ToggleStatPerformance()
 {
-	bool IsEnabled = IsStatPerformanceChecked();
+	const bool IsEnabled = IsStatPerformanceChecked();
 	if (VmStatEnabledVar)
 	{
 		VmStatEnabledVar->Set(!IsEnabled);
+
+		if (SystemViewModel.IsValid())
+		{
+			const bool bCompileForEdit = !IsEnabled ? bProfileInEditMode : true;
+			SystemViewModel->SetCompileForEdit(bCompileForEdit);
+		}
 	}
 	if (IsStatPerformanceGPUChecked() == IsEnabled)
 	{
@@ -1103,6 +1116,21 @@ bool FNiagaraSystemToolkit::IsStatPerformanceGPUChecked()
 	return false;
 }
 
+void FNiagaraSystemToolkit::ToggleProfileInEditMode()
+{
+	bProfileInEditMode = !bProfileInEditMode;
+
+	if (SystemViewModel.IsValid() && IsStatPerformanceChecked())
+	{
+		SystemViewModel->SetCompileForEdit(bProfileInEditMode);
+	}
+}
+
+bool FNiagaraSystemToolkit::IsProfileInEditMode() const
+{
+	return bProfileInEditMode;
+}
+
 void FNiagaraSystemToolkit::UpdateOriginalEmitter()
 {
 	checkf(SystemToolkitMode == ESystemToolkitMode::Emitter, TEXT("There is no original emitter to update in system mode."));
diff --git a/Engine/Plugins/FX/Niagara/Source/NiagaraEditor/Private/Toolkits/NiagaraSystemToolkit.h b/Engine/Plugins/FX/Niagara/Source/NiagaraEditor/Private/Toolkits/NiagaraSystemToolkit.h
--- a/Engine/Plugins/FX/Niagara/Source/NiagaraEditor/Private/Toolkits/NiagaraSystemToolkit.h
+++ b/Engine/Plugins/FX/Niagara/Source/NiagaraEditor/Private/Toolkits/NiagaraSystemToolkit.h
@@ -145,6 +145,8 @@ protected:
 	void ClearStatPerformance();
 	void ToggleStatPerformance();
 	bool IsStatPerformanceChecked();
+	void ToggleProfileInEditMode();
+	bool IsProfileInEditMode() const;
 	void ToggleStatPerformanceGPU();
 	bool IsStatPerformanceGPUChecked();
 	void ToggleStatPerformanceTypeAvg();
@@ -211,8 +213,8 @@ private:
 	void OpenAddEmitterMenu();
 	bool CanAddEmitters() const;
 	void GetSequencerAddMenuContent(FMenuBuilder& MenuBuilder, TSharedRef<ISequencer> Sequencer);
-private:
 
+private:
 	/** The System being edited in system mode, or the placeholder system being edited in emitter mode. */
 	TObjectPtr<UNiagaraSystem> System = nullptr;
 
@@ -232,6 +234,8 @@ private:
 		screenshot uobject is reused, so a pointer comparison doesn't work to checking if the images has been updated. */
 	bool bEmitterThumbnailUpdated;
 
+	bool bProfileInEditMode = false;
+
 	ESystemToolkitMode SystemToolkitMode;
 
 	/** The currently active workflow mode */
diff --git a/Engine/Plugins/FX/Niagara/Source/NiagaraEditor/Private/ViewModels/NiagaraSystemViewModel.cpp b/Engine/Plugins/FX/Niagara/Source/NiagaraEditor/Private/ViewModels/NiagaraSystemViewModel.cpp
--- a/Engine/Plugins/FX/Niagara/Source/NiagaraEditor/Private/ViewModels/NiagaraSystemViewModel.cpp
+++ b/Engine/Plugins/FX/Niagara/Source/NiagaraEditor/Private/ViewModels/NiagaraSystemViewModel.cpp
@@ -81,13 +81,6 @@ DECLARE_CYCLE_STAT(TEXT("Niagara - SystemViewModel - CompileSystem"), STAT_Niaga
 
 template<> TMap<UNiagaraSystem*, TArray<FNiagaraSystemViewModel*>> TNiagaraViewModelManager<UNiagaraSystem, FNiagaraSystemViewModel>::ObjectsToViewModels{};
 
-FNiagaraSystemViewModelOptions::FNiagaraSystemViewModelOptions()
-	: bCanAutoCompile(true)
-	, bCanSimulate(true)
-	, bIsForDataProcessingOnly(false)
-{
-}
-
 FNiagaraSystemViewModel::FNiagaraSystemViewModel()
 	: System(nullptr)
 	, PreviewComponent(nullptr)
@@ -95,7 +88,7 @@ FNiagaraSystemViewModel::FNiagaraSystemViewModel()
 	, NiagaraSequence(nullptr)
 	, bSettingSequencerTimeDirectly(false)
 	, bForceAutoCompileOnce(false)
-	, bSupportCompileForEdit(true)
+	, bSupportCompileForEdit(false)
 	, bUpdatingEmittersFromSequencerDataChange(false)
 	, bUpdatingSequencerFromEmitterDataChange(false)
 	, bUpdatingSystemSelectionFromSequencer(false)
@@ -170,9 +163,12 @@ void FNiagaraSystemViewModel::Initialize(UNiagaraSystem& InSystem, FNiagaraSyste
 	SystemGraphSelectionViewModel = MakeShared<FNiagaraSystemGraphSelectionViewModel>();
 	SystemGraphSelectionViewModel->Initialize(this->AsShared());
 
-	if (CanImpactCompileForEdit())
+	bSupportCompileForEdit = !bIsForDataProcessingOnly && (EditMode == ENiagaraSystemViewModelEditMode::SystemAsset || EditMode == ENiagaraSystemViewModelEditMode::EmitterAsset);
+	bIsCompileForEdit = false;
+	if (bSupportCompileForEdit)
 	{
-		System->SetCompileForEdit(true);
+		bIsCompileForEdit = InOptions.bCompileForEdit;
+		System->SetCompileForEdit(bIsCompileForEdit);
 	}
 	
 	if (bIsForDataProcessingOnly == false)
@@ -266,7 +262,7 @@ void FNiagaraSystemViewModel::Cleanup()
 
 	if (System)
 	{
-		if (CanImpactCompileForEdit())
+		if (bSupportCompileForEdit && bIsCompileForEdit)
 		{
 			System->SetCompileForEdit(false);
 			System->InvalidateActiveCompiles();
@@ -434,6 +430,19 @@ TSharedPtr<FNiagaraSystemScriptViewModel> FNiagaraSystemViewModel::GetSystemScri
 	return SystemScriptViewModel;
 }
 
+void FNiagaraSystemViewModel::SetCompileForEdit(bool bNewState)
+{
+	if (bSupportCompileForEdit && bIsCompileForEdit != bNewState)
+	{
+		check(System->GetCompileForEdit() == bIsCompileForEdit);
+		bIsCompileForEdit = bNewState;
+
+		System->SetCompileForEdit(bNewState);
+		System->InvalidateActiveCompiles();
+		System->RequestCompile(false);
+	}
+}
+
 void FNiagaraSystemViewModel::CompileSystem(bool bForce)
 {
 	if (EditMode == ENiagaraSystemViewModelEditMode::EmitterDuringMerge && bForce == false)
@@ -1149,7 +1158,7 @@ void FNiagaraSystemViewModel::NotifyPreSave()
 	}
 
 	// we want to save without compile for edit turned on for best results, but make sure to turn in back on again after save is done
-	if (CanImpactCompileForEdit())
+	if (bSupportCompileForEdit && bIsCompileForEdit)
 	{
 		check(System->GetCompileForEdit());
 		System->SetCompileForEdit(false);
@@ -1163,7 +1172,7 @@ void FNiagaraSystemViewModel::NotifyPreSave()
 void FNiagaraSystemViewModel::NotifyPostSave()
 {
 	// we compile for edit again after having it turned off for PreSave. This should fetch DDC data and must not update the ChangeID
-	if (CanImpactCompileForEdit())
+	if (bSupportCompileForEdit && bIsCompileForEdit)
 	{
 		check(!System->GetCompileForEdit());
 		System->SetCompileForEdit(true);
@@ -3512,9 +3521,4 @@ TSharedPtr<FNiagaraOverviewGraphViewModel> FNiagaraSystemViewModel::GetOverviewG
 	return OverviewGraphViewModel;
 }
 
-bool FNiagaraSystemViewModel::CanImpactCompileForEdit() const
-{
-	return bSupportCompileForEdit && !bIsForDataProcessingOnly && (EditMode == ENiagaraSystemViewModelEditMode::SystemAsset || EditMode == ENiagaraSystemViewModelEditMode::EmitterAsset);
-}
-
 #undef LOCTEXT_NAMESPACE // NiagaraSystemViewModel
diff --git a/Engine/Plugins/FX/Niagara/Source/NiagaraEditor/Public/NiagaraEditorCommands.h b/Engine/Plugins/FX/Niagara/Source/NiagaraEditor/Public/NiagaraEditorCommands.h
--- a/Engine/Plugins/FX/Niagara/Source/NiagaraEditor/Public/NiagaraEditorCommands.h
+++ b/Engine/Plugins/FX/Niagara/Source/NiagaraEditor/Public/NiagaraEditorCommands.h
@@ -61,6 +61,7 @@ public:
 
 	TSharedPtr<FUICommandInfo> ToggleStatPerformance;
 	TSharedPtr<FUICommandInfo> ToggleStatPerformanceGPU;
+	TSharedPtr<FUICommandInfo> ToggleProfileInEditMode;
 	TSharedPtr<FUICommandInfo> ClearStatPerformance;
 	TSharedPtr<FUICommandInfo> ToggleStatPerformanceTypeAvg;
 	TSharedPtr<FUICommandInfo> ToggleStatPerformanceTypeMax;
diff --git a/Engine/Plugins/FX/Niagara/Source/NiagaraEditor/Public/ViewModels/NiagaraSystemViewModel.h b/Engine/Plugins/FX/Niagara/Source/NiagaraEditor/Public/ViewModels/NiagaraSystemViewModel.h
--- a/Engine/Plugins/FX/Niagara/Source/NiagaraEditor/Public/ViewModels/NiagaraSystemViewModel.h
+++ b/Engine/Plugins/FX/Niagara/Source/NiagaraEditor/Public/ViewModels/NiagaraSystemViewModel.h
@@ -77,8 +77,6 @@ DECLARE_DELEGATE_RetVal(FName, FOnGetWorkflowMode);
 /** Defines options for the niagara System view model */
 struct FNiagaraSystemViewModelOptions
 {
-	NIAGARAEDITOR_API FNiagaraSystemViewModelOptions();
-
 	/** Whether or not the user can edit emitters from the timeline. */
 	bool bCanModifyEmittersFromTimeline = true;
 
@@ -86,10 +84,10 @@ struct FNiagaraSystemViewModelOptions
 	FOnGetAddMenuContent OnGetSequencerAddMenuContent;
 
 	/** Whether or not the system represented by this view model can be automatically compiled.  True by default. */
-	bool bCanAutoCompile;
+	bool bCanAutoCompile = true;
 
 	/** Whether or not the system represented by this view model can be simulated. True by default. */
-	bool bCanSimulate;
+	bool bCanSimulate = true;
 
 	/** An optional unique id which associates this view model with messages in the message manager. */
 	TOptional<const FGuid> MessageLogGuid;
@@ -98,7 +96,13 @@ struct FNiagaraSystemViewModelOptions
 	ENiagaraSystemViewModelEditMode EditMode = ENiagaraSystemViewModelEditMode::SystemAsset;
 
 	/** Specifies that the view model is being constructed for data processing only and will not be displayed in the UI. */
-	bool bIsForDataProcessingOnly;
+	bool bIsForDataProcessingOnly = false;
+
+	/**
+	Should we compile the system for edit or not, this is only relevant if it's allowed.
+	For example when bIsForDataProcessingOnly is true we won't allow edit mode.
+	*/
+	bool bCompileForEdit = true;
 };
 
 struct FNiagaraStackModuleData
@@ -309,6 +313,11 @@ public:
 	NIAGARAEDITOR_API void ToggleIsolateEmitterAndSelectedEmitters(FGuid OriginEmitter);
 	NIAGARAEDITOR_API void IsolateSelectedEmitters();
 	void DisableSelectedEmitters();
+
+	/** Update the compile for edit mode. */
+	void SetCompileForEdit(bool bNewState);
+	/** Are we running in compile for edit mode or not. */
+	bool IsCompileForEdit() const { return bIsCompileForEdit; }
 	
 	/** Compiles the spawn and update scripts. */
 	void CompileSystem(bool bForce);
@@ -631,9 +640,6 @@ private:
 	/** Utility method to gather all emitter handle viewmodels marked as selected via the system selection viewmodel. */
 	TArray<TSharedRef<FNiagaraEmitterHandleViewModel>> GetSelectedEmitterHandleViewModels();
 
-	/** Indicates whether this instance of the view model should impact whether the System's CompileForEdit flag should be updated by the viewmodel. */
-	bool CanImpactCompileForEdit() const;
-
 private:
 	/** The System being viewed and edited by this view model. */
 	UNiagaraSystem* System;
@@ -683,6 +689,9 @@ private:
 	/** Whether we should use the CompileForEdit mode for the assigned system. */
 	bool bSupportCompileForEdit;
 
+	/** What mode compile for edit is currently in. */
+	bool bIsCompileForEdit;
+
 	/** The current editing mode for this view model. */
 	ENiagaraSystemViewModelEditMode EditMode;
 
diff --git a/Engine/Plugins/FX/Niagara/Source/NiagaraEditorWidgets/Private/SNiagaraOverviewGraphTitleBar.cpp b/Engine/Plugins/FX/Niagara/Source/NiagaraEditorWidgets/Private/SNiagaraOverviewGraphTitleBar.cpp
--- a/Engine/Plugins/FX/Niagara/Source/NiagaraEditorWidgets/Private/SNiagaraOverviewGraphTitleBar.cpp
+++ b/Engine/Plugins/FX/Niagara/Source/NiagaraEditorWidgets/Private/SNiagaraOverviewGraphTitleBar.cpp
@@ -107,6 +107,27 @@ void SNiagaraOverviewGraphTitleBar::RebuildWidget()
 		];
 	}
 	
+	// Warning about compilation mode setting and impacting iteration
+	if (SystemViewModel.IsValid() && SystemViewModel.Pin()->GetEditMode() == ENiagaraSystemViewModelEditMode::SystemAsset)
+	{
+		ContainerWidget->AddSlot()
+		.HAlign(HAlign_Fill)
+		.AutoHeight()
+		[
+			SNew(SBorder)
+			.BorderImage(FAppStyle::GetBrush(TEXT("Graph.TitleBackground")))
+			.ColorAndOpacity(this, &SNiagaraOverviewGraphTitleBar::GetSystemCompileModeHeaderColor)
+			.Visibility(this, &SNiagaraOverviewGraphTitleBar::GetSystemCompileModeHeaderVisibility)
+			.HAlign(HAlign_Fill)
+			[
+				SNew(STextBlock)
+				.Text(this, &SNiagaraOverviewGraphTitleBar::GetSystemCompileModeHeaderText)
+				.TextStyle(FAppStyle::Get(), TEXT("GraphBreadcrumbButtonText"))
+				.Justification(ETextJustify::Center)
+			]
+		];
+	}
+
 	if(bScalabilityModeActive)
 	{
 		ContainerWidget->AddSlot()
@@ -185,6 +206,26 @@ FLinearColor SNiagaraOverviewGraphTitleBar::GetSystemSubheaderColor() const
 	return IsUsingDeprecatedEmitter() ? FNiagaraEditorWidgetsStyle::Get().GetColor("NiagaraEditor.SystemOverview.AffectedAssetsWarningColor") : FLinearColor::White;
 }
 
+EVisibility SNiagaraOverviewGraphTitleBar::GetSystemCompileModeHeaderVisibility() const
+{
+	if (SystemViewModel.IsValid())
+	{
+		const UNiagaraSystem& NiagaraSystem = SystemViewModel.Pin()->GetSystem();
+		return NiagaraSystem.GetCompileForEdit() ? EVisibility::Hidden : EVisibility::Visible;
+	}
+	return EVisibility::Hidden;
+}
+
+FText SNiagaraOverviewGraphTitleBar::GetSystemCompileModeHeaderText() const
+{
+	return LOCTEXT("SystemInCompiledForPerformance", "Note: System is compiled for performance.  Edits may result in more compilation.");
+}
+
+FLinearColor SNiagaraOverviewGraphTitleBar::GetSystemCompileModeHeaderColor() const
+{
+	return FNiagaraEditorWidgetsStyle::Get().GetColor("NiagaraEditor.SystemOverview.AffectedAssetsWarningColor");
+}
+
 void SNiagaraOverviewGraphTitleBar::ResetAssetCount(const FAssetData&)
 {
 	EmitterAffectedAssets.Reset();
diff --git a/Engine/Plugins/FX/Niagara/Source/NiagaraEditorWidgets/Private/SNiagaraOverviewGraphTitleBar.h b/Engine/Plugins/FX/Niagara/Source/NiagaraEditorWidgets/Private/SNiagaraOverviewGraphTitleBar.h
--- a/Engine/Plugins/FX/Niagara/Source/NiagaraEditorWidgets/Private/SNiagaraOverviewGraphTitleBar.h
+++ b/Engine/Plugins/FX/Niagara/Source/NiagaraEditorWidgets/Private/SNiagaraOverviewGraphTitleBar.h
@@ -44,7 +44,10 @@ private:
 	EVisibility GetSystemSubheaderVisibility() const;
 	FLinearColor GetSystemSubheaderColor() const;
 	
-	
+	EVisibility GetSystemCompileModeHeaderVisibility() const;
+	FText GetSystemCompileModeHeaderText() const;
+	FLinearColor GetSystemCompileModeHeaderColor() const;
+
 	int32 GetEmitterAffectedAssets() const;
 	bool IsUsingDeprecatedEmitter() const;
 	void ResetAssetCount(const FAssetData& InAssetData);
