- Change check to ensure to avoid crashing the cook / editor
- This can trigger from bad content, which can still pass compilation.  One example was a static variable which referenced an enum that was deleted.  The module itself or the path through the module was not used in the final output so it's safe to fail.

#rnx

diff --git a/Engine/Plugins/FX/Niagara/Source/NiagaraEditor/Private/NiagaraGraphDigest.cpp b/Engine/Plugins/FX/Niagara/Source/NiagaraEditor/Private/NiagaraGraphDigest.cpp
--- a/Engine/Plugins/FX/Niagara/Source/NiagaraEditor/Private/NiagaraGraphDigest.cpp
+++ b/Engine/Plugins/FX/Niagara/Source/NiagaraEditor/Private/NiagaraGraphDigest.cpp
@@ -1056,7 +1056,7 @@ void FNiagaraCompilationGraphInstanced::ValidateRefinement() const
 		{
 			case FNiagaraCompilationNode::ENodeType::StaticSwitch:
 				// only switches handled by the compiler are allowed to make it to this stage
-				check(Node->AsType<FNiagaraCompilationNodeStaticSwitch>()->bSetByCompiler);
+				ensure(Node->AsType<FNiagaraCompilationNodeStaticSwitch>()->bSetByCompiler);
 			break;
 
 			// it's possible that we could eliminate these during instantiation as well, but currently
