- Add missing null check for EmitterData

#rnx

diff --git a/Engine/Plugins/FX/Niagara/Source/Niagara/Private/Stateless/NiagaraStatelessEmitterInstance.cpp b/Engine/Plugins/FX/Niagara/Source/Niagara/Private/Stateless/NiagaraStatelessEmitterInstance.cpp
--- a/Engine/Plugins/FX/Niagara/Source/Niagara/Private/Stateless/NiagaraStatelessEmitterInstance.cpp
+++ b/Engine/Plugins/FX/Niagara/Source/Niagara/Private/Stateless/NiagaraStatelessEmitterInstance.cpp
@@ -937,12 +937,17 @@ void FNiagaraStatelessEmitterInstance::InitEmitterData()
 		return;
 	}
 	EmitterData = StatelessEmitter->GetEmitterData();
-	SimTarget = EmitterData->SimTarget;
 
-	bCanEverExecute =
-		EmitterData != nullptr &&
-		EmitterData->bCanEverExecute &&
-		EmitterHandle.GetIsEnabled();
+	if (EmitterData != nullptr)
+	{
+		// We need to extract the sim target here to ensure we are in sync when iterating available renderers
+		SimTarget = EmitterData->SimTarget;
+		bCanEverExecute = EmitterData->bCanEverExecute && EmitterHandle.GetIsEnabled();
+	}
+	else
+	{
+		bCanEverExecute = false;
+	}
 }
 
 void FNiagaraStatelessEmitterInstance::CaptureForDebugging(FNiagaraDataBuffer* DataBuffer) const
