[AutoRTFM] Use a transactionally-safe RW mutex for UNiagaraDataInterfaceArray::ArrayRWGuard

This can be constructed from Verse, and a FRWLock cannot be constructed in a closed transaction.

diff --git a/Engine/Plugins/FX/Niagara/Source/Niagara/Classes/NiagaraDataInterfaceArray.h b/Engine/Plugins/FX/Niagara/Source/Niagara/Classes/NiagaraDataInterfaceArray.h
--- a/Engine/Plugins/FX/Niagara/Source/Niagara/Classes/NiagaraDataInterfaceArray.h
+++ b/Engine/Plugins/FX/Niagara/Source/Niagara/Classes/NiagaraDataInterfaceArray.h
@@ -2,6 +2,7 @@
 
 #pragma once
 
+#include "Misc/TransactionallySafeRWLock.h"
 #include "NiagaraDataInterfaceRW.h"
 #include "NiagaraCommon.h"
 #include "NiagaraShared.h"
@@ -131,7 +132,7 @@ public:
 	NIAGARA_API FString SimCacheVisualizerRead(const UNDIArraySimCacheData* CacheData, const FNDIArraySimCacheDataFrame& FrameData, int Element) const;
 
 	/** ReadWrite lock to ensure safe access to the underlying array. */
-	FRWLock ArrayRWGuard;
+	FTransactionallySafeRWLock ArrayRWGuard;
 
 	/** How to do we want to synchronize modifications to the array data. */
 	UPROPERTY(EditAnywhere, BlueprintReadOnly, AdvancedDisplay, Category="Array")
