Added a new mode for emitter interpolated spawn to only run the spawn + update script without value interpolation.

#tests Niagara engine tests still work, so there should be no changes to existing content

diff --git a/Engine/Plugins/FX/Niagara/Source/Niagara/Classes/NiagaraEmitter.h b/Engine/Plugins/FX/Niagara/Source/Niagara/Classes/NiagaraEmitter.h
--- a/Engine/Plugins/FX/Niagara/Source/Niagara/Classes/NiagaraEmitter.h
+++ b/Engine/Plugins/FX/Niagara/Source/Niagara/Classes/NiagaraEmitter.h
@@ -282,15 +282,15 @@ struct FVersionedNiagaraEmitterData
 	UPROPERTY(EditAnywhere, BlueprintReadWrite, Category = "Emitter", meta = (EditCondition = "bDeterminism", EditConditionHides="true"))
 	int32 RandomSeed = 0;
 
-	/**
-	When enabled, this will spawn using interpolated parameter values and perform a partial update at spawn time.
-	This adds additional cost for spawning but will produce much smoother results for high spawn rates, erratic frame rates and fast moving emitters.
-	NOTE: When not disabled only the spawn script will run for newly spawned particles.  We do not run the update script.
-	*/
-	UPROPERTY(EditAnywhere, BlueprintReadWrite, Category = "Emitter")
-	uint32 bInterpolatedSpawning : 1;
+	/** This defines if newly spawned particles run only the spawn script on the first frame or both spawn and update with optional parameter interpolation. */
+	UPROPERTY(EditAnywhere, BlueprintReadWrite, Category = "Emitter", meta = (SegmentedDisplay))
+	ENiagaraInterpolatedSpawnMode InterpolatedSpawnMode;
 
 #if WITH_EDITORONLY_DATA
+
+	UPROPERTY(meta = (DeprecatedProperty))
+	uint32 bInterpolatedSpawning_DEPRECATED : 1;
+	
 	/**
 	GPU scripts were incorrectly running both particle spawn & update, CPU only runs spawn when interpolated spawning is disabled.
 	This flag allows backwards compatability so content does not change post the fix.
@@ -387,6 +387,7 @@ struct FVersionedNiagaraEmitterData
 	NIAGARA_API void CacheFromCompiledData(const FNiagaraDataSetCompiledData* CompiledData, const UNiagaraEmitter& Emitter);
 	NIAGARA_API void CacheFromShaderCompiled();
 	bool IsAllowedToExecute() const { return bIsAllowedToExecute; }
+	NIAGARA_API bool UsesInterpolatedSpawning() const;
 
 	NIAGARA_API FGraphEventArray PrecacheComputePSOs(const UNiagaraEmitter& NiagaraEmitter);
 	bool DidPSOPrecacheFail() const { return PSOPrecacheResult == EPSOPrecacheResult::NotSupported; }
diff --git a/Engine/Plugins/FX/Niagara/Source/Niagara/Classes/NiagaraScript.h b/Engine/Plugins/FX/Niagara/Source/Niagara/Classes/NiagaraScript.h
--- a/Engine/Plugins/FX/Niagara/Source/Niagara/Classes/NiagaraScript.h
+++ b/Engine/Plugins/FX/Niagara/Source/Niagara/Classes/NiagaraScript.h
@@ -247,15 +247,19 @@ public:
 	UPROPERTY()
 	FGuid CompilerVersionID;
 
-	/** The type of script this was used for.*/
+	/** Do we require interpolated spawning */
 	UPROPERTY()
-	ENiagaraScriptUsage ScriptUsageType;
+	ENiagaraInterpolatedSpawnMode InterpolatedSpawnMode;
 
+#if WITH_EDITORONLY_DATA
 	/** The instance id of this script usage type.*/
 	UPROPERTY()
 	FGuid ScriptUsageTypeID;
-
-#if WITH_EDITORONLY_DATA
+	
+	/** The type of script this was used for.*/
+	UPROPERTY()
+	ENiagaraScriptUsage ScriptUsageType;
+	
 	/** Configuration options*/
 	UPROPERTY()
 	TArray<FString> AdditionalDefines;
@@ -264,30 +268,20 @@ public:
 	TArray<FNiagaraVariableBase> AdditionalVariables;
 
 	NIAGARA_API TArray<FString> GetAdditionalVariableStrings() const;
-#endif
-
-	/** Whether or not we need to bake Rapid Iteration params. True to keep params, false to bake.*/
-	UPROPERTY()
-	uint32 bUsesRapidIterationParams : 1;
 
 	/** Whether or not we allow debug switches to be used. */
 	UPROPERTY()
 	uint32 bDisableDebugSwitches : 1;
-
-	/** Do we require interpolated spawning */
-	UPROPERTY()
-	uint32 bInterpolatedSpawn : 1;
-
+	
 	/** Do we require persistent IDs */
 	UPROPERTY()
 	uint32 bRequiresPersistentIDs : 1;
 
-	/**
-	* The GUID of the subgraph this shader primarily represents.
-	*/
+	/** Whether or not we need to bake Rapid Iteration params. True to keep params, false to bake.*/
 	UPROPERTY()
-	FGuid BaseScriptID_DEPRECATED;
-
+	uint32 bUsesRapidIterationParams : 1;
+#endif
+	
 	/**
 	* The hash of the subgraph this shader primarily represents.
 	*/
@@ -298,20 +292,21 @@ public:
 	/** Compile hashes of any top level scripts the script was dependent on that might trigger a recompile if they change. */
 	UPROPERTY()
 	TArray<FNiagaraCompileHash> ReferencedCompileHashes;
-#endif
-
+	
 	/** The version of the script that was compiled. If empty then just the latest version. */
 	UPROPERTY()
 	FGuid ScriptVersionID;
+#endif
 
 	FNiagaraVMExecutableDataId()
 		: CompilerVersionID()
+		, InterpolatedSpawnMode(ENiagaraInterpolatedSpawnMode::NoInterpolation)
+#if WITH_EDITORONLY_DATA
 		, ScriptUsageType(ENiagaraScriptUsage::Function)
-		, bUsesRapidIterationParams(true)
 		, bDisableDebugSwitches(false)
-		, bInterpolatedSpawn(false)
 		, bRequiresPersistentIDs(false)
-		, BaseScriptID_DEPRECATED(0, 0, 0, 0)
+		, bUsesRapidIterationParams(true)
+#endif
 	{ }
 
 
@@ -332,7 +327,6 @@ public:
 	}
 
 	NIAGARA_API bool HasInterpolatedParameters() const;
-	NIAGARA_API bool RequiresPersistentIDs() const;
 #if 0
 	/** Hashes the script-specific part of this shader map Id. */
 	NIAGARA_API void GetScriptHash(FSHAHash& OutHash) const;
diff --git a/Engine/Plugins/FX/Niagara/Source/Niagara/Private/NiagaraEmitter.cpp b/Engine/Plugins/FX/Niagara/Source/Niagara/Private/NiagaraEmitter.cpp
--- a/Engine/Plugins/FX/Niagara/Source/Niagara/Private/NiagaraEmitter.cpp
+++ b/Engine/Plugins/FX/Niagara/Source/Niagara/Private/NiagaraEmitter.cpp
@@ -672,7 +672,7 @@ void UNiagaraEmitter::PostLoad()
 		{
 			if ( Data.SimTarget == ENiagaraSimTarget::GPUComputeSim )
 			{
-				Data.bGpuAlwaysRunParticleUpdateScript = Data.bInterpolatedSpawning ? false : true;
+				Data.bGpuAlwaysRunParticleUpdateScript = Data.UsesInterpolatedSpawning();
 			}
 			else
 			{
@@ -717,6 +717,11 @@ void FVersionedNiagaraEmitterData::PostLoad(UNiagaraEmitter& Emitter, int32 Niag
 			EventScript.UpdateAttributeInitialValues = false;
 		}
 	}
+
+	if (NiagaraVer < FNiagaraCustomVersion::AddedNewInterpolatedSpawnMode)
+	{
+		InterpolatedSpawnMode = bInterpolatedSpawning_DEPRECATED ? ENiagaraInterpolatedSpawnMode::Interpolation : ENiagaraInterpolatedSpawnMode::NoInterpolation;
+	}
 #endif
 
 #if WITH_EDITORONLY_DATA && WITH_EDITOR
@@ -852,9 +857,9 @@ void FVersionedNiagaraEmitterData::PostLoad(UNiagaraEmitter& Emitter, int32 Niag
 	if (SpawnScriptProps.Script)
 	{
 		bool bActualInterpolatedSpawning = SpawnScriptProps.Script->IsInterpolatedParticleSpawnScript();
-		if (bInterpolatedSpawning != bActualInterpolatedSpawning)
+		if (UsesInterpolatedSpawning() != bActualInterpolatedSpawning)
 		{
-			bInterpolatedSpawning = false;
+			InterpolatedSpawnMode = ENiagaraInterpolatedSpawnMode::NoInterpolation;
 			if (bActualInterpolatedSpawning)
 			{
 #if WITH_EDITORONLY_DATA
@@ -1209,14 +1214,15 @@ void UNiagaraEmitter::PostEditChangeVersionedProperty(FPropertyChangedEvent& Pro
 
 	bool bNeedsRecompile = false;
 	bool bRecomputeExecutionOrder = false;
-	if (PropertyName == GET_MEMBER_NAME_CHECKED(FVersionedNiagaraEmitterData, bInterpolatedSpawning))
+	if (PropertyName == GET_MEMBER_NAME_CHECKED(FVersionedNiagaraEmitterData, InterpolatedSpawnMode))
 	{
-		bool bActualInterpolatedSpawning = EmitterData->SpawnScriptProps.Script->IsInterpolatedParticleSpawnScript();
-		if (EmitterData->bInterpolatedSpawning != bActualInterpolatedSpawning)
+		ENiagaraInterpolatedSpawnMode ScriptInterpolationMode = EmitterData->SpawnScriptProps.Script->GetComputedVMCompilationId().InterpolatedSpawnMode;
+		ENiagaraInterpolatedSpawnMode EmitterInterpolationMode = EmitterData->InterpolatedSpawnMode;
+		if (EmitterInterpolationMode != ScriptInterpolationMode)
 		{
 			//Recompile spawn script if we've altered the interpolated spawn property.
-			EmitterData->SpawnScriptProps.Script->SetUsage(EmitterData->bInterpolatedSpawning ? ENiagaraScriptUsage::ParticleSpawnScriptInterpolated : ENiagaraScriptUsage::ParticleSpawnScript);
-			UE_LOG(LogNiagara, Log, TEXT("Updating script usage: Script->IsInterpolatdSpawn %d Emitter->bInterpolatedSpawning %d"), (int32)EmitterData->SpawnScriptProps.Script->IsInterpolatedParticleSpawnScript(), EmitterData->bInterpolatedSpawning);
+			EmitterData->SpawnScriptProps.Script->SetUsage(EmitterData->UsesInterpolatedSpawning() ? ENiagaraScriptUsage::ParticleSpawnScriptInterpolated : ENiagaraScriptUsage::ParticleSpawnScript);
+			UE_LOG(LogNiagara, Log, TEXT("Updating script usage: Script->IsInterpolatdSpawn %d Emitter->bInterpolatedSpawning %d"), (int32)EmitterData->SpawnScriptProps.Script->IsInterpolatedParticleSpawnScript(), EmitterData->UsesInterpolatedSpawning());
 			if (EmitterData->GraphSource != nullptr)
 			{
 				EmitterData->GraphSource->MarkNotSynchronized(TEXT("Emitter interpolated spawn changed"));
@@ -1996,6 +2002,11 @@ void FVersionedNiagaraEmitterData::CacheFromShaderCompiled()
 	}
 }
 
+bool FVersionedNiagaraEmitterData::UsesInterpolatedSpawning() const
+{
+	return InterpolatedSpawnMode != ENiagaraInterpolatedSpawnMode::NoInterpolation;
+}
+
 void UNiagaraEmitter::UpdateEmitterAfterLoad()
 {
 	if (bFullyLoaded)
@@ -2867,13 +2878,15 @@ void UNiagaraEmitter::CheckVersionDataAvailable()
 	Data.RendererProperties = RendererProperties_DEPRECATED;
 	Data.ScalabilityOverrides = ScalabilityOverrides_DEPRECATED;
 	Data.AttributesToPreserve = AttributesToPreserve_DEPRECATED;
-	Data.bInterpolatedSpawning = bInterpolatedSpawning_DEPRECATED;
 	Data.bRequiresPersistentIDs = bRequiresPersistentIDs_DEPRECATED;
 	Data.EventHandlerScriptProps = EventHandlerScriptProps_DEPRECATED;
 	Data.SharedEventGeneratorIds = SharedEventGeneratorIds_DEPRECATED;
 	Data.EmitterSpawnScriptProps = EmitterSpawnScriptProps_DEPRECATED;
 	Data.EmitterUpdateScriptProps = EmitterUpdateScriptProps_DEPRECATED;
 	Data.MaxGPUParticlesSpawnPerFrame = MaxGPUParticlesSpawnPerFrame_DEPRECATED;
+	
+	Data.bInterpolatedSpawning_DEPRECATED = bInterpolatedSpawning_DEPRECATED;
+	Data.InterpolatedSpawnMode = bInterpolatedSpawning_DEPRECATED ? ENiagaraInterpolatedSpawnMode::Interpolation : ENiagaraInterpolatedSpawnMode::NoInterpolation;
 
 	if (bFixedBounds_DEPRECATED)
 	{
@@ -2915,12 +2928,14 @@ void UNiagaraEmitter::CheckVersionDataAvailable()
 #endif
 
 FVersionedNiagaraEmitterData::FVersionedNiagaraEmitterData()
-	: bInterpolatedSpawning(false)
+	: InterpolatedSpawnMode(ENiagaraInterpolatedSpawnMode::NoInterpolation)
 #if WITH_EDITORONLY_DATA
+	, bInterpolatedSpawning_DEPRECATED(false)
 	, bGpuAlwaysRunParticleUpdateScript(false)
 #endif
 	, bRequiresPersistentIDs(false)
 	, MaxGPUParticlesSpawnPerFrame(0)
+	, bIsAllowedToExecute(true)
 	, bRequiresViewUniformBuffer(false)
 	, bNeedsPartialDepthTexture(false)
 #if WITH_EDITORONLY_DATA
diff --git a/Engine/Plugins/FX/Niagara/Source/Niagara/Private/NiagaraResolveDIHelpers.cpp b/Engine/Plugins/FX/Niagara/Source/Niagara/Private/NiagaraResolveDIHelpers.cpp
--- a/Engine/Plugins/FX/Niagara/Source/Niagara/Private/NiagaraResolveDIHelpers.cpp
+++ b/Engine/Plugins/FX/Niagara/Source/Niagara/Private/NiagaraResolveDIHelpers.cpp
@@ -406,7 +406,7 @@ namespace FNiagaraResolveDIHelpers
 			ResolveDIsForScript(System, VersionedNiagaraEmitterData->UpdateScriptProps.Script, EmitterHandle.GetUniqueInstanceName(), EmitterVariableReadMap, EmitterVariableAssignmentMap, EmitterCompactedVariableBindingMap, OutErrorMessages);
 			ResolveDIsForScript(System, VersionedNiagaraEmitterData->GetGPUComputeScript(), EmitterHandle.GetUniqueInstanceName(), EmitterVariableReadMap, EmitterVariableAssignmentMap, EmitterCompactedVariableBindingMap, OutErrorMessages);
 
-			if (VersionedNiagaraEmitterData->bInterpolatedSpawning)
+			if (VersionedNiagaraEmitterData->UsesInterpolatedSpawning())
 			{
 				TArray<UNiagaraScript*> SourceScripts = { VersionedNiagaraEmitterData->UpdateScriptProps.Script };
 				ResolveInternalDataInterfaces(SourceScripts, VersionedNiagaraEmitterData->SpawnScriptProps.Script);
diff --git a/Engine/Plugins/FX/Niagara/Source/Niagara/Private/NiagaraScript.cpp b/Engine/Plugins/FX/Niagara/Source/Niagara/Private/NiagaraScript.cpp
--- a/Engine/Plugins/FX/Niagara/Source/Niagara/Private/NiagaraScript.cpp
+++ b/Engine/Plugins/FX/Niagara/Source/Niagara/Private/NiagaraScript.cpp
@@ -852,12 +852,7 @@ TArray<FString> FNiagaraVMExecutableDataId::GetAdditionalVariableStrings() const
 
 bool FNiagaraVMExecutableDataId::HasInterpolatedParameters() const
 {
-	return bInterpolatedSpawn;
-}
-
-bool FNiagaraVMExecutableDataId::RequiresPersistentIDs() const
-{
-	return bRequiresPersistentIDs;
+	return InterpolatedSpawnMode != ENiagaraInterpolatedSpawnMode::NoInterpolation;
 }
 
 /**
@@ -869,16 +864,16 @@ bool FNiagaraVMExecutableDataId::RequiresPersistentIDs() const
 bool FNiagaraVMExecutableDataId::operator==(const FNiagaraVMExecutableDataId& ReferenceSet) const
 {
 	if (CompilerVersionID != ReferenceSet.CompilerVersionID ||
-		ScriptUsageType != ReferenceSet.ScriptUsageType ||
-		ScriptUsageTypeID != ReferenceSet.ScriptUsageTypeID ||
 #if WITH_EDITORONLY_DATA
+		ScriptUsageTypeID != ReferenceSet.ScriptUsageTypeID ||
+		ScriptUsageType != ReferenceSet.ScriptUsageType ||
 		BaseScriptCompileHash != ReferenceSet.BaseScriptCompileHash ||
-#endif
-		bUsesRapidIterationParams != ReferenceSet.bUsesRapidIterationParams ||
 		bDisableDebugSwitches != ReferenceSet.bDisableDebugSwitches ||
-		bInterpolatedSpawn != ReferenceSet.bInterpolatedSpawn ||
 		bRequiresPersistentIDs != ReferenceSet.bRequiresPersistentIDs ||
-		ScriptVersionID != ReferenceSet.ScriptVersionID)
+		bUsesRapidIterationParams != ReferenceSet.bUsesRapidIterationParams ||
+		ScriptVersionID != ReferenceSet.ScriptVersionID ||
+#endif
+		InterpolatedSpawnMode != ReferenceSet.InterpolatedSpawnMode)
 	{
 		return false;
 	}
@@ -1114,7 +1109,6 @@ void UNiagaraScript::ComputeVMCompilationId(FNiagaraVMExecutableDataId& Id, cons
 
 	Id.bUsesRapidIterationParams = true;
 	Id.bDisableDebugSwitches = false;
-	Id.bInterpolatedSpawn = false;
 	Id.bRequiresPersistentIDs = false;
 	Id.ScriptVersionID = IsVersioningEnabled() ? (VersionGuid.IsValid() ? VersionGuid : ExposedVersion) : FGuid();
 	
@@ -1274,19 +1268,24 @@ void UNiagaraScript::ComputeVMCompilationId(FNiagaraVMExecutableDataId& Id, cons
 
 		// Gather up the scripts we will use to iterate for static rapid iteration vars later...
 		EmitterData->GetScripts(Scripts, false, true);
-
-		if ((EmitterData->bInterpolatedSpawning && Usage == ENiagaraScriptUsage::ParticleGPUComputeScript) ||
-			(EmitterData->bInterpolatedSpawning && Usage == ENiagaraScriptUsage::ParticleSpawnScript) ||
+		bool bUsesInterpolatedSpawning = EmitterData->UsesInterpolatedSpawning();
+		
+		if ((bUsesInterpolatedSpawning && Usage == ENiagaraScriptUsage::ParticleGPUComputeScript) ||
+			(bUsesInterpolatedSpawning && Usage == ENiagaraScriptUsage::ParticleSpawnScript) ||
 			Usage == ENiagaraScriptUsage::ParticleSpawnScriptInterpolated)
 		{
-			Id.bInterpolatedSpawn = true;
+			Id.InterpolatedSpawnMode = EmitterData->InterpolatedSpawnMode;
 			Id.AdditionalDefines.Add(TEXT("InterpolatedSpawn"));
 		}
-		if (EmitterData->bInterpolatedSpawning)
+		if (bUsesInterpolatedSpawning)
 		{
 			Id.AdditionalDefines.Add(TEXT("Emitter.InterpolatedSpawn"));
+			if (EmitterData->InterpolatedSpawnMode == ENiagaraInterpolatedSpawnMode::RunUpdateScript)
+			{
+				Id.AdditionalDefines.Add(TEXT("RunUpdateWithoutInterpolation"));
+			}
 		}
-		if (EmitterData->bInterpolatedSpawning == false && EmitterData->bGpuAlwaysRunParticleUpdateScript == true)
+		if (bUsesInterpolatedSpawning == false && EmitterData->bGpuAlwaysRunParticleUpdateScript == true)
 		{
 			Id.AdditionalDefines.Add(TEXT("GpuAlwaysRunParticleUpdateScript"));
 		}
@@ -1385,7 +1384,7 @@ void UNiagaraScript::ComputeVMCompilationId(FNiagaraVMExecutableDataId& Id, cons
 				{
 					Id.AdditionalDefines.Add(EmitterHandle.GetUniqueInstanceName() + TEXT(".Determinism"));
 				}
-				if (HandleData->bInterpolatedSpawning)
+				if (HandleData->UsesInterpolatedSpawning())
 				{
 					Id.AdditionalDefines.Add(EmitterHandle.GetUniqueInstanceName() + TEXT(".InterpolatedSpawn"));
 				}
diff --git a/Engine/Plugins/FX/Niagara/Source/Niagara/Private/NiagaraSystem.cpp b/Engine/Plugins/FX/Niagara/Source/Niagara/Private/NiagaraSystem.cpp
--- a/Engine/Plugins/FX/Niagara/Source/Niagara/Private/NiagaraSystem.cpp
+++ b/Engine/Plugins/FX/Niagara/Source/Niagara/Private/NiagaraSystem.cpp
@@ -3192,7 +3192,7 @@ void UNiagaraSystem::EvaluateCompileResultDependencies() const
 			ParticleUpdateVisitOrder.Add(ParticleUpdateScriptIdx);
 
 			// If the spawn script is interpolated, allow the particle update script to resolve dependencies for spawn events and spawn sim stages.
-			if (EmitterData->bInterpolatedSpawning)
+			if (EmitterData->UsesInterpolatedSpawning())
 			{
 				ParticleSpawnVisitOrder.Add(ParticleUpdateScriptIdx);
 			}
@@ -3532,7 +3532,7 @@ void UNiagaraSystem::PrepareRapidIterationParametersForCompilation()
 					ScriptsToIterate.Emplace(ComputeScript, VersionedEmitter);
 				}
 			}
-			else if (EmitterData && EmitterData->bInterpolatedSpawning)
+			else if (EmitterData && EmitterData->UsesInterpolatedSpawning())
 			{
 				Scripts.AddUnique(EmitterData->SpawnScriptProps.Script);
 				ScriptDependencyMap.Add(CompiledScript, EmitterData->SpawnScriptProps.Script);
diff --git a/Engine/Plugins/FX/Niagara/Source/Niagara/Public/NiagaraCommon.h b/Engine/Plugins/FX/Niagara/Source/Niagara/Public/NiagaraCommon.h
--- a/Engine/Plugins/FX/Niagara/Source/Niagara/Public/NiagaraCommon.h
+++ b/Engine/Plugins/FX/Niagara/Source/Niagara/Public/NiagaraCommon.h
@@ -227,6 +227,19 @@ enum class ENiagaraInputNodeUsage : uint8
 	RapidIterationParameter
 };
 
+UENUM()
+enum class ENiagaraInterpolatedSpawnMode : uint8
+{
+	/** This mode does no interpolation and skips the update script on the first frame for newly spawned particles. This means it has better performance, but can lead to visual pops when the update scripts needs to scale things like size or color. */
+	NoInterpolation = 0,
+	
+	/** This mode runs both the spawn and update script, but without any value interpolation for particles. */
+	RunUpdateScript,
+	
+	/** This mode runs both the spawn and update script and performs a partial update at spawn time. This adds some performance overhead, but will produce much smoother spawning for high spawn rates, erratic frame rates and fast moving emitters. */
+	Interpolation
+};
+
 /**
 * Enumerates states a Niagara script can be in.
 */
diff --git a/Engine/Plugins/FX/Niagara/Source/NiagaraCore/Public/NiagaraCustomVersion.h b/Engine/Plugins/FX/Niagara/Source/NiagaraCore/Public/NiagaraCustomVersion.h
--- a/Engine/Plugins/FX/Niagara/Source/NiagaraCore/Public/NiagaraCustomVersion.h
+++ b/Engine/Plugins/FX/Niagara/Source/NiagaraCore/Public/NiagaraCustomVersion.h
@@ -213,6 +213,8 @@ struct FNiagaraCustomVersion
 
 		SerializeUsageBitMaskToGPUFunctionInfo, // Make sure that we serialize the UsageBitMask with the GPUFunctionInfo
 		
+		AddedNewInterpolatedSpawnMode, // Changed the interpolated spawn property into an enum
+		
 		// DO NOT ADD A NEW VERSION UNLESS YOU HAVE TALKED TO THE NIAGARA LEAD. Mismanagement of these versions can lead to data loss if it is adjusted in multiple streams simultaneously.
 		// -----<new versions can be added above this line>  -------------------------------------------------
 		VersionPlusOne,
diff --git a/Engine/Plugins/FX/Niagara/Source/NiagaraEditor/Private/Commandlets/NiagaraSystemAuditCommandlet.cpp b/Engine/Plugins/FX/Niagara/Source/NiagaraEditor/Private/Commandlets/NiagaraSystemAuditCommandlet.cpp
--- a/Engine/Plugins/FX/Niagara/Source/NiagaraEditor/Private/Commandlets/NiagaraSystemAuditCommandlet.cpp
+++ b/Engine/Plugins/FX/Niagara/Source/NiagaraEditor/Private/Commandlets/NiagaraSystemAuditCommandlet.cpp
@@ -371,7 +371,7 @@ bool UNiagaraSystemAuditCommandlet::ProcessNiagaraSystems()
 				TStringBuilder<512> GpuEmitterBuilder;
 				GpuEmitterBuilder.Append(EmitterData->GetDebugSimName());
 
-				GpuEmitterBuilder.Append(EmitterData->bInterpolatedSpawning ? TEXT(",true") : TEXT(",false"));
+				GpuEmitterBuilder.Append(EmitterData->UsesInterpolatedSpawning() ? TEXT(",true") : TEXT(",false"));
 
 				for (int32 iQualityLevel=0; iQualityLevel < NiagaraSettings->QualityLevels.Num(); ++iQualityLevel)
 				{
diff --git a/Engine/Plugins/FX/Niagara/Source/NiagaraEditor/Private/NiagaraCompilationTasks.cpp b/Engine/Plugins/FX/Niagara/Source/NiagaraEditor/Private/NiagaraCompilationTasks.cpp
--- a/Engine/Plugins/FX/Niagara/Source/NiagaraEditor/Private/NiagaraCompilationTasks.cpp
+++ b/Engine/Plugins/FX/Niagara/Source/NiagaraEditor/Private/NiagaraCompilationTasks.cpp
@@ -137,7 +137,7 @@ namespace NiagaraCompilationTasksImpl
 				{
 					DependentScripts.Add(EmitterData->GetGPUComputeScript());
 				}
-				else if (EmitterData->bInterpolatedSpawning)
+				else if (EmitterData->UsesInterpolatedSpawning())
 				{
 					DependentScripts.Add(EmitterData->SpawnScriptProps.Script);
 				}
diff --git a/Engine/Plugins/FX/Niagara/Source/NiagaraEditor/Private/NiagaraCompiler.cpp b/Engine/Plugins/FX/Niagara/Source/NiagaraEditor/Private/NiagaraCompiler.cpp
--- a/Engine/Plugins/FX/Niagara/Source/NiagaraEditor/Private/NiagaraCompiler.cpp
+++ b/Engine/Plugins/FX/Niagara/Source/NiagaraEditor/Private/NiagaraCompiler.cpp
@@ -1174,7 +1174,7 @@ TSharedPtr<FNiagaraCompileRequestDataBase, ESPMode::ThreadSafe> FNiagaraEditorMo
 							ScriptsToIterate.Emplace(ComputeScript, ScriptEmitterPair.Value);
 						}
 					}
-					else if (EmitterData && EmitterData->bInterpolatedSpawning)
+					else if (EmitterData && EmitterData->UsesInterpolatedSpawning())
 					{
 						Scripts.AddUnique(EmitterData->SpawnScriptProps.Script);
 						ScriptDependencyMap.Add(CompiledScript, EmitterData->SpawnScriptProps.Script);
diff --git a/Engine/Plugins/FX/Niagara/Source/NiagaraEditor/Private/NiagaraEmitterFactoryNew.cpp b/Engine/Plugins/FX/Niagara/Source/NiagaraEditor/Private/NiagaraEmitterFactoryNew.cpp
--- a/Engine/Plugins/FX/Niagara/Source/NiagaraEditor/Private/NiagaraEmitterFactoryNew.cpp
+++ b/Engine/Plugins/FX/Niagara/Source/NiagaraEditor/Private/NiagaraEmitterFactoryNew.cpp
@@ -238,7 +238,7 @@ void UNiagaraEmitterFactoryNew::InitializeEmitter(UNiagaraEmitter* NewEmitter, b
 		}
 
 		FNiagaraStackGraphUtilities::RelayoutGraph(*Source->NodeGraph);
-		EmitterData->bInterpolatedSpawning = true;
+		EmitterData->InterpolatedSpawnMode = ENiagaraInterpolatedSpawnMode::Interpolation;
 		EmitterData->bDeterminism = false; // NOTE: Default to non-determinism
 		EmitterData->SpawnScriptProps.Script->SetUsage(ENiagaraScriptUsage::ParticleSpawnScriptInterpolated);
 	}
diff --git a/Engine/Plugins/FX/Niagara/Source/NiagaraEditor/Private/NiagaraHlslTranslator.cpp b/Engine/Plugins/FX/Niagara/Source/NiagaraEditor/Private/NiagaraHlslTranslator.cpp
--- a/Engine/Plugins/FX/Niagara/Source/NiagaraEditor/Private/NiagaraHlslTranslator.cpp
+++ b/Engine/Plugins/FX/Niagara/Source/NiagaraEditor/Private/NiagaraHlslTranslator.cpp
@@ -1624,7 +1624,7 @@ FNiagaraTranslateResults TNiagaraHlslTranslator<GraphBridge>::Translate(const FN
 		FNiagaraHlslTranslationStage& UpdateStage = TranslationStages.Emplace_GetRef(ENiagaraScriptUsage::ParticleUpdateScript, FGuid());
 		UpdateStage.PassNamespace = TEXT("MapUpdate");
 		UpdateStage.ChunkModeIndex = ENiagaraCodeChunkMode::UpdateBody;
-		UpdateStage.bInterpolatePreviousParams = true;
+		UpdateStage.InterpolatedSpawnMode = CompileOptions.AdditionalDefines.Contains(TEXT("RunUpdateWithoutInterpolation")) ? ENiagaraInterpolatedSpawnMode::RunUpdateScript : ENiagaraInterpolatedSpawnMode::Interpolation;
 		UpdateStage.SimulationStageIndex = 0;
 		UpdateStage.NumIterations = 1;
 		UpdateStage.bWritesParticles = true;
@@ -1648,7 +1648,11 @@ FNiagaraTranslateResults TNiagaraHlslTranslator<GraphBridge>::Translate(const FN
 		FNiagaraHlslTranslationStage& UpdateStage = TranslationStages.Emplace_GetRef(ENiagaraScriptUsage::ParticleUpdateScript, FGuid());
 		UpdateStage.PassNamespace = TEXT("MapUpdate");
 		UpdateStage.ChunkModeIndex = ENiagaraCodeChunkMode::UpdateBody;
-		UpdateStage.bInterpolatePreviousParams = bHasInterpolatedSpawn;
+		UpdateStage.InterpolatedSpawnMode = ENiagaraInterpolatedSpawnMode::NoInterpolation;
+		if (CompileOptions.AdditionalDefines.Contains(TEXT("InterpolatedSpawn")))
+		{
+			UpdateStage.InterpolatedSpawnMode = CompileOptions.AdditionalDefines.Contains(TEXT("RunUpdateWithoutInterpolation")) ? ENiagaraInterpolatedSpawnMode::RunUpdateScript : ENiagaraInterpolatedSpawnMode::Interpolation;
+		}
 		UpdateStage.SimulationStageIndex = 0;
 		UpdateStage.NumIterations = 1;
 		UpdateStage.bWritesParticles = true;
@@ -1718,7 +1722,7 @@ FNiagaraTranslateResults TNiagaraHlslTranslator<GraphBridge>::Translate(const FN
 							FText::AsNumber((int32)ENiagaraCodeChunkMode::SimulationStageBodyMax - (int32)ENiagaraCodeChunkMode::SimulationStageBody),
 							FText::AsNumber((int32)SimStage.ChunkModeIndex)), nullptr, nullptr);
 					}
-					SimStage.bInterpolatePreviousParams = false;
+					SimStage.InterpolatedSpawnMode = ENiagaraInterpolatedSpawnMode::NoInterpolation;
 					SimStage.bCopyPreviousParams = false;
 					SimStage.SimulationStageIndex = DestSimStageIndex;					
 					SimStage.EnabledBinding = CompileSimStageData.EnabledBinding;
@@ -2097,7 +2101,7 @@ FNiagaraTranslateResults TNiagaraHlslTranslator<GraphBridge>::Translate(const FN
 				BuildMissingDefaults();
 			}
 
-			if (TranslationStages[i].bInterpolatePreviousParams)
+			if (TranslationStages[i].InterpolatedSpawnMode != ENiagaraInterpolatedSpawnMode::NoInterpolation)
 			{
 				bInterpolateParams = true;
 			}
@@ -3097,7 +3101,7 @@ void TNiagaraHlslTranslator<GraphBridge>::DefineInterpolatedParametersFunction(F
 {
 	for (int32 i = 0; i < TranslationStages.Num(); i++)
 	{
-		if (!TranslationStages[i].bInterpolatePreviousParams)
+		if (TranslationStages[i].InterpolatedSpawnMode == ENiagaraInterpolatedSpawnMode::NoInterpolation)
 		{
 			continue;
 		}
@@ -3114,8 +3118,20 @@ void TNiagaraHlslTranslator<GraphBridge>::DefineInterpolatedParametersFunction(F
 		{
 			FExpressionPermutationContext PermutationContext(*this, TranslationStages[i], HlslOutputString);
 
-			HlslOutputString += TEXT("\tint InterpSpawn_Index = ExecIndex();\n");
-			HlslOutputString += TEXT("\tfloat InterpSpawn_SpawnTime = ") + Emitter_InterpSpawnStartDt + TEXT(" + (") + Emitter_SpawnInterval + TEXT(" * InterpSpawn_Index);\n");
+			if (TranslationStages[i].InterpolatedSpawnMode == ENiagaraInterpolatedSpawnMode::RunUpdateScript)
+			{
+				HlslOutputString += TEXT("\tint InterpSpawn_Index = 0;\n");
+				HlslOutputString += TEXT("\tfloat InterpSpawn_SpawnTime = Engine_DeltaTime - ") + Emitter_InterpSpawnStartDt + TEXT(";\n");
+			}
+			else if (TranslationStages[i].InterpolatedSpawnMode == ENiagaraInterpolatedSpawnMode::Interpolation)
+			{
+				HlslOutputString += TEXT("\tint InterpSpawn_Index = ExecIndex();\n");
+				HlslOutputString += TEXT("\tfloat InterpSpawn_SpawnTime = ") + Emitter_InterpSpawnStartDt + TEXT(" + (") + Emitter_SpawnInterval + TEXT(" * InterpSpawn_Index);\n");
+			}
+			else
+			{
+				Error(FText::Format(LOCTEXT("UnsupportedInterpSpawn", "Unknown interpolated spawn value {0}"), FText::AsNumber(static_cast<int32>(TranslationStages[i].InterpolatedSpawnMode))));
+			}
 			HlslOutputString += TEXT("\tfloat InterpSpawn_UpdateTime = Engine_DeltaTime - InterpSpawn_SpawnTime;\n");
 			HlslOutputString += TEXT("\tfloat InterpSpawn_InvSpawnTime = 1.0 / InterpSpawn_SpawnTime;\n");
 			HlslOutputString += TEXT("\tfloat InterpSpawn_InvUpdateTime = 1.0 / InterpSpawn_UpdateTime;\n");
@@ -6270,7 +6286,7 @@ bool FNiagaraHlslTranslator::RequiresInterpolation() const
 {
 	for (int32 i = 0; i < TranslationStages.Num(); i++)
 	{
-		if (TranslationStages[i].bInterpolatePreviousParams)
+		if (TranslationStages[i].InterpolatedSpawnMode != ENiagaraInterpolatedSpawnMode::NoInterpolation)
 		{
 			return true;
 		}
diff --git a/Engine/Plugins/FX/Niagara/Source/NiagaraEditor/Private/NiagaraHlslTranslator.h b/Engine/Plugins/FX/Niagara/Source/NiagaraEditor/Private/NiagaraHlslTranslator.h
--- a/Engine/Plugins/FX/Niagara/Source/NiagaraEditor/Private/NiagaraHlslTranslator.h
+++ b/Engine/Plugins/FX/Niagara/Source/NiagaraEditor/Private/NiagaraHlslTranslator.h
@@ -171,7 +171,7 @@ public:
 	ENiagaraScriptUsage ScriptUsage;
 	FGuid UsageId;
 	FString PassNamespace;
-	bool bInterpolatePreviousParams = false;
+	ENiagaraInterpolatedSpawnMode InterpolatedSpawnMode = ENiagaraInterpolatedSpawnMode::NoInterpolation;
 	bool bCopyPreviousParams = true;
 	ENiagaraCodeChunkMode ChunkModeIndex = (ENiagaraCodeChunkMode)-1;
 	ENiagaraIterationSource IterationSourceType = ENiagaraIterationSource::Particles;
diff --git a/Engine/Plugins/FX/Niagara/Source/NiagaraEditor/Private/NiagaraParameterMapHistory.cpp b/Engine/Plugins/FX/Niagara/Source/NiagaraEditor/Private/NiagaraParameterMapHistory.cpp
--- a/Engine/Plugins/FX/Niagara/Source/NiagaraEditor/Private/NiagaraParameterMapHistory.cpp
+++ b/Engine/Plugins/FX/Niagara/Source/NiagaraEditor/Private/NiagaraParameterMapHistory.cpp
@@ -2422,7 +2422,7 @@ bool FCompileConstantResolver::ResolveConstant(FNiagaraVariable& OutConstant) co
 		}
 		if (OutConstant == FNiagaraVariable(FNiagaraTypeDefinition::GetBoolDef(), TEXT("Emitter.InterpolatedSpawn")))
 		{
-			OutConstant.SetValue(EmitterData->bInterpolatedSpawning ? FNiagaraBool(true) : FNiagaraBool(false));
+			OutConstant.SetValue(EmitterData->UsesInterpolatedSpawning() ? FNiagaraBool(true) : FNiagaraBool(false));
 			return true;
 		}
 		if (OutConstant == FNiagaraVariable(FNiagaraTypeDefinition::GetSimulationTargetEnum(), TEXT("Emitter.SimulationTarget")))
@@ -2559,10 +2559,10 @@ uint32 FCompileConstantResolver::BuildTypeHash() const
 		HashUnion.bHasEmitter = true;
 		if (FVersionedNiagaraEmitterData* EmitterData = Emitter.GetEmitterData())
 		{
-			HashUnion.EmitterSimTarget = (uint8)EmitterData->SimTarget;
+			HashUnion.EmitterSimTarget = static_cast<uint8>(EmitterData->SimTarget);
 			HashUnion.bEmitterLocalSpace = EmitterData->bLocalSpace;
 			HashUnion.bEmitterDeterminism = EmitterData->bDeterminism;
-			HashUnion.bEmitterInterpolatedSpawn = EmitterData->bInterpolatedSpawning;
+			HashUnion.bEmitterInterpolatedSpawn = EmitterData->UsesInterpolatedSpawning();
 		}
 	}
 
@@ -2571,8 +2571,8 @@ uint32 FCompileConstantResolver::BuildTypeHash() const
 		HashUnion.bHasSystem = true;
 	}
 
-	HashUnion.Usage = (uint8)Usage;
-	HashUnion.DebugState = (uint8)DebugState;
+	HashUnion.Usage = static_cast<uint8>(Usage);
+	HashUnion.DebugState = static_cast<uint8>(DebugState);
 
 	return HashUnion.Hash;
 }
diff --git a/Engine/Plugins/FX/Niagara/Source/NiagaraEditor/Private/Toolkits/NiagaraSystemToolkit.cpp b/Engine/Plugins/FX/Niagara/Source/NiagaraEditor/Private/Toolkits/NiagaraSystemToolkit.cpp
--- a/Engine/Plugins/FX/Niagara/Source/NiagaraEditor/Private/Toolkits/NiagaraSystemToolkit.cpp
+++ b/Engine/Plugins/FX/Niagara/Source/NiagaraEditor/Private/Toolkits/NiagaraSystemToolkit.cpp
@@ -206,7 +206,7 @@ void FNiagaraSystemToolkit::InitializeRapidIterationParameters(const FVersionedN
 		ScriptDependencyMap.Add(EmitterData->SpawnScriptProps.Script, EmitterData->GetGPUComputeScript());
 		ScriptDependencyMap.Add(EmitterData->UpdateScriptProps.Script, EmitterData->GetGPUComputeScript());
 	} 
-	else if (EmitterData->bInterpolatedSpawning)
+	else if (EmitterData->UsesInterpolatedSpawning())
 	{
 		ScriptDependencyMap.Add(EmitterData->UpdateScriptProps.Script, EmitterData->SpawnScriptProps.Script);
 	}
diff --git a/Engine/Plugins/FX/Niagara/Source/NiagaraEditor/Private/ViewModels/Stack/NiagaraStackEmitterSettingsGroup.cpp b/Engine/Plugins/FX/Niagara/Source/NiagaraEditor/Private/ViewModels/Stack/NiagaraStackEmitterSettingsGroup.cpp
--- a/Engine/Plugins/FX/Niagara/Source/NiagaraEditor/Private/ViewModels/Stack/NiagaraStackEmitterSettingsGroup.cpp
+++ b/Engine/Plugins/FX/Niagara/Source/NiagaraEditor/Private/ViewModels/Stack/NiagaraStackEmitterSettingsGroup.cpp
@@ -207,9 +207,9 @@ void UNiagaraStackEmitterPropertiesItem::RefreshIssues(TArray<FStackIssue>& NewI
 			);
 		}
 
-		if (ActualEmitterData->bInterpolatedSpawning == false)
+		if (ActualEmitterData->UsesInterpolatedSpawning() == false)
 		{
-			const FText InterpolatedSpawningDisabledInfoFix = LOCTEXT("InterpolatedSpawningDisabledInfoFix", "Enable interpolated spawning");
+			const FText InterpolatedSpawningDisabledInfoFix = LOCTEXT("InterpolatedSpawningDisabledInfoFix", "Run update script on spawn");
 			FVersionedNiagaraEmitterWeakPtr WeakEmitter = GetEmitterViewModel()->GetEmitter().ToWeakPtr();
 
 			NewIssues.Emplace(
@@ -229,7 +229,7 @@ void UNiagaraStackEmitterPropertiesItem::RefreshIssues(TArray<FStackIssue>& NewI
 								FScopedTransaction ScopedTransaction(InterpolatedSpawningDisabledInfoFix);
 
 								PinnedEmitter.Emitter->Modify();
-								PinnedEmitter.GetEmitterData()->bInterpolatedSpawning = true;
+								PinnedEmitter.GetEmitterData()->InterpolatedSpawnMode = ENiagaraInterpolatedSpawnMode::RunUpdateScript;
 								UNiagaraSystem::RequestCompileForEmitter(PinnedEmitter);
 							}
 						}
diff --git a/Engine/Plugins/FX/Niagara/Source/NiagaraEditor/Private/ViewModels/Stack/NiagaraStackFunctionInput.cpp b/Engine/Plugins/FX/Niagara/Source/NiagaraEditor/Private/ViewModels/Stack/NiagaraStackFunctionInput.cpp
--- a/Engine/Plugins/FX/Niagara/Source/NiagaraEditor/Private/ViewModels/Stack/NiagaraStackFunctionInput.cpp
+++ b/Engine/Plugins/FX/Niagara/Source/NiagaraEditor/Private/ViewModels/Stack/NiagaraStackFunctionInput.cpp
@@ -1822,7 +1822,7 @@ void UNiagaraStackFunctionInput::GetAvailableParameterHandles(TArray<FNiagaraPar
 				}
 			}
 		}
-		bool bInterpolatedSpawn = GetEmitterViewModel().IsValid() && GetEmitterViewModel()->GetEmitter().GetEmitterData()->bInterpolatedSpawning;
+		bool bInterpolatedSpawn = GetEmitterViewModel().IsValid() && GetEmitterViewModel()->GetEmitter().GetEmitterData()->UsesInterpolatedSpawning();
 		bool bCheckInterpSpawn = bInterpolatedSpawn || !bSpawnScript;
 		if (OutputNode == CurrentOutputNode || (CurrentOutputNode != nullptr && UsageRunsBefore(OutputNode->GetUsage(), CurrentOutputNode->GetUsage(), bCheckInterpSpawn, Emitter, OutputNode->GetUsageId(), CurrentOutputNode->GetUsageId())) || (CurrentOutputNode != nullptr && IsSpawnUsage(CurrentOutputNode->GetUsage())))
 		{
diff --git a/Engine/Plugins/FX/Niagara/Source/NiagaraEditorWidgets/Private/Stack/SNiagaraStackRowPerfWidget.cpp b/Engine/Plugins/FX/Niagara/Source/NiagaraEditorWidgets/Private/Stack/SNiagaraStackRowPerfWidget.cpp
--- a/Engine/Plugins/FX/Niagara/Source/NiagaraEditorWidgets/Private/Stack/SNiagaraStackRowPerfWidget.cpp
+++ b/Engine/Plugins/FX/Niagara/Source/NiagaraEditorWidgets/Private/Stack/SNiagaraStackRowPerfWidget.cpp
@@ -482,7 +482,7 @@ bool SNiagaraStackRowPerfWidget::IsInterpolatedSpawnEnabled() const
 	{
 		return false;
 	}
-	return EmitterData->bInterpolatedSpawning;
+	return EmitterData->UsesInterpolatedSpawning();
 }
 
 bool SNiagaraStackRowPerfWidget::IsGpuEmitter() const
diff --git a/Engine/Plugins/FX/Niagara/Source/NiagaraShader/Public/NiagaraShared.h b/Engine/Plugins/FX/Niagara/Source/NiagaraShader/Public/NiagaraShared.h
--- a/Engine/Plugins/FX/Niagara/Source/NiagaraShader/Public/NiagaraShared.h
+++ b/Engine/Plugins/FX/Niagara/Source/NiagaraShader/Public/NiagaraShared.h
@@ -321,11 +321,6 @@ public:
 	/** Feature level that the shader map is going to be compiled for.  */
 	LAYOUT_FIELD(ERHIFeatureLevel::Type, FeatureLevel);
 
-	/**
-	* The base id of the subgraph this shader primarily represents.
-	*/
-	LAYOUT_FIELD(FGuid, BaseScriptID_DEPRECATED);
-
 	/** Configuration options */
 	LAYOUT_FIELD(TMemoryImageArray<FMemoryImageString>, AdditionalDefines);
 
@@ -947,9 +942,6 @@ private:
 	TSharedPtr<FNiagaraShaderScriptParametersMetadata> ScriptParametersMetadata;
 	TSharedPtr<FNiagaraShaderScriptParametersMetadata> ScriptParametersMetadata_RT;
 
-	/** Guid id for base script*/
-	FGuid BaseScriptId_DEPRECATED;
-
 	/** Configuration options */
 	TArray<FString> AdditionalDefines;
 
diff --git a/Engine/Source/Runtime/Core/Private/UObject/FortniteMainBranchObjectVersion.cpp b/Engine/Source/Runtime/Core/Private/UObject/FortniteMainBranchObjectVersion.cpp
--- a/Engine/Source/Runtime/Core/Private/UObject/FortniteMainBranchObjectVersion.cpp
+++ b/Engine/Source/Runtime/Core/Private/UObject/FortniteMainBranchObjectVersion.cpp
@@ -10,8 +10,8 @@ TMap<FGuid, FGuid> FFortniteMainBranchObjectVersion::GetSystemGuids()
 	SystemGuids.Add(DevGuids.LANDSCAPE_MOBILE_COOK_VERSION, FGuid("32D02EF867C74B71A0D4E0FA41392732"));
 	SystemGuids.Add(DevGuids.MATERIALSHADERMAP_DERIVEDDATA_VER, FGuid("04B8C817071A453C9EBEFBFB51D5F785"));
 	SystemGuids.Add(DevGuids.NANITE_DERIVEDDATA_VER, FGuid("42C46B24E1586986880AFE62B0F575AB"));
-	SystemGuids.Add(DevGuids.NIAGARASHADERMAP_DERIVEDDATA_VER, FGuid("45D5978F33D940EEA6602B78559900DC"));
-	SystemGuids.Add(DevGuids.Niagara_LatestScriptCompileVersion, FGuid("670A2EBDEE1A4794BB77BBFF331EA7EF"));
+	SystemGuids.Add(DevGuids.NIAGARASHADERMAP_DERIVEDDATA_VER, FGuid("6360A977062842A29ED30E3A7ACB0E64"));
+	SystemGuids.Add(DevGuids.Niagara_LatestScriptCompileVersion, FGuid("041F479DAAB746DAA556A4C1ABB07E04"));
 	SystemGuids.Add(DevGuids.SkeletalMeshDerivedDataVersion, FGuid("716424BCEC314E969ACFE36D2C2ECE49"));
 	SystemGuids.Add(DevGuids.STATICMESH_DERIVEDDATA_VER, FGuid("3ABF32149F9D4D838750368AB95573BE"));
 	SystemGuids.Add(DevGuids.MaterialTranslationDDCVersion, FGuid("84CF0B79D69D4A4BBE8CAFD4B3D48AB7"));
