Fixed renames of staticswitch parameters getting lost in the input hierarchy
Also renamed the NiagaraHierarchyScriptParametersViewModel -> NiagaraScriptParametersHierarchyViewModel for more consistency

diff --git a/Engine/Plugins/FX/Niagara/Source/NiagaraEditor/Private/NiagaraEditorUtilities.cpp b/Engine/Plugins/FX/Niagara/Source/NiagaraEditor/Private/NiagaraEditorUtilities.cpp
--- a/Engine/Plugins/FX/Niagara/Source/NiagaraEditor/Private/NiagaraEditorUtilities.cpp
+++ b/Engine/Plugins/FX/Niagara/Source/NiagaraEditor/Private/NiagaraEditorUtilities.cpp
@@ -81,7 +81,7 @@
 #include "ViewModels/NiagaraEmitterViewModel.h"
 #include "ViewModels/NiagaraParameterPanelViewModel.h"
 #include "Widgets/SDataHierarchyEditor.h"
-#include "ViewModels/HierarchyEditor/NiagaraHierarchyScriptParametersViewModel.h"
+#include "ViewModels/HierarchyEditor/NiagaraScriptParametersHierarchyViewModel.h"
 #include "Widgets/SToolTip.h"
 
 #define LOCTEXT_NAMESPACE "FNiagaraEditorUtilities"
diff --git a/Engine/Plugins/FX/Niagara/Source/NiagaraEditor/Private/NiagaraGraph.cpp b/Engine/Plugins/FX/Niagara/Source/NiagaraEditor/Private/NiagaraGraph.cpp
--- a/Engine/Plugins/FX/Niagara/Source/NiagaraEditor/Private/NiagaraGraph.cpp
+++ b/Engine/Plugins/FX/Niagara/Source/NiagaraEditor/Private/NiagaraGraph.cpp
@@ -38,7 +38,7 @@
 #include "String/ParseTokens.h"
 #include "ViewModels/NiagaraScriptViewModel.h"
 #include "DataHierarchyViewModelBase.h"
-#include "ViewModels/HierarchyEditor/NiagaraHierarchyScriptParametersViewModel.h"
+#include "ViewModels/HierarchyEditor/NiagaraScriptParametersHierarchyViewModel.h"
 
 #include UE_INLINE_GENERATED_CPP_BY_NAME(NiagaraGraph)
 
@@ -2759,6 +2759,7 @@ bool UNiagaraGraph::RenameParameter(const FNiagaraVariable& Parameter, FName New
 
 	bool bParameterMerged = false;
 
+	Modify();
 	// Create the new parameter
 	FNiagaraVariable NewParameter = Parameter;
 	NewParameter.SetName(NewName);
@@ -2861,6 +2862,7 @@ bool UNiagaraGraph::RenameStaticSwitch(UNiagaraNodeStaticSwitch* SwitchNode, FNa
 		return false;
 	}
 
+	Modify();
 	bIsRenamingParameter = true;
 	ON_SCOPE_EXIT{ bIsRenamingParameter = false; };
 
@@ -2902,6 +2904,12 @@ bool UNiagaraGraph::RenameStaticSwitch(UNiagaraNodeStaticSwitch* SwitchNode, FNa
 		CreatedParameter->Variable = DestVariable;
 	}
 
+	TInstancedStruct<FNiagaraParametersChangedData> RenameData = TInstancedStruct<FNiagaraParameterRenamedData>::Make();
+	FNiagaraParameterRenamedData& Data = RenameData.GetMutable<FNiagaraParameterRenamedData>();
+	Data.OldScriptVariable = OrigScriptVariable;
+	Data.NewScriptVariable = VariableToScriptVariable[DestVariable];
+	OnParametersChangedDelegate.Broadcast(RenameData);
+	
 	// now see if we need to remove any script variables that are now obsolete because they have been renamed
 	if (OrigScriptVariable)
 	{
@@ -2915,7 +2923,7 @@ bool UNiagaraGraph::RenameStaticSwitch(UNiagaraNodeStaticSwitch* SwitchNode, FNa
 
 	// Fixup reference collection and pin names
 	FixupReferenceCollectionsPostRename(OrigVariable, DestVariable, DestScriptVariable != nullptr /* bIsMerging */, false /* bSuppressEvents */);
-
+	
 	NotifyGraphChanged();
 
 	// force the graph to refresh the metadata
diff --git a/Engine/Plugins/FX/Niagara/Source/NiagaraEditor/Private/NiagaraScriptSource.cpp b/Engine/Plugins/FX/Niagara/Source/NiagaraEditor/Private/NiagaraScriptSource.cpp
--- a/Engine/Plugins/FX/Niagara/Source/NiagaraEditor/Private/NiagaraScriptSource.cpp
+++ b/Engine/Plugins/FX/Niagara/Source/NiagaraEditor/Private/NiagaraScriptSource.cpp
@@ -21,7 +21,7 @@
 #include "NiagaraNodeParameterMapSet.h"
 #include "NiagaraScript.h"
 #include "Logging/LogMacros.h"
-#include "ViewModels/HierarchyEditor/NiagaraHierarchyScriptParametersViewModel.h"
+#include "ViewModels/HierarchyEditor/NiagaraScriptParametersHierarchyViewModel.h"
 #include "ViewModels/Stack/NiagaraStackGraphUtilities.h"
 
 #include UE_INLINE_GENERATED_CPP_BY_NAME(NiagaraScriptSource)
diff --git a/Engine/Plugins/FX/Niagara/Source/NiagaraEditor/Private/Toolkits/NiagaraScriptToolkit.cpp b/Engine/Plugins/FX/Niagara/Source/NiagaraEditor/Private/Toolkits/NiagaraScriptToolkit.cpp
--- a/Engine/Plugins/FX/Niagara/Source/NiagaraEditor/Private/Toolkits/NiagaraScriptToolkit.cpp
+++ b/Engine/Plugins/FX/Niagara/Source/NiagaraEditor/Private/Toolkits/NiagaraScriptToolkit.cpp
@@ -40,7 +40,7 @@
 #include "ViewModels/NiagaraParameterDefinitionsPanelViewModel.h"
 #include "ViewModels/NiagaraParameterPanelViewModel.h"
 #include "ViewModels/NiagaraScriptViewModel.h"
-#include "ViewModels/HierarchyEditor/NiagaraHierarchyScriptParametersViewModel.h"
+#include "ViewModels/HierarchyEditor/NiagaraScriptParametersHierarchyViewModel.h"
 #include "Widgets/SDataHierarchyEditor.h"
 #include "Widgets/Docking/SDockTab.h"
 #include "Widgets/SNiagaraParameterDefinitionsPanel.h"
@@ -210,7 +210,7 @@ void FNiagaraScriptToolkit::Initialize( const EToolkitMode::Type Mode, const TSh
 	ScriptViewModel = MakeShareable(new FNiagaraStandaloneScriptViewModel(GetGraphEditorDisplayName(), ENiagaraParameterEditMode::EditAll, NiagaraMessageLogViewModel, MessageLogGuidKey, bIsForDataProcessingOnly));
 	ScriptViewModel->Initialize(EditedNiagaraScript, OriginalNiagaraScript);
 
-	ParametersHierarchyViewModel = NewObject<UNiagaraHierarchyScriptParametersViewModel>();
+	ParametersHierarchyViewModel = NewObject<UNiagaraScriptParametersHierarchyViewModel>();
 	ParametersHierarchyViewModel->Initialize(ScriptViewModel.ToSharedRef());
 	ParametersHierarchyViewModel->OnHierarchyChanged().AddSP(this, &FNiagaraScriptToolkit::OnHierarchyChanged);
 	ParametersHierarchyViewModel->OnHierarchyPropertiesChanged().AddSP(this, &FNiagaraScriptToolkit::OnHierarchyPropertiesChanged);
@@ -982,7 +982,7 @@ void FNiagaraScriptToolkit::AddReferencedObjects(FReferenceCollector& Collector)
 	Collector.AddReferencedObject(ParametersHierarchyViewModel);
 }
 
-UNiagaraHierarchyScriptParametersViewModel* FNiagaraScriptToolkit::GetHierarchyViewModel() const
+UNiagaraScriptParametersHierarchyViewModel* FNiagaraScriptToolkit::GetHierarchyViewModel() const
 {
 	return ParametersHierarchyViewModel;
 }
diff --git a/Engine/Plugins/FX/Niagara/Source/NiagaraEditor/Private/Toolkits/NiagaraScriptToolkit.h b/Engine/Plugins/FX/Niagara/Source/NiagaraEditor/Private/Toolkits/NiagaraScriptToolkit.h
--- a/Engine/Plugins/FX/Niagara/Source/NiagaraEditor/Private/Toolkits/NiagaraScriptToolkit.h
+++ b/Engine/Plugins/FX/Niagara/Source/NiagaraEditor/Private/Toolkits/NiagaraScriptToolkit.h
@@ -11,9 +11,9 @@
 #include "EditorUndoClient.h"
 #include "NiagaraEditorCommon.h"
 #include "NiagaraScript.h"
-#include "ViewModels/HierarchyEditor/NiagaraHierarchyScriptParametersViewModel.h"
+#include "ViewModels/HierarchyEditor/NiagaraScriptParametersHierarchyViewModel.h"
 
-class UNiagaraHierarchyScriptParametersViewModel;
+class UNiagaraScriptParametersHierarchyViewModel;
 class UNiagaraVersionMetaData;
 struct FCustomExpanderData;
 class IDetailsView;
@@ -66,7 +66,7 @@ public:
 		return TEXT("FNiagaraScriptToolkit");
 	}
 	
-	UNiagaraHierarchyScriptParametersViewModel* GetHierarchyViewModel() const;
+	UNiagaraScriptParametersHierarchyViewModel* GetHierarchyViewModel() const;
 	
 	/**
 	* Updates list of module info used to show stats
@@ -180,7 +180,7 @@ private:
 	/** The Script being edited */
 	TSharedPtr<FNiagaraStandaloneScriptViewModel> ScriptViewModel;
 
-	TObjectPtr<UNiagaraHierarchyScriptParametersViewModel> ParametersHierarchyViewModel;
+	TObjectPtr<UNiagaraScriptParametersHierarchyViewModel> ParametersHierarchyViewModel;
 
 	/** The Parameter Panel displaying graph variables */
 	TSharedPtr<FNiagaraScriptToolkitParameterPanelViewModel> ParameterPanelViewModel;
diff --git a/Engine/Plugins/FX/Niagara/Source/NiagaraEditor/Private/Toolkits/SystemToolkitModes/NiagaraSystemToolkitModeBase.cpp b/Engine/Plugins/FX/Niagara/Source/NiagaraEditor/Private/Toolkits/SystemToolkitModes/NiagaraSystemToolkitModeBase.cpp
--- a/Engine/Plugins/FX/Niagara/Source/NiagaraEditor/Private/Toolkits/SystemToolkitModes/NiagaraSystemToolkitModeBase.cpp
+++ b/Engine/Plugins/FX/Niagara/Source/NiagaraEditor/Private/Toolkits/SystemToolkitModes/NiagaraSystemToolkitModeBase.cpp
@@ -1139,7 +1139,7 @@ TSharedRef<SWidget> FNiagaraSystemToolkitModeBase::CreateScratchPadHierarchyWidg
 	TSharedPtr<SWidget> ContentWidget = SNullWidget::NullWidget;
 	if(ActiveScratchPadViewModel.IsValid())
 	{
-		UNiagaraHierarchyScriptParametersViewModel* ScriptHierarchyViewModel = ActiveScratchPadViewModel->GetHierarchyViewModel();
+		UNiagaraScriptParametersHierarchyViewModel* ScriptHierarchyViewModel = ActiveScratchPadViewModel->GetHierarchyViewModel();
 		
 		ContentWidget = SNew(SDataHierarchyEditor, ScriptHierarchyViewModel)
 		.OnGenerateRowContentWidget_Static(&FNiagaraEditorUtilities::HierarchyEditor::Scripts::GenerateRowContentForScriptParameterHierarchyEditor);
diff --git a/Engine/Plugins/FX/Niagara/Source/NiagaraEditor/Private/ViewModels/HierarchyEditor/NiagaraHierarchyScriptParametersViewModel.cpp b/Engine/Plugins/FX/Niagara/Source/NiagaraEditor/Private/ViewModels/HierarchyEditor/NiagaraHierarchyScriptParametersViewModel.cpp
deleted file mode 100644
diff --git a/Engine/Plugins/FX/Niagara/Source/NiagaraEditor/Private/ViewModels/HierarchyEditor/NiagaraScriptParametersHierarchyViewModel.cpp b/Engine/Plugins/FX/Niagara/Source/NiagaraEditor/Private/ViewModels/HierarchyEditor/NiagaraScriptParametersHierarchyViewModel.cpp
new file mode 100644
--- /dev/null
+++ b/Engine/Plugins/FX/Niagara/Source/NiagaraEditor/Private/ViewModels/HierarchyEditor/NiagaraScriptParametersHierarchyViewModel.cpp
@@ -0,0 +1,357 @@
+﻿// Copyright Epic Games, Inc. All Rights Reserved.
+
+#include "ViewModels/HierarchyEditor/NiagaraScriptParametersHierarchyViewModel.h"
+
+#include "NiagaraConstants.h"
+#include "NiagaraEditorUtilities.h"
+#include "NiagaraGraph.h"
+#include "NiagaraScriptSource.h"
+#include "ScopedTransaction.h"
+#include "Logging/StructuredLog.h"
+#include "ViewModels/NiagaraScriptViewModel.h"
+
+#define LOCTEXT_NAMESPACE "NiagaraScriptParameterHierarchyEditor"
+
+void UNiagaraHierarchyScriptParameter::Initialize(const UNiagaraScriptVariable& InParameterScriptVariable)
+{
+	if(InParameterScriptVariable.Metadata.GetVariableGuid().IsValid() == false)
+	{
+		UE_LOGFMT(LogNiagaraEditor, Warning, "Invalid hierarchy script parameter initialization. No valid guid found. This hierarchy element will be deleted on next refresh.");
+	}
+	
+	SetIdentity(FHierarchyElementIdentity({InParameterScriptVariable.Metadata.GetVariableGuid()}, {}));
+}
+
+bool UNiagaraHierarchyScriptParameter::IsValid() const
+{
+	return GetScriptVariable() != nullptr;
+}
+
+FString UNiagaraHierarchyScriptParameter::ToString() const
+{
+	TOptional<FNiagaraVariable> Variable = GetVariable();
+
+	if(Variable.IsSet())
+	{
+		return Variable.GetValue().GetName().ToString();
+	}
+	
+	return TEXT("Invalid"); 
+}
+
+FText UNiagaraHierarchyScriptParameter::GetTooltip() const
+{
+	if(UNiagaraScriptVariable* ScriptVariable = GetScriptVariable())
+	{
+		return ScriptVariable->Metadata.Description; 
+	}
+
+	return FText::GetEmpty();
+}
+
+UNiagaraScriptVariable* UNiagaraHierarchyScriptParameter::GetScriptVariable() const
+{
+	UNiagaraGraph* OwningGraph = GetTypedOuter<UNiagaraGraph>();
+	if(ensureMsgf(OwningGraph != nullptr, TEXT("When retrieving the script variable matching this hierarchy parameter, there should always be an owning graph")))
+	{
+		// If the script variable does not exist, this is an indicator it has been deleted and this hierarchy parameter should be deleted
+		return OwningGraph->GetScriptVariable(GetPersistentIdentity().Guids[0]);
+	}
+
+	return nullptr;
+}
+
+TOptional<FNiagaraVariable> UNiagaraHierarchyScriptParameter::GetVariable() const
+{
+	if(UNiagaraScriptVariable* ScriptVariable = GetScriptVariable())
+	{
+		return ScriptVariable->Variable;
+	}
+	
+	return TOptional<FNiagaraVariable>(); 
+}
+
+void UNiagaraScriptParametersHierarchyViewModel::Initialize(TSharedRef<FNiagaraScriptViewModel> InScriptViewModel)
+{
+	ScriptViewModelWeak = InScriptViewModel;
+
+	Cast<UNiagaraScriptSource>(ScriptViewModelWeak.Pin()->GetStandaloneScript().GetScriptData()->GetSource())->NodeGraph->OnParametersChanged().AddUObject(this, &UNiagaraScriptParametersHierarchyViewModel::OnParametersChanged);
+	UDataHierarchyViewModelBase::Initialize();
+
+	UNiagaraHierarchyScriptParameterRefreshContext* ScriptParameterRefreshContext = NewObject<UNiagaraHierarchyScriptParameterRefreshContext>(this);
+	ScriptParameterRefreshContext->SetNiagaraGraph(Cast<UNiagaraScriptSource>(InScriptViewModel->GetStandaloneScript().GetScriptData()->GetSource())->NodeGraph);
+	SetRefreshContext(ScriptParameterRefreshContext);
+}
+
+TSharedRef<FNiagaraScriptViewModel> UNiagaraScriptParametersHierarchyViewModel::GetScriptViewModel() const
+{
+	TSharedPtr<FNiagaraScriptViewModel> ScriptViewModelPinned = ScriptViewModelWeak.Pin();
+	checkf(ScriptViewModelPinned.IsValid(), TEXT("Script view model destroyed before parameters hierarchy view model."));
+	return ScriptViewModelPinned.ToSharedRef();
+}
+
+
+UHierarchyRoot* UNiagaraScriptParametersHierarchyViewModel::GetHierarchyRoot() const
+{
+	const TArray<FVersionedNiagaraScriptWeakPtr>& Scripts = GetScriptViewModel()->GetScripts();
+	if(ensure(Scripts.Num() >= 1 && Scripts[0].Pin().Script != nullptr) == false)
+	{
+		return nullptr;
+	}
+
+	FVersionedNiagaraScriptData* ScriptData = Scripts[0].Pin().GetScriptData();
+	if(ensure(ScriptData != nullptr) == false)
+	{
+		return nullptr;
+	}
+
+	return Cast<UNiagaraScriptSource>(ScriptData->GetSource())->NodeGraph->GetScriptParameterHierarchyRoot();
+}
+
+UObject* UNiagaraScriptParametersHierarchyViewModel::GetOuterForSourceRoot() const
+{
+	// The graph always needs to be valid as the outer
+	if(UNiagaraGraph* Graph = Cast<UNiagaraScriptSource>(ScriptViewModelWeak.Pin()->GetStandaloneScript().GetScriptData()->GetSource())->NodeGraph)
+	{
+		return Graph;
+	}
+
+	checkNoEntry();
+	return nullptr;
+}
+
+TSubclassOf<UHierarchyCategory> UNiagaraScriptParametersHierarchyViewModel::GetCategoryDataClass() const
+{
+	return UNiagaraHierarchyScriptCategory::StaticClass();
+}
+
+TSharedPtr<FHierarchyElementViewModel> UNiagaraScriptParametersHierarchyViewModel::CreateCustomViewModelForElement(UHierarchyElement* Element, TSharedPtr<FHierarchyElementViewModel> Parent)
+{
+	if(UNiagaraHierarchyScriptParameter* Item = Cast<UNiagaraHierarchyScriptParameter>(Element))
+	{
+		return MakeShared<FNiagaraHierarchyScriptParameterViewModel>(Item, Parent.ToSharedRef(), this);
+	}
+	else if(UHierarchyRoot* Root = Cast<UHierarchyRoot>(Element))
+	{
+		// If the root is the hierarchy root, we know it's for the hierarchy. If not, it's the transient source root
+		bool bIsForHierarchy = GetHierarchyRoot() == Element;
+		return MakeShared<FNiagaraHierarchyScriptRootViewModel>(Root, this, bIsForHierarchy);
+	}
+
+	return nullptr;
+}
+
+void UNiagaraScriptParametersHierarchyViewModel::PrepareSourceItems(UHierarchyRoot* SourceRoot, TSharedPtr<FHierarchyRootViewModel> SourceRootViewModel)
+{
+	UNiagaraScriptSourceBase* SourceBase = GetScriptViewModel()->GetStandaloneScript().GetScriptData()->GetSource();
+	UNiagaraScriptSource* ScriptSource = Cast<UNiagaraScriptSource>(SourceBase);
+	const TMap<FNiagaraVariable, TObjectPtr<UNiagaraScriptVariable>>& ScriptVariableMap = ScriptSource->NodeGraph->GetAllMetaData();
+
+	for(const auto& ScriptVariablePair : ScriptVariableMap)
+	{
+		// We only want to be able to organize module inputs & static switches 
+		if(ScriptVariablePair.Key.IsInNameSpace(FNiagaraConstants::ModuleNamespace) == false && ScriptVariablePair.Value->GetIsStaticSwitch() == false)
+		{
+			continue;
+		}
+		
+		UNiagaraScriptVariable* ScriptVariable = ScriptVariablePair.Value;
+		
+		UNiagaraHierarchyScriptParameter* ScriptParameterHierarchyObject = NewObject<UNiagaraHierarchyScriptParameter>(SourceRoot, NAME_None, RF_Transient);
+		ScriptParameterHierarchyObject->Initialize(*ScriptVariable);
+		SourceRoot->GetChildrenMutable().Add(ScriptParameterHierarchyObject);
+	}
+}
+
+void UNiagaraScriptParametersHierarchyViewModel::SetupCommands()
+{
+	Super::SetupCommands();
+}
+
+TSharedRef<FHierarchyDragDropOp> UNiagaraScriptParametersHierarchyViewModel::CreateDragDropOp(TSharedRef<FHierarchyElementViewModel> Item)
+{
+	if(UHierarchyCategory* HierarchyCategory = Cast<UHierarchyCategory>(Item->GetDataMutable()))
+	{
+		TSharedRef<FHierarchyDragDropOp> CategoryDragDropOp = MakeShared<FHierarchyDragDropOp>(Item);
+		CategoryDragDropOp->Construct();
+		return CategoryDragDropOp;
+	}
+	else if(UNiagaraHierarchyScriptParameter* ScriptParameter = Cast<UNiagaraHierarchyScriptParameter>(Item->GetDataMutable()))
+	{
+		TSharedPtr<FHierarchyItemViewModel> ScriptParameterViewModel = StaticCastSharedRef<FHierarchyItemViewModel>(Item);
+		TSharedRef<FHierarchyDragDropOp> ScriptParameterDragDropOp = MakeShared<FNiagaraHierarchyScriptParameterDragDropOp>(ScriptParameterViewModel);
+		ScriptParameterDragDropOp->Construct();
+		return ScriptParameterDragDropOp;
+	}
+	
+	check(false);
+	return MakeShared<FHierarchyDragDropOp>(nullptr);
+}
+
+void UNiagaraScriptParametersHierarchyViewModel::FinalizeInternal()
+{
+	if(ScriptViewModelWeak.IsValid())
+	{
+		// If this is called during Undo, it's possible the Graph does no longer exist
+		if(UNiagaraGraph* Graph = Cast<UNiagaraScriptSource>(ScriptViewModelWeak.Pin()->GetStandaloneScript().GetScriptData()->GetSource())->NodeGraph)
+		{
+			Graph->OnParametersChanged().RemoveAll(this);
+		}
+	}
+}
+
+void UNiagaraScriptParametersHierarchyViewModel::OnParametersChanged(TOptional<TInstancedStruct<FNiagaraParametersChangedData>> ParametersChangedData)
+{
+	if(ParametersChangedData.IsSet())
+	{
+		TInstancedStruct<FNiagaraParametersChangedData> ChangedData = ParametersChangedData.GetValue();
+		if(ChangedData.GetScriptStruct() == FNiagaraParameterRenamedData::StaticStruct())
+		{
+			const FNiagaraParameterRenamedData& RenamedData = ChangedData.Get<FNiagaraParameterRenamedData>();
+			if(RenamedData.OldScriptVariable->GetIsStaticSwitch() && RenamedData.NewScriptVariable->GetIsStaticSwitch())
+			{
+				TArray<UNiagaraHierarchyScriptParameter*> AllHierarchyScriptParameters;
+				HierarchyRoot->GetChildrenOfType<UNiagaraHierarchyScriptParameter>(AllHierarchyScriptParameters, true);
+
+				UNiagaraHierarchyScriptParameter** FoundHierarchyScriptParameter = AllHierarchyScriptParameters.FindByPredicate([Guid = RenamedData.OldScriptVariable->Metadata.GetVariableGuid()](UNiagaraHierarchyScriptParameter* Candidate)
+				{
+					return Candidate->GetPersistentIdentity().Guids[0] == Guid;
+				});
+
+				if(FoundHierarchyScriptParameter)
+				{
+					HierarchyRoot->Modify();
+					(*FoundHierarchyScriptParameter)->Initialize(*RenamedData.NewScriptVariable);
+				}
+			}
+		}
+	}
+	ForceFullRefresh();
+}
+
+bool UNiagaraScriptParametersHierarchyViewModel::MatchesContext(const FTransactionContext& InContext, const TArray<TPair<UObject*, FTransactionObjectEvent>>& TransactionObjectContexts) const
+{
+	return Super::MatchesContext(InContext, TransactionObjectContexts) || TransactionObjectContexts.ContainsByPredicate([](TPair<UObject*, FTransactionObjectEvent> ContextPair)
+	{
+		return ContextPair.Key->IsA<UNiagaraGraph>() || ContextPair.Key->IsA<UNiagaraScriptVariable>();
+	});
+}
+
+TSharedRef<SWidget> FNiagaraHierarchyScriptParameterDragDropOp::CreateCustomDecorator() const
+{
+	if(const UNiagaraHierarchyScriptParameter* ScriptParameter = Cast<UNiagaraHierarchyScriptParameter>(DraggedElement.Pin()->GetData()))
+	{
+		TOptional<FNiagaraVariable> Variable = ScriptParameter->GetVariable();
+		if(Variable.IsSet())
+		{
+			return FNiagaraParameterUtilities::GetParameterWidget(Variable.GetValue(), false, false);
+		}
+	}
+
+	return SNullWidget::NullWidget;
+}
+
+bool FNiagaraHierarchyScriptParameterViewModel::DoesExternalDataStillExist(const UHierarchyDataRefreshContext* Context) const
+{
+	// GetScriptVariable checks the owning graph if it still exists
+	const UNiagaraScriptVariable* ScriptVariable = Cast<UNiagaraHierarchyScriptParameter>(GetDataMutable())->GetScriptVariable();
+	
+	if(ScriptVariable == nullptr)
+	{
+		return false;
+	}	
+
+	// We make sure that the variable not only still exists but also qualifies for the hierarchy (namespace can change for example)
+	if(ScriptVariable->GetIsStaticSwitch() == false && ScriptVariable->Variable.IsInNameSpace(FNiagaraConstants::ModuleNamespace) == false)
+	{
+		return false;
+	}
+	
+	return true;
+}
+
+FHierarchyElementViewModel::FCanPerformActionResults FNiagaraHierarchyScriptParameterViewModel::CanDropOnInternal(TSharedPtr<FHierarchyElementViewModel> DraggedItem, EItemDropZone ItemDropZone)
+{
+	// if the input isn't editable, we don't allow any drops on/above/below the item.
+	// Even though it technically works, the merge process will only re-add the item at the end and not preserve order so there is no point in allowing dropping above/below
+	if(IsEditableByUser().bCanPerform == false)
+	{
+		return false;
+	}
+
+	FCanPerformActionResults AllowDrop(false);
+	
+	TSharedPtr<FHierarchyElementViewModel> TargetDropItem = AsShared();
+
+	// we only allow drops if some general conditions are fulfilled
+	if(DraggedItem->GetData() != TargetDropItem->GetData() &&
+		(!DraggedItem->HasParent(TargetDropItem, false) || ItemDropZone != EItemDropZone::OntoItem)  &&
+		!TargetDropItem->HasParent(DraggedItem, true))
+	{
+		if(ItemDropZone == EItemDropZone::OntoItem)
+		{
+			// We support nested inputs
+			if(DraggedItem->GetData()->IsA<UNiagaraHierarchyScriptParameter>() && TargetDropItem->GetData()->IsA<UNiagaraHierarchyScriptParameter>())
+			{
+				FText BaseMessage = LOCTEXT("DroppingInputOnInputNestedChild", "This will nest input {0} under input {1}");
+				AllowDrop.CanPerformMessage = FText::FormatOrdered(BaseMessage, DraggedItem->ToStringAsText(), TargetDropItem->ToStringAsText());
+				AllowDrop.bCanPerform = true;
+			}
+		}
+		else
+		{
+			// if the dragged item is an input, we generally allow above/below, even for nested child inputs
+			if(DraggedItem->GetData()->IsA<UNiagaraHierarchyScriptParameter>())
+			{
+				AllowDrop.bCanPerform = true;
+			}
+			else
+			{
+				// we use default logic only if there is no parent input. Nested children are not allowed to contain anything but other inputs.
+				if(TargetDropItem->GetParent().Pin()->GetData<UNiagaraHierarchyScriptParameter>() == nullptr)
+				{
+					AllowDrop = FHierarchyItemViewModel::CanDropOnInternal(DraggedItem, ItemDropZone);
+				}
+			}
+		}
+	}
+
+	return AllowDrop;
+}
+
+void FNiagaraHierarchyScriptParameterViewModel::OnDroppedOnInternal(TSharedPtr<FHierarchyElementViewModel> DroppedItem, EItemDropZone ItemDropZone)
+{
+	if(ItemDropZone != EItemDropZone::OntoItem)
+	{
+		return FHierarchyItemViewModel::OnDroppedOnInternal(DroppedItem, ItemDropZone);
+	}
+	else
+	{
+		FScopedTransaction Transaction(LOCTEXT("Transaction_AddedChildInput", "Added child input"));
+		HierarchyViewModel->GetHierarchyRoot()->Modify();
+		
+		if(DroppedItem->IsForHierarchy() == false)
+		{
+			TSharedRef<FHierarchyElementViewModel> AddedItemViewModel = DuplicateToThis(DroppedItem);
+			AddedItemViewModel->SyncViewModelsToData();
+		}
+		else
+		{
+			TSharedRef<FHierarchyElementViewModel> ReparentedViewModel = ReparentToThis(DroppedItem);
+			ReparentedViewModel->SyncViewModelsToData();
+		}
+
+		HierarchyViewModel->RefreshHierarchyView();
+		HierarchyViewModel->RefreshSourceView();
+	}
+}
+
+void FNiagaraHierarchyScriptRootViewModel::SortChildrenData() const
+{
+	GetDataMutable()->GetChildrenMutable().StableSort([](const UHierarchyElement& ItemA, const UHierarchyElement& ItemB)
+		{
+			return ItemA.IsA<UHierarchyItem>() && ItemB.IsA<UHierarchyCategory>();
+		});
+}
+
+#undef LOCTEXT_NAMESPACE
diff --git a/Engine/Plugins/FX/Niagara/Source/NiagaraEditor/Private/ViewModels/NiagaraScratchPadScriptViewModel.cpp b/Engine/Plugins/FX/Niagara/Source/NiagaraEditor/Private/ViewModels/NiagaraScratchPadScriptViewModel.cpp
--- a/Engine/Plugins/FX/Niagara/Source/NiagaraEditor/Private/ViewModels/NiagaraScratchPadScriptViewModel.cpp
+++ b/Engine/Plugins/FX/Niagara/Source/NiagaraEditor/Private/ViewModels/NiagaraScratchPadScriptViewModel.cpp
@@ -79,7 +79,7 @@ void FNiagaraScratchPadScriptViewModel::Initialize(UNiagaraScript* Script, UNiag
 	);
 	ParameterPaneViewModel->Init(UIContext);
 
-	ScriptHierarchyViewModel = NewObject<UNiagaraHierarchyScriptParametersViewModel>();
+	ScriptHierarchyViewModel = NewObject<UNiagaraScriptParametersHierarchyViewModel>();
 	ScriptHierarchyViewModel->Initialize(AsShared());
 
 	ScriptHierarchyViewModel->OnHierarchyChanged().AddSP(this, &FNiagaraScratchPadScriptViewModel::OnHierarchyChanged);
@@ -138,7 +138,7 @@ TArray<UNiagaraGraph*> FNiagaraScratchPadScriptViewModel::GetEditableGraphs()
 	return EditableGraphs;
 }
 
-UNiagaraHierarchyScriptParametersViewModel* FNiagaraScratchPadScriptViewModel::GetHierarchyViewModel()
+UNiagaraScriptParametersHierarchyViewModel* FNiagaraScratchPadScriptViewModel::GetHierarchyViewModel()
 {
 	return ScriptHierarchyViewModel;
 }
diff --git a/Engine/Plugins/FX/Niagara/Source/NiagaraEditor/Private/ViewModels/Stack/NiagaraStackInputCategory.cpp b/Engine/Plugins/FX/Niagara/Source/NiagaraEditor/Private/ViewModels/Stack/NiagaraStackInputCategory.cpp
--- a/Engine/Plugins/FX/Niagara/Source/NiagaraEditor/Private/ViewModels/Stack/NiagaraStackInputCategory.cpp
+++ b/Engine/Plugins/FX/Niagara/Source/NiagaraEditor/Private/ViewModels/Stack/NiagaraStackInputCategory.cpp
@@ -9,7 +9,7 @@
 #include "NiagaraConstants.h"
 #include "NiagaraSimulationStageBase.h"
 #include "ViewModels/NiagaraEmitterViewModel.h"
-#include "ViewModels/HierarchyEditor/NiagaraHierarchyScriptParametersViewModel.h"
+#include "ViewModels/HierarchyEditor/NiagaraScriptParametersHierarchyViewModel.h"
 #include "ViewModels/HierarchyEditor/NiagaraSummaryViewViewModel.h"
 #include "ViewModels/Stack/NiagaraStackEmitterSettingsGroup.h"
 #include "ViewModels/Stack/NiagaraStackEventScriptItemGroup.h"
diff --git a/Engine/Plugins/FX/Niagara/Source/NiagaraEditor/Private/Widgets/SNiagaraScriptInputPreviewPanel.cpp b/Engine/Plugins/FX/Niagara/Source/NiagaraEditor/Private/Widgets/SNiagaraScriptInputPreviewPanel.cpp
--- a/Engine/Plugins/FX/Niagara/Source/NiagaraEditor/Private/Widgets/SNiagaraScriptInputPreviewPanel.cpp
+++ b/Engine/Plugins/FX/Niagara/Source/NiagaraEditor/Private/Widgets/SNiagaraScriptInputPreviewPanel.cpp
@@ -12,7 +12,7 @@
 #include "SlateOptMacros.h"
 #include "Framework/Application/SlateApplication.h"
 #include "Toolkits/NiagaraScriptToolkit.h"
-#include "ViewModels/HierarchyEditor/NiagaraHierarchyScriptParametersViewModel.h"
+#include "ViewModels/HierarchyEditor/NiagaraScriptParametersHierarchyViewModel.h"
 #include "Widgets/Input/SButton.h"
 #include "Widgets/Input/SComboButton.h"
 #include "Widgets/Text/STextBlock.h"
@@ -350,7 +350,7 @@ void SNiagaraScriptInputPreviewPanel::RemoveDelegates()
 	{
 		Cast<UNiagaraScriptSource>(ScriptToolkit.Pin()->EditedNiagaraScript.GetScriptData()->GetSource())->NodeGraph->OnParametersChanged().RemoveAll(this);
 
-		if(UNiagaraHierarchyScriptParametersViewModel* HierarchyViewModel = ScriptToolkit.Pin()->GetHierarchyViewModel())
+		if(UNiagaraScriptParametersHierarchyViewModel* HierarchyViewModel = ScriptToolkit.Pin()->GetHierarchyViewModel())
 		{
 			HierarchyViewModel->OnHierarchyChanged().RemoveAll(this);
 			HierarchyViewModel->OnHierarchyPropertiesChanged().RemoveAll(this);
diff --git a/Engine/Plugins/FX/Niagara/Source/NiagaraEditor/Public/ViewModels/HierarchyEditor/NiagaraHierarchyScriptParametersViewModel.h b/Engine/Plugins/FX/Niagara/Source/NiagaraEditor/Public/ViewModels/HierarchyEditor/NiagaraHierarchyScriptParametersViewModel.h
deleted file mode 100644
diff --git a/Engine/Plugins/FX/Niagara/Source/NiagaraEditor/Public/ViewModels/HierarchyEditor/NiagaraScriptParametersHierarchyViewModel.h b/Engine/Plugins/FX/Niagara/Source/NiagaraEditor/Public/ViewModels/HierarchyEditor/NiagaraScriptParametersHierarchyViewModel.h
new file mode 100644
--- /dev/null
+++ b/Engine/Plugins/FX/Niagara/Source/NiagaraEditor/Public/ViewModels/HierarchyEditor/NiagaraScriptParametersHierarchyViewModel.h
@@ -0,0 +1,120 @@
+﻿// Copyright Epic Games, Inc. All Rights Reserved.
+
+#pragma once
+
+#include "NiagaraGraph.h"
+#include "DataHierarchyViewModelBase.h"
+#include "NiagaraScriptVariable.h"
+#include "ViewModels/NiagaraScriptViewModel.h"
+#include "NiagaraScriptParametersHierarchyViewModel.generated.h"
+
+class UNiagaraGraph;
+
+/** The refresh context is used to determine if hierarchy script variables should be removed. */
+UCLASS()
+class UNiagaraHierarchyScriptParameterRefreshContext : public UHierarchyDataRefreshContext
+{
+	GENERATED_BODY()
+
+public:
+	void SetNiagaraGraph(UNiagaraGraph* InGraph) { NiagaraGraph = InGraph; }
+	const UNiagaraGraph* GetNiagaraGraph() const { return NiagaraGraph; }
+private:
+	UPROPERTY(Transient)
+	TObjectPtr<UNiagaraGraph> NiagaraGraph;
+};
+
+/** A hierarchy script parameter is an optional object embedded in the hierarchy. */
+UCLASS()
+class UNiagaraHierarchyScriptParameter : public UHierarchyItem
+{
+	GENERATED_BODY()
+
+public:
+	void Initialize(const UNiagaraScriptVariable& InParameterScriptVariable);
+
+	bool IsValid() const;
+	
+	virtual FString ToString() const override;
+	FText GetTooltip() const;
+	
+	UNiagaraScriptVariable* GetScriptVariable() const;
+	TOptional<FNiagaraVariable> GetVariable() const;
+};
+
+/** The category class used for the script hierarchy editor. It lets us add additional data later on. */
+UCLASS()
+class UNiagaraHierarchyScriptCategory : public UHierarchyCategory
+{
+	GENERATED_BODY()
+};
+
+/**
+ * The view model that defines the script editor's hierarchy editor for input parameters.
+ */
+UCLASS()
+class NIAGARAEDITOR_API UNiagaraScriptParametersHierarchyViewModel : public UDataHierarchyViewModelBase
+{
+	GENERATED_BODY()
+public:
+	void Initialize(TSharedRef<FNiagaraScriptViewModel> InScriptViewModel);
+	TSharedRef<FNiagaraScriptViewModel> GetScriptViewModel() const;
+
+	/** UDataHierarchyViewModelBase */
+	virtual UHierarchyRoot* GetHierarchyRoot() const override;
+	/** We set the outer for the source root to the graph, so that both source & hierarchy items have the same outer. */
+	virtual UObject* GetOuterForSourceRoot() const override;
+	virtual TSubclassOf<UHierarchyCategory> GetCategoryDataClass() const override;
+	virtual TSharedPtr<FHierarchyElementViewModel> CreateCustomViewModelForElement(UHierarchyElement* Element, TSharedPtr<FHierarchyElementViewModel> Parent) override;
+	virtual void PrepareSourceItems(UHierarchyRoot* SourceRoot, TSharedPtr<FHierarchyRootViewModel> SourceRootViewModel) override;
+	virtual void SetupCommands() override;
+	virtual TSharedRef<FHierarchyDragDropOp> CreateDragDropOp(TSharedRef<FHierarchyElementViewModel> Item) override;
+	virtual bool SupportsDetailsPanel() override { return true; }
+	virtual void FinalizeInternal() override;
+private:
+	void OnParametersChanged(TOptional<TInstancedStruct<FNiagaraParametersChangedData>> ParametersChangedData);
+
+protected:
+	virtual bool MatchesContext(const FTransactionContext& InContext, const TArray<TPair<UObject*, FTransactionObjectEvent>>& TransactionObjectContexts) const override;
+private:
+	TWeakPtr<FNiagaraScriptViewModel> ScriptViewModelWeak;
+};
+
+class FNiagaraHierarchyScriptParameterDragDropOp : public FHierarchyDragDropOp
+{
+public:
+	DRAG_DROP_OPERATOR_TYPE(FNiagaraHierarchyScriptParameterDragDropOp, FHierarchyDragDropOp)
+
+	FNiagaraHierarchyScriptParameterDragDropOp(TSharedPtr<FHierarchyItemViewModel> InputViewModel) : FHierarchyDragDropOp(InputViewModel) {}
+	
+	virtual TSharedRef<SWidget> CreateCustomDecorator() const override;
+};
+
+struct FNiagaraHierarchyScriptParameterViewModel : public FHierarchyItemViewModel
+{
+	FNiagaraHierarchyScriptParameterViewModel(UNiagaraHierarchyScriptParameter* ScriptParameter, TSharedRef<FHierarchyElementViewModel> InParent, TWeakObjectPtr<UNiagaraScriptParametersHierarchyViewModel> ViewModel)
+		: FHierarchyItemViewModel(ScriptParameter, InParent, ViewModel)	{}
+
+protected:
+	virtual bool RepresentsExternalData() const override { return true; }
+	virtual bool DoesExternalDataStillExist(const UHierarchyDataRefreshContext* Context) const override;
+	virtual UObject* GetDataForEditing() override { return Cast<UNiagaraHierarchyScriptParameter>(GetDataMutable())->GetScriptVariable(); }
+	/** We want to be able to edit in the details panel regardless of source or hierarchy item. */
+	virtual bool AllowEditingInDetailsPanel() const override { return true; }
+	virtual bool CanRenameInternal() override { return false; }
+	
+	virtual FCanPerformActionResults CanDropOnInternal(TSharedPtr<FHierarchyElementViewModel>, EItemDropZone ItemDropZone) override;
+	virtual void OnDroppedOnInternal(TSharedPtr<FHierarchyElementViewModel> DroppedItem, EItemDropZone ItemDropZone) override;
+	
+	virtual bool CanHaveChildren() const override { return bIsForHierarchy == true; }
+};
+
+struct FNiagaraHierarchyScriptRootViewModel : public FHierarchyRootViewModel
+{
+	FNiagaraHierarchyScriptRootViewModel(UHierarchyRoot* Root, TWeakObjectPtr<UNiagaraScriptParametersHierarchyViewModel> ViewModel, bool bInIsForHierarchy)
+		: FHierarchyRootViewModel(Root, ViewModel, bInIsForHierarchy)	{}
+
+protected:
+	// In the script, loose parameters are always added before categories (reverse of the default implementation)
+	virtual void SortChildrenData() const override;
+};
diff --git a/Engine/Plugins/FX/Niagara/Source/NiagaraEditor/Public/ViewModels/NiagaraScratchPadScriptViewModel.h b/Engine/Plugins/FX/Niagara/Source/NiagaraEditor/Public/ViewModels/NiagaraScratchPadScriptViewModel.h
--- a/Engine/Plugins/FX/Niagara/Source/NiagaraEditor/Public/ViewModels/NiagaraScratchPadScriptViewModel.h
+++ b/Engine/Plugins/FX/Niagara/Source/NiagaraEditor/Public/ViewModels/NiagaraScratchPadScriptViewModel.h
@@ -4,7 +4,7 @@
 
 #include "ViewModels/NiagaraScriptViewModel.h"
 #include "GraphEditAction.h"
-#include "HierarchyEditor/NiagaraHierarchyScriptParametersViewModel.h"
+#include "HierarchyEditor/NiagaraScriptParametersHierarchyViewModel.h"
 #include "HierarchyEditor/NiagaraSummaryViewViewModel.h"
 
 class INiagaraParameterPanelViewModel;
@@ -41,7 +41,7 @@ public:
 
 	NIAGARAEDITOR_API TArray<UNiagaraGraph*> GetEditableGraphs();
 
-	NIAGARAEDITOR_API UNiagaraHierarchyScriptParametersViewModel* GetHierarchyViewModel();
+	NIAGARAEDITOR_API UNiagaraScriptParametersHierarchyViewModel* GetHierarchyViewModel();
 
 	//~ Begin NiagaraParameterDefinitionsSubscriberViewModel Interface
 protected:
@@ -134,7 +134,7 @@ private:
 	TSharedPtr<FUICommandList> ParameterPanelCommands;
 	TSharedPtr<FNiagaraScriptToolkitParameterPanelViewModel> ParameterPaneViewModel;
 
-	TObjectPtr<UNiagaraHierarchyScriptParametersViewModel> ScriptHierarchyViewModel;
+	TObjectPtr<UNiagaraScriptParametersHierarchyViewModel> ScriptHierarchyViewModel;
 	
 	FDelegateHandle OnGraphNeedsRecompileHandle;
 
diff --git a/Engine/Plugins/FX/Niagara/Source/NiagaraEditor/Public/ViewModels/Stack/NiagaraStackScriptHierarchyRoot.h b/Engine/Plugins/FX/Niagara/Source/NiagaraEditor/Public/ViewModels/Stack/NiagaraStackScriptHierarchyRoot.h
--- a/Engine/Plugins/FX/Niagara/Source/NiagaraEditor/Public/ViewModels/Stack/NiagaraStackScriptHierarchyRoot.h
+++ b/Engine/Plugins/FX/Niagara/Source/NiagaraEditor/Public/ViewModels/Stack/NiagaraStackScriptHierarchyRoot.h
@@ -5,7 +5,7 @@
 #include "ViewModels/Stack/NiagaraStackItem.h"
 #include "NiagaraTypes.h"
 #include "NiagaraStackSection.h"
-#include "ViewModels/HierarchyEditor/NiagaraHierarchyScriptParametersViewModel.h"
+#include "ViewModels/HierarchyEditor/NiagaraScriptParametersHierarchyViewModel.h"
 #include "NiagaraStackScriptHierarchyRoot.generated.h"
 
 class UNiagaraNodeFunctionCall;
