Fix for CIS warning (Disabling interpolated spawn because emitter flag and script type don't match) in some niagara content with parent emitters

#rnx

diff --git a/Engine/Plugins/FX/Niagara/Source/Niagara/Private/NiagaraEmitter.cpp b/Engine/Plugins/FX/Niagara/Source/Niagara/Private/NiagaraEmitter.cpp
--- a/Engine/Plugins/FX/Niagara/Source/Niagara/Private/NiagaraEmitter.cpp
+++ b/Engine/Plugins/FX/Niagara/Source/Niagara/Private/NiagaraEmitter.cpp
@@ -718,9 +718,9 @@ void FVersionedNiagaraEmitterData::PostLoad(UNiagaraEmitter& Emitter, int32 Niag
 		}
 	}
 
-	if (NiagaraVer < FNiagaraCustomVersion::AddedNewInterpolatedSpawnMode)
+	if (NiagaraVer < FNiagaraCustomVersion::AddedNewInterpolatedSpawnMode && bInterpolatedSpawning_DEPRECATED)
 	{
-		InterpolatedSpawnMode = bInterpolatedSpawning_DEPRECATED ? ENiagaraInterpolatedSpawnMode::Interpolation : ENiagaraInterpolatedSpawnMode::NoInterpolation;
+		InterpolatedSpawnMode = ENiagaraInterpolatedSpawnMode::Interpolation;
 	}
 #endif
 
