- Remove the vector field noise option for stateless as it has issues when forcing the CPU only path

#rnx

diff --git a/Engine/Plugins/FX/Niagara/Shaders/Private/Stateless/Modules/NiagaraStatelessModule_SolveVelocitiesAndForces.ush b/Engine/Plugins/FX/Niagara/Shaders/Private/Stateless/Modules/NiagaraStatelessModule_SolveVelocitiesAndForces.ush
--- a/Engine/Plugins/FX/Niagara/Shaders/Private/Stateless/Modules/NiagaraStatelessModule_SolveVelocitiesAndForces.ush
+++ b/Engine/Plugins/FX/Niagara/Shaders/Private/Stateless/Modules/NiagaraStatelessModule_SolveVelocitiesAndForces.ush
@@ -29,8 +29,8 @@ float3	SolveVelocitiesAndForces_PointOrigin;
 uint	SolveVelocitiesAndForces_NoiseEnabled;
 float	SolveVelocitiesAndForces_NoiseAmplitude;
 float3	SolveVelocitiesAndForces_NoiseFrequency;
-Texture3D		SolveVelocitiesAndForces_NoiseTexture;
-SamplerState	SolveVelocitiesAndForces_NoiseSampler;
+//Texture3D		SolveVelocitiesAndForces_NoiseTexture;
+//SamplerState	SolveVelocitiesAndForces_NoiseSampler;
 
 uint	SolveVelocitiesAndForces_NoiseMode;
 uint	SolveVelocitiesAndForces_NoiseLUTOffset;
@@ -99,20 +99,20 @@ void SolveVelocitiesAndForces_Simulate(inout FStatelessParticle Particle)
 	{
 		const float3 NoiseOffset = RandomFloat3(10);
 		float3 NoiseValue = 0;
+		//if ( SolveVelocitiesAndForces_NoiseMode == 0 )
+		//{
+		//	const float3 NoiseTravel = float3(Particle.Age, Particle.NormalizedAge, Particle.UniqueIndex);
+		//	const float3 NoisePosition = NoiseOffset + (NoiseTravel * SolveVelocitiesAndForces_NoiseFrequency);
+		//	NoiseValue = Texture3DSample(SolveVelocitiesAndForces_NoiseTexture, SolveVelocitiesAndForces_NoiseSampler, NoisePosition).xyz;
+		//}
 		if ( SolveVelocitiesAndForces_NoiseMode == 0 )
-		{
-			const float3 NoiseTravel = float3(Particle.Age, Particle.NormalizedAge, Particle.UniqueIndex);
-			const float3 NoisePosition = NoiseOffset + (NoiseTravel * SolveVelocitiesAndForces_NoiseFrequency);
-			NoiseValue = Texture3DSample(SolveVelocitiesAndForces_NoiseTexture, SolveVelocitiesAndForces_NoiseSampler, NoisePosition).xyz;
-		}
-		else if ( SolveVelocitiesAndForces_NoiseMode == 1 )
 		{
 			const float3 NoiseTravel = float3(Particle.Age, Particle.NormalizedAge, Particle.UniqueIndex) * 5.0f;
 			const float3 NoisePosition = NoiseOffset + (NoiseTravel * SolveVelocitiesAndForces_NoiseFrequency);
 			const float3x4 J = JacobianSimplex_ALU(NoisePosition, false, 1.0f);
 			NoiseValue = float3(J[1][2] - J[2][1], J[2][0] - J[0][2], J[0][1] - J[1][0]); // See comments to JacobianSimplex_ALU in Random.ush
 		}
-		else if ( SolveVelocitiesAndForces_NoiseMode == 2 )
+		else if ( SolveVelocitiesAndForces_NoiseMode == 1 )
 		{
 			const uint3 ChannelIndex = RandomUInt3(11) % SolveVelocitiesAndForces_NoiseLUTNumChannel;
 			const float3 ChannelLerp = (Particle.Age * SolveVelocitiesAndForces_NoiseFrequency) + RandomFloat(12);
@@ -131,7 +131,7 @@ void SolveVelocitiesAndForces_Simulate(inout FStatelessParticle Particle)
 			);
 			NoiseValue = lerp(ValuesA, ValuesB, frac(ChannelLerp));
 		}
-		else if ( SolveVelocitiesAndForces_NoiseMode == 3 )
+		else if ( SolveVelocitiesAndForces_NoiseMode == 2 )
 		{
 			const uint3 ChannelIndex = RandomUInt3(13) % SolveVelocitiesAndForces_NoiseLUTNumChannel;
 			const float ChannelLerp = (Particle.Age * SolveVelocitiesAndForces_NoiseFrequency.x) + RandomFloat(14);
diff --git a/Engine/Plugins/FX/Niagara/Source/Niagara/Internal/Stateless/Modules/NiagaraStatelessModule_CurlNoiseForce.h b/Engine/Plugins/FX/Niagara/Source/Niagara/Internal/Stateless/Modules/NiagaraStatelessModule_CurlNoiseForce.h
--- a/Engine/Plugins/FX/Niagara/Source/Niagara/Internal/Stateless/Modules/NiagaraStatelessModule_CurlNoiseForce.h
+++ b/Engine/Plugins/FX/Niagara/Source/Niagara/Internal/Stateless/Modules/NiagaraStatelessModule_CurlNoiseForce.h
@@ -14,7 +14,7 @@ class UVectorField;
 UENUM()
 enum class ENSM_NoiseMode
 {
-	VectorField,
+	//VectorField,
 	JacobNoise,
 	LUTJacob,
 	LUTJacobBicubic,
@@ -32,10 +32,10 @@ public:
 	float NoiseFrequency = 1.0f;
 
 	UPROPERTY(EditAnywhere, Category = "Parameters")
-	ENSM_NoiseMode NoiseMode = ENSM_NoiseMode::VectorField;
+	ENSM_NoiseMode NoiseMode = ENSM_NoiseMode::JacobNoise;
 
-	UPROPERTY(EditAnywhere, Category = "Parameters", meta = (AllowedClasses = "/Script/Engine.VectorField,/Script/Engine.VolumeTexture,/Script/Engine.TextureRenderTargetVolume"))
-	TObjectPtr<UObject> NoiseTexture;
+	//UPROPERTY(EditAnywhere, Category = "Parameters", meta = (AllowedClasses = "/Script/Engine.VectorField,/Script/Engine.VolumeTexture,/Script/Engine.TextureRenderTargetVolume"))
+	//TObjectPtr<UObject> NoiseTexture;
 
 	//-TODO: Add support for GPU once we settle on a feature set for this module
 	virtual ENiagaraStatelessFeatureMask GetFeatureMask() const { return ENiagaraStatelessFeatureMask::ExecuteGPU; }
diff --git a/Engine/Plugins/FX/Niagara/Source/Niagara/Private/Stateless/Modules/NiagaraStatelessModule_CurlNoiseForce.cpp b/Engine/Plugins/FX/Niagara/Source/Niagara/Private/Stateless/Modules/NiagaraStatelessModule_CurlNoiseForce.cpp
--- a/Engine/Plugins/FX/Niagara/Source/Niagara/Private/Stateless/Modules/NiagaraStatelessModule_CurlNoiseForce.cpp
+++ b/Engine/Plugins/FX/Niagara/Source/Niagara/Private/Stateless/Modules/NiagaraStatelessModule_CurlNoiseForce.cpp
@@ -14,7 +14,7 @@ void UNiagaraStatelessModule_CurlNoiseForce::BuildEmitterData(const FNiagaraStat
 	PhysicsBuildData.bNoiseEnabled = true;
 	PhysicsBuildData.NoiseAmplitude = NoiseAmplitude;
 	PhysicsBuildData.NoiseFrequency = NoiseFrequency;
-	PhysicsBuildData.NoiseTexture = NoiseTexture;
+	//PhysicsBuildData.NoiseTexture = NoiseTexture;
 	PhysicsBuildData.NoiseMode = int32(NoiseMode);
 
 	// Build LUT
diff --git a/Engine/Plugins/FX/Niagara/Source/Niagara/Private/Stateless/Modules/NiagaraStatelessModule_SolveVelocitiesAndForces.cpp b/Engine/Plugins/FX/Niagara/Source/Niagara/Private/Stateless/Modules/NiagaraStatelessModule_SolveVelocitiesAndForces.cpp
--- a/Engine/Plugins/FX/Niagara/Source/Niagara/Private/Stateless/Modules/NiagaraStatelessModule_SolveVelocitiesAndForces.cpp
+++ b/Engine/Plugins/FX/Niagara/Source/Niagara/Private/Stateless/Modules/NiagaraStatelessModule_SolveVelocitiesAndForces.cpp
@@ -184,24 +184,24 @@ void UNiagaraStatelessModule_SolveVelocitiesAndForces::SetShaderParameters(const
 		Parameters->SolveVelocitiesAndForces_NoiseLUTChannelWidth	= PhysicsData.NoiseLUTChannelWidth;
 	}
 
-	FVectorFieldTextureAccessor TextureAccessor(Cast<UVectorField>(PhysicsData.NoiseTexture));
-
-	ENQUEUE_RENDER_COMMAND(FUpdateNoiseTexture)(
-		[Parameters, TextureAccessor](FRHICommandListImmediate& RHICmdList)
-		{
-			FRHITexture* NoiseTextureRHI = TextureAccessor.GetTexture();
-			Parameters->SolveVelocitiesAndForces_NoiseTexture = NoiseTextureRHI;
-			Parameters->SolveVelocitiesAndForces_NoiseSampler = TStaticSamplerState<SF_Bilinear, AM_Wrap, AM_Wrap, AM_Wrap>::GetRHI();
-
-			if (Parameters->SolveVelocitiesAndForces_NoiseMode < 2)
-			{
-				const FIntVector TextureSize = NoiseTextureRHI ? NoiseTextureRHI->GetSizeXYZ() : FIntVector(1, 1, 1);
-				Parameters->SolveVelocitiesAndForces_NoiseFrequency.X *= 1.0f / float(TextureSize.X);
-				Parameters->SolveVelocitiesAndForces_NoiseFrequency.Y *= 1.0f / float(TextureSize.Y);
-				Parameters->SolveVelocitiesAndForces_NoiseFrequency.Z *= 1.0f / float(TextureSize.Z);
-			}
-		}
-	);
+	//FVectorFieldTextureAccessor TextureAccessor(Cast<UVectorField>(PhysicsData.NoiseTexture));
+
+	//ENQUEUE_RENDER_COMMAND(FUpdateNoiseTexture)(
+	//	[Parameters, TextureAccessor](FRHICommandListImmediate& RHICmdList)
+	//	{
+	//		FRHITexture* NoiseTextureRHI = TextureAccessor.GetTexture();
+	//		Parameters->SolveVelocitiesAndForces_NoiseTexture = NoiseTextureRHI;
+	//		Parameters->SolveVelocitiesAndForces_NoiseSampler = TStaticSamplerState<SF_Bilinear, AM_Wrap, AM_Wrap, AM_Wrap>::GetRHI();
+
+	//		if (Parameters->SolveVelocitiesAndForces_NoiseMode < 2)
+	//		{
+	//			const FIntVector TextureSize = NoiseTextureRHI ? NoiseTextureRHI->GetSizeXYZ() : FIntVector(1, 1, 1);
+	//			Parameters->SolveVelocitiesAndForces_NoiseFrequency.X *= 1.0f / float(TextureSize.X);
+	//			Parameters->SolveVelocitiesAndForces_NoiseFrequency.Y *= 1.0f / float(TextureSize.Y);
+	//			Parameters->SolveVelocitiesAndForces_NoiseFrequency.Z *= 1.0f / float(TextureSize.Z);
+	//		}
+	//	}
+	//);
 }
 
 #if WITH_EDITORONLY_DATA
diff --git a/Engine/Plugins/FX/Niagara/Source/NiagaraShader/Internal/Stateless/NiagaraStatelessModuleShaderParameters.h b/Engine/Plugins/FX/Niagara/Source/NiagaraShader/Internal/Stateless/NiagaraStatelessModuleShaderParameters.h
--- a/Engine/Plugins/FX/Niagara/Source/NiagaraShader/Internal/Stateless/NiagaraStatelessModuleShaderParameters.h
+++ b/Engine/Plugins/FX/Niagara/Source/NiagaraShader/Internal/Stateless/NiagaraStatelessModuleShaderParameters.h
@@ -173,8 +173,8 @@ namespace NiagaraStateless
 		SHADER_PARAMETER(uint32,		SolveVelocitiesAndForces_NoiseEnabled)
 		SHADER_PARAMETER(float,			SolveVelocitiesAndForces_NoiseAmplitude)
 		SHADER_PARAMETER(FVector3f,		SolveVelocitiesAndForces_NoiseFrequency)
-		SHADER_PARAMETER_TEXTURE(Texture3D,		SolveVelocitiesAndForces_NoiseTexture)
-		SHADER_PARAMETER_SAMPLER(SamplerState,	SolveVelocitiesAndForces_NoiseSampler)
+		//SHADER_PARAMETER_TEXTURE(Texture3D,		SolveVelocitiesAndForces_NoiseTexture)
+		//SHADER_PARAMETER_SAMPLER(SamplerState,	SolveVelocitiesAndForces_NoiseSampler)
 
 		SHADER_PARAMETER(uint32, SolveVelocitiesAndForces_NoiseMode)
 		SHADER_PARAMETER(uint32, SolveVelocitiesAndForces_NoiseLUTOffset)
