Allow Blueprint and C++ delegates to subscribe to data channel writes. This ensures that the data is always received on the same tick that it's written and doesn't require the listener to actively poll the data channel each tick.

#tests Added engine test + tested manually with a BP delegate

diff --git a/Engine/Plugins/FX/Niagara/Source/Niagara/Private/NiagaraDataChannel.cpp b/Engine/Plugins/FX/Niagara/Source/Niagara/Private/NiagaraDataChannel.cpp
--- a/Engine/Plugins/FX/Niagara/Source/Niagara/Private/NiagaraDataChannel.cpp
+++ b/Engine/Plugins/FX/Niagara/Source/Niagara/Private/NiagaraDataChannel.cpp
@@ -1146,6 +1146,7 @@ int32 FNiagaraDataChannelData::ConsumePublishRequests(UNiagaraDataChannelHandler
 	}
 #endif
 
+	Owner->NotifySubscribers(this, GameDataOrigSize, NewGameDataChannel - GameDataOrigSize);
 	return RequestCount;
 }
 
@@ -1208,7 +1209,7 @@ bool FNiagaraDataChannelData::HasData()const
 	return GameData->Num() > 0 || (CPUSimData && CPUSimData->GetCurrentData()->GetNumInstances() > 0) || PublishRequestsForGPU.Num() > 0;
 }
 
-FNiagaraDataChannelGameDataPtr FNiagaraDataChannelData::GetGameDataForWriteGT(int32 Count, bool bVisibleToGame, bool bVisibleToCPU, bool bVisibleToGPU)
+FNiagaraDataChannelGameDataPtr FNiagaraDataChannelData::GetGameDataForWriteGT(int32 Count, bool bVisibleToGame, bool bVisibleToCPU, bool bVisibleToGPU, const FString& DebugSource)
 {
 	//We keep an array of cached data ptrs, one for each of the most recent combinations of visible flags.
 	//This dynamic choice of sending data to Game/CPU/GPU is honestly annoying. I think I might look at removing and having only as a system wide thing on the NDC	
@@ -1252,7 +1253,7 @@ FNiagaraDataChannelGameDataPtr FNiagaraDataChannelData::GetGameDataForWriteGT(in
 		PublishRequest.bVisibleToGPUSims = bVisibleToGPU;
 		PublishRequest.GameData = MakeShared<FNiagaraDataChannelGameData>(LayoutInfo);
 #if !UE_BUILD_SHIPPING
-		PublishRequest.DebugSource = TEXT("Accumulated Pending Game Data");
+		PublishRequest.DebugSource = DebugSource;
 #endif
 		Publish(PublishRequest); 
 		return PublishRequest.GameData;
@@ -1563,6 +1564,22 @@ void UNiagaraDataChannelLibrary::WriteToNiagaraDataChannelSingle(const UObject*,
 	// this function is just a placeholder and calls into CreateDataChannelWriter and its individual write functions from the BP node
 }
 
+void UNiagaraDataChannelLibrary::SubscribeToNiagaraDataChannel(const UObject* WorldContextObject, const UNiagaraDataChannelAsset* Channel, FNiagaraDataChannelSearchParameters SearchParams, const FOnNewNiagaraDataChannelPublish& UpdateDelegate, int32& UnsubscribeToken)
+{
+	if (UNiagaraDataChannelHandler* NiagaraDataChannelHandler = GetNiagaraDataChannel(WorldContextObject, Channel))
+	{
+		NiagaraDataChannelHandler->SubscribeToDataChannelUpdates(UpdateDelegate, SearchParams, UnsubscribeToken);
+	}
+}
+
+void UNiagaraDataChannelLibrary::UnsubscribeFromNiagaraDataChannel(const UObject* WorldContextObject, const UNiagaraDataChannelAsset* Channel, const int32& UnsubscribeToken)
+{
+	if (UNiagaraDataChannelHandler* NiagaraDataChannelHandler = GetNiagaraDataChannel(WorldContextObject, Channel))
+    {
+    	NiagaraDataChannelHandler->UnsubscribeFromDataChannelUpdates(UnsubscribeToken);
+    }
+}
+
 UNiagaraDataChannelHandler* UNiagaraDataChannelLibrary::FindDataChannelHandler(const UObject* WorldContextObject, const UNiagaraDataChannel* Channel)
 {
 	UWorld* World = (WorldContextObject != nullptr) ? GEngine->GetWorldFromContextObject(WorldContextObject, EGetWorldErrorMode::LogAndReturnNull) : nullptr;
diff --git a/Engine/Plugins/FX/Niagara/Source/Niagara/Private/NiagaraDataChannelAccessor.cpp b/Engine/Plugins/FX/Niagara/Source/Niagara/Private/NiagaraDataChannelAccessor.cpp
--- a/Engine/Plugins/FX/Niagara/Source/Niagara/Private/NiagaraDataChannelAccessor.cpp
+++ b/Engine/Plugins/FX/Niagara/Source/Niagara/Private/NiagaraDataChannelAccessor.cpp
@@ -172,7 +172,11 @@ bool UNiagaraDataChannelWriter::InitWrite(FNiagaraDataChannelSearchParameters Se
 	if(FNiagaraDataChannelDataPtr DestData = Owner->FindData(SearchParams, ENiagaraResourceAccess::WriteOnly))
 	{
 		//Attempt to use an existing cached dest data.
-		Data = DestData->GetGameDataForWriteGT(Count, bVisibleToGame, bVisibleToCPU, bVisibleToGPU);	
+		FString SourceString;
+#if !UE_BUILD_SHIPPING
+		SourceString = DebugSource.IsEmpty() ? GetPathName() : DebugSource;  
+#endif
+		Data = DestData->GetGameDataForWriteGT(Count, bVisibleToGame, bVisibleToCPU, bVisibleToGPU, SourceString);	
 		StartIndex = Data->Num();
 		Data->SetNum(StartIndex + Count);
 
@@ -375,7 +379,7 @@ bool FNDCWriterBase::BeginWrite(const UObject* WorldContextObject, const UNiagar
 				if (FNiagaraDataChannelDataPtr DestData = Handler->FindData(SearchParams, ENiagaraResourceAccess::WriteOnly))
 				{
 					//Attempt to use an existing cached dest data.
-					Data = DestData->GetGameDataForWriteGT(InCount, bVisibleToGame, bVisibleToCPU, bVisibleToGPU);
+					Data = DestData->GetGameDataForWriteGT(InCount, bVisibleToGame, bVisibleToCPU, bVisibleToGPU, DebugSource);
 
 					//We are potentially writing into the middle of an existing buffer here so track base index and count.
 					StartIndex = Data->Num();
diff --git a/Engine/Plugins/FX/Niagara/Source/Niagara/Private/NiagaraDataChannelHandler.cpp b/Engine/Plugins/FX/Niagara/Source/Niagara/Private/NiagaraDataChannelHandler.cpp
--- a/Engine/Plugins/FX/Niagara/Source/Niagara/Private/NiagaraDataChannelHandler.cpp
+++ b/Engine/Plugins/FX/Niagara/Source/Niagara/Private/NiagaraDataChannelHandler.cpp
@@ -88,6 +88,35 @@ UNiagaraDataChannelReader* UNiagaraDataChannelHandler::GetDataChannelReader()
 	return Reader;
 }
 
+void UNiagaraDataChannelHandler::SubscribeToDataChannelUpdates(FOnNewNiagaraDataChannelPublish UpdateDelegate, FNiagaraDataChannelSearchParameters SearchParams, int32& UnsubscribeToken)
+{
+	if (!UpdateDelegate.IsBound())
+	{
+		UnsubscribeToken = -1;
+		return;
+	}
+
+	SubscriberTokens++;
+	UnsubscribeToken = SubscriberTokens;
+	
+	FChannelSubscription& ChannelSubscription = ChannelSubscriptions.AddDefaulted_GetRef();
+	ChannelSubscription.SubscriptionToken = SubscriberTokens;
+	ChannelSubscription.OnPublishDelegate = UpdateDelegate;
+	ChannelSubscription.SearchParams = SearchParams;
+}
+
+void UNiagaraDataChannelHandler::UnsubscribeFromDataChannelUpdates(const int32& UnsubscribeToken)
+{
+	for (int i = 0; i < ChannelSubscriptions.Num(); i++)
+	{
+		if (ChannelSubscriptions[i].SubscriptionToken == UnsubscribeToken)
+		{
+			ChannelSubscriptions.RemoveAtSwap(i);
+			break;
+		}
+	}
+}
+
 FNiagaraDataChannelDataPtr UNiagaraDataChannelHandler::CreateData()
 {
 	FNiagaraDataChannelDataPtr Ret = MakeShared<FNiagaraDataChannelData>();
@@ -96,6 +125,36 @@ FNiagaraDataChannelDataPtr UNiagaraDataChannelHandler::CreateData()
 	return Ret;
 }
 
+void UNiagaraDataChannelHandler::NotifySubscribers(FNiagaraDataChannelData* Source, int32 StartIndex, int32 NumNewElements)
+{
+	if (NumNewElements == 0 || ChannelSubscriptions.IsEmpty() || StartIndex < 0)
+	{
+		return;
+	}
+
+	FNiagaraDataChannelUpdateContext UpdateContext;
+	UpdateContext.Reader = GetDataChannelReader();
+	UpdateContext.FirstNewDataIndex = StartIndex;
+	UpdateContext.LastNewDataIndex = StartIndex + NumNewElements - 1;
+	UpdateContext.NewElementCount = NumNewElements;
+	for (int i = ChannelSubscriptions.Num() - 1; i >= 0; i--)
+	{
+		FChannelSubscription& ChannelSubscription = ChannelSubscriptions[i];
+		if (ChannelSubscription.OnPublishDelegate.IsCompactable())
+		{
+			ChannelSubscriptions.RemoveAt(i);
+			continue;
+		}
+
+		FNiagaraDataChannelDataPtr ChannelData = FindData(ChannelSubscription.SearchParams, ENiagaraResourceAccess::ReadOnly);
+		if (ChannelData.Get() == Source)
+		{
+			UpdateContext.Reader->Data = ChannelData;
+			ChannelSubscription.OnPublishDelegate.Execute(UpdateContext);
+		}
+	}
+}
+
 void UNiagaraDataChannelHandler::OnComputeDispatchInterfaceDestroyed(FNiagaraGpuComputeDispatchInterface* InComputeDispatchInterface)
 {
 	//Destroy all RT proxies when the dispatcher is destroyed.
@@ -105,4 +164,4 @@ void UNiagaraDataChannelHandler::OnComputeDispatchInterfaceDestroyed(FNiagaraGpu
 		check(NDCData);
 		NDCData->DestroyRenderThreadProxy(InComputeDispatchInterface);
 	});
-}
\ No newline at end of file
+}
diff --git a/Engine/Plugins/FX/Niagara/Source/Niagara/Public/NiagaraDataChannel.h b/Engine/Plugins/FX/Niagara/Source/Niagara/Public/NiagaraDataChannel.h
--- a/Engine/Plugins/FX/Niagara/Source/Niagara/Public/NiagaraDataChannel.h
+++ b/Engine/Plugins/FX/Niagara/Source/Niagara/Public/NiagaraDataChannel.h
@@ -322,6 +322,28 @@ class NIAGARA_API UNiagaraDataChannelLibrary : public UBlueprintFunctionLibrary
 	UFUNCTION(BlueprintCallable, Category = NiagaraDataChannel, meta = (bReadPreviousFrame="true", AdvancedDisplay = "SearchParams, bReadPreviousFrame", Keywords = "niagara DataChannel num size", WorldContext = "WorldContextObject", UnsafeDuringActorConstruction = "true"))
 	static int32 GetDataChannelElementCount(const UObject* WorldContextObject, const UNiagaraDataChannelAsset* Channel, FNiagaraDataChannelSearchParameters SearchParams, bool bReadPreviousFrame);
 
+	/**
+	 * Subscribes to a single data channel and calls a delegate every times new data is written to the data channel.
+	 *
+	 * @param WorldContextObject	World to execute in
+	 * @param Channel				The channel to subscribe to for updates
+	 * @param SearchParams			Parameters used when retrieving a specific set of Data Channel Data to read - only used by some types of data channels, like the island types.
+	 * @param UpdateDelegate		The delegate to be called when new data is available in the data channel. Can be called multiple times per tick.
+	 * @param UnsubscribeToken		This token can be used to unsubscribe from the data channel.
+	 */
+	UFUNCTION(BlueprintCallable, Category = NiagaraDataChannel, meta = (Keywords = "niagara DataChannel event reader subscription delegate listener updates", WorldContext = "WorldContextObject", UnsafeDuringActorConstruction = "true"))
+	static void SubscribeToNiagaraDataChannel(const UObject* WorldContextObject, const UNiagaraDataChannelAsset* Channel, FNiagaraDataChannelSearchParameters SearchParams, const FOnNewNiagaraDataChannelPublish& UpdateDelegate, int32& UnsubscribeToken);
+
+	/**
+	 * Removes a prior registration from a data channel
+	 *
+	 * @param WorldContextObject	World to execute in
+	 * @param Channel				The channel to unsubscribe from
+	 * @param UnsubscribeToken		The token returned from the subscription call
+	 */
+	UFUNCTION(BlueprintCallable, Category = NiagaraDataChannel, meta = (Keywords = "niagara DataChannel event reader subscription delegate listener updates", WorldContext = "WorldContextObject", UnsafeDuringActorConstruction = "true"))
+	static void UnsubscribeFromNiagaraDataChannel(const UObject* WorldContextObject, const UNiagaraDataChannelAsset* Channel, const int32& UnsubscribeToken);
+	
 	/**
 	 * Reads a single entry from the given data channel, if possible.
 	 *
diff --git a/Engine/Plugins/FX/Niagara/Source/Niagara/Public/NiagaraDataChannelAccessor.h b/Engine/Plugins/FX/Niagara/Source/Niagara/Public/NiagaraDataChannelAccessor.h
--- a/Engine/Plugins/FX/Niagara/Source/Niagara/Public/NiagaraDataChannelAccessor.h
+++ b/Engine/Plugins/FX/Niagara/Source/Niagara/Public/NiagaraDataChannelAccessor.h
@@ -21,7 +21,8 @@ class UNiagaraDataChannelReader : public UObject
 {
 	GENERATED_BODY()
 private:
-
+	friend UNiagaraDataChannelHandler;
+	
 	FNiagaraDataChannelDataPtr Data = nullptr;
 	bool bReadingPreviousFrame = false;
 
@@ -338,6 +339,9 @@ public:
 
 	//Finalizes the write and free/reset any data.
 	void EndWrite();
+
+	// Debug string that can be used to track the source of data in a data channel 
+	FString DebugSource;
 };
 
 /** Utility class for simplified writing of variables to a Niagara Data Channel. See full description above. */
diff --git a/Engine/Plugins/FX/Niagara/Source/Niagara/Public/NiagaraDataChannelCommon.h b/Engine/Plugins/FX/Niagara/Source/Niagara/Public/NiagaraDataChannelCommon.h
--- a/Engine/Plugins/FX/Niagara/Source/Niagara/Public/NiagaraDataChannelCommon.h
+++ b/Engine/Plugins/FX/Niagara/Source/Niagara/Public/NiagaraDataChannelCommon.h
@@ -110,7 +110,7 @@ struct FNiagaraDataChannelData final : public TSharedFromThis<FNiagaraDataChanne
 	bool HasData()const;
 
 	//Returns a game data buffer into which we can write Count valuse on the Game Thread.
-	FNiagaraDataChannelGameDataPtr GetGameDataForWriteGT(int32 Count, bool bVisibleToGame, bool bVisibleToCPU, bool bVisibleToGPU);
+	FNiagaraDataChannelGameDataPtr GetGameDataForWriteGT(int32 Count, bool bVisibleToGame, bool bVisibleToCPU, bool bVisibleToGPU, const FString& DebugSource);
 
 private:
 
diff --git a/Engine/Plugins/FX/Niagara/Source/Niagara/Public/NiagaraDataChannelHandler.h b/Engine/Plugins/FX/Niagara/Source/Niagara/Public/NiagaraDataChannelHandler.h
--- a/Engine/Plugins/FX/Niagara/Source/Niagara/Public/NiagaraDataChannelHandler.h
+++ b/Engine/Plugins/FX/Niagara/Source/Niagara/Public/NiagaraDataChannelHandler.h
@@ -24,7 +24,7 @@ class UNiagaraDataChannelHandler : public UObject
 public:
 
 	GENERATED_BODY()
-	
+
 	//UObject Interface
 	NIAGARA_API virtual void BeginDestroy()override;
 	NIAGARA_API virtual bool IsReadyForFinishDestroy()override;
@@ -59,6 +59,14 @@ public:
 	UFUNCTION(BlueprintCallable, Category = "Data Channel")
 	NIAGARA_API UNiagaraDataChannelReader* GetDataChannelReader();
 
+	/** The provided delegate will be called whenever new entries are added to the relevant data channel. This means the delegate can be called multiple times per tick.
+	 * This is only relevant for data published to the game thread, so no gpu data or data that's only visible to niagara systems. */
+	UFUNCTION(BlueprintCallable, Category = "Data Channel")
+	NIAGARA_API void SubscribeToDataChannelUpdates(FOnNewNiagaraDataChannelPublish UpdateDelegate, FNiagaraDataChannelSearchParameters SearchParams, int32& UnsubscribeToken);
+
+	UFUNCTION(BlueprintCallable, Category = "Data Channel")
+	NIAGARA_API void UnsubscribeFromDataChannelUpdates(const int32& UnsubscribeToken);
+
 	template<typename T> 
 	T* GetChannelTyped()const{ return Cast<T>(DataChannel); }
 
@@ -67,6 +75,7 @@ public:
 	/** Returns the tick group we're currently being processed or most recently processed. */
 	ETickingGroup GetCurrentTickGroup() { return CurrentTG; }
 
+	void NotifySubscribers(FNiagaraDataChannelData* Source, int32 StartIndex, int32 NumNewElements);
 	void OnComputeDispatchInterfaceDestroyed(FNiagaraGpuComputeDispatchInterface* InComputeDispatchInterface);
 
 	template<typename TAction>
@@ -92,6 +101,19 @@ protected:
 	TArray<TWeakPtr<FNiagaraDataChannelData>> WeakDataArray;
 
 	FRenderCommandFence RTFence;
+
+private:
+
+	struct FChannelSubscription
+	{
+		FOnNewNiagaraDataChannelPublish OnPublishDelegate;
+		FNiagaraDataChannelSearchParameters SearchParams;
+		int32 SubscriptionToken = -1;
+	};
+	TArray<FChannelSubscription> ChannelSubscriptions;
+
+	// used to generate the unsubscribe tokens that are stored in ChannelSubscriptions
+	int32 SubscriberTokens = 0;
 };
 
 template<typename TAction>
diff --git a/Engine/Plugins/FX/Niagara/Source/Niagara/Public/NiagaraDataChannelPublic.h b/Engine/Plugins/FX/Niagara/Source/Niagara/Public/NiagaraDataChannelPublic.h
--- a/Engine/Plugins/FX/Niagara/Source/Niagara/Public/NiagaraDataChannelPublic.h
+++ b/Engine/Plugins/FX/Niagara/Source/Niagara/Public/NiagaraDataChannelPublic.h
@@ -18,6 +18,26 @@ using FNiagaraDataChannelGameDataPtr = TSharedPtr<FNiagaraDataChannelGameData>;
 using FNiagaraDataChannelDataPtr = TSharedPtr<FNiagaraDataChannelData>;
 using FNiagaraDataChannelLayoutInfoPtr = TSharedPtr<FNiagaraDataChannelLayoutInfo>;
 
+USTRUCT(BlueprintType)
+struct NIAGARA_API FNiagaraDataChannelUpdateContext
+{
+	GENERATED_BODY();
+
+	UPROPERTY(BlueprintReadOnly, Category = "Data Channel")
+	TObjectPtr<class UNiagaraDataChannelReader> Reader;
+
+	UPROPERTY(BlueprintReadOnly, Category = "Data Channel")
+	int32 FirstNewDataIndex;
+
+	UPROPERTY(BlueprintReadOnly, Category = "Data Channel")
+	int32 LastNewDataIndex;
+
+	UPROPERTY(BlueprintReadOnly, Category = "Data Channel")
+	int32 NewElementCount;
+};
+
+DECLARE_DYNAMIC_DELEGATE_OneParam(FOnNewNiagaraDataChannelPublish, const FNiagaraDataChannelUpdateContext&, NewDataContext);
+
 /** Niagara Data Channels are a system for communication between Niagara Systems and with game code/Blueprint.
 
 Data channel assets define the payload as well as some transfer settings.
