Fix for visualizing compressed particle attributes in the debug hud
#tests enable a compress attributes on a single particle system and then visualize the position of the particle in debug hud

diff --git a/Engine/Plugins/FX/Niagara/Source/Niagara/Private/NiagaraDataSetDebugAccessor.cpp b/Engine/Plugins/FX/Niagara/Source/Niagara/Private/NiagaraDataSetDebugAccessor.cpp
--- a/Engine/Plugins/FX/Niagara/Source/Niagara/Private/NiagaraDataSetDebugAccessor.cpp
+++ b/Engine/Plugins/FX/Niagara/Source/Niagara/Private/NiagaraDataSetDebugAccessor.cpp
@@ -52,7 +52,7 @@ float FNiagaraDataSetDebugAccessor::ReadHalf(const FNiagaraDataBuffer* DataBuffe
 {
 	if (DataBuffer != nullptr && ComponentIndexHalf != INDEX_NONE && Instance < DataBuffer->GetNumInstances() && Component < NumComponentsHalf)
 	{
-		const FFloat16* HalfData = reinterpret_cast<const FFloat16*>(DataBuffer->GetComponentPtrHalf(NumComponentsHalf + Component));
+		const FFloat16* HalfData = reinterpret_cast<const FFloat16*>(DataBuffer->GetComponentPtrHalf(ComponentIndexHalf + Component));
 		return HalfData[Instance];
 	}
 	return 0.0f;
