- Downgrade check to ensure
- We still have some paths occasionally hitting this but the repro is not clear.  This will make sure we don't take down the editor.

#rnx

diff --git a/Engine/Plugins/FX/Niagara/Source/Niagara/Public/NiagaraParameterStore.h b/Engine/Plugins/FX/Niagara/Source/Niagara/Public/NiagaraParameterStore.h
--- a/Engine/Plugins/FX/Niagara/Source/Niagara/Public/NiagaraParameterStore.h
+++ b/Engine/Plugins/FX/Niagara/Source/Niagara/Public/NiagaraParameterStore.h
@@ -440,7 +440,7 @@ public:
 
 	FORCEINLINE const uint8* GetParameterData(int32 Offset, int32 SizeInBytes) const
 	{
-		check(Offset >= 0 && (Offset + SizeInBytes) <= ParameterData.Num());
+		ensureMsgf(Offset >= 0 && (Offset + SizeInBytes) <= ParameterData.Num(), TEXT("OOB immutable access offset(%d) size(%d) store size(%d)"), Offset, SizeInBytes, ParameterData.Num());
 		return ParameterData.GetData() + Offset;
 	}
 
