- Add Direct Set & Binding support to stateless sub UV animation module

diff --git a/Engine/Plugins/FX/Niagara/Shaders/Private/Stateless/Modules/NiagaraStatelessModule_SubUVAnimation.ush b/Engine/Plugins/FX/Niagara/Shaders/Private/Stateless/Modules/NiagaraStatelessModule_SubUVAnimation.ush
--- a/Engine/Plugins/FX/Niagara/Shaders/Private/Stateless/Modules/NiagaraStatelessModule_SubUVAnimation.ush
+++ b/Engine/Plugins/FX/Niagara/Shaders/Private/Stateless/Modules/NiagaraStatelessModule_SubUVAnimation.ush
@@ -16,12 +16,12 @@ void SubUVAnimation_Simulate(inout FStatelessParticle Particle)
 	const uint SeedOffset = uint(Particle.Age * SubUVAnimation_InitialFrameRateChange);
 	float Frame = RandomFloat(SeedOffset) * SubUVAnimation_InitialFrameScale + SubUVAnimation_InitialFrameBias;
 
-	if ( SubUVAnimation_Mode == 0 )
+	if ( SubUVAnimation_Mode == 1 )
 	{
 		const float Interp = Particle.Age * SubUVAnimation_RateScale;
 		Frame = frac(Frame + SubUVAnimation_AnimFrameStart + (Interp * SubUVAnimation_AnimFrameRange));
 	}
-	else if ( SubUVAnimation_Mode == 1 )
+	else if ( SubUVAnimation_Mode == 2 )
 	{
 		const float Interp = Particle.NormalizedAge * SubUVAnimation_RateScale;
 		Frame = saturate(Frame + SubUVAnimation_AnimFrameStart + (Interp * SubUVAnimation_AnimFrameRange));
diff --git a/Engine/Plugins/FX/Niagara/Source/Niagara/Internal/Stateless/Modules/NiagaraStatelessModule_SubUVAnimation.h b/Engine/Plugins/FX/Niagara/Source/Niagara/Internal/Stateless/Modules/NiagaraStatelessModule_SubUVAnimation.h
--- a/Engine/Plugins/FX/Niagara/Source/Niagara/Internal/Stateless/Modules/NiagaraStatelessModule_SubUVAnimation.h
+++ b/Engine/Plugins/FX/Niagara/Source/Niagara/Internal/Stateless/Modules/NiagaraStatelessModule_SubUVAnimation.h
@@ -3,15 +3,13 @@
 #pragma once
 
 #include "Stateless/NiagaraStatelessModule.h"
-#include "Stateless/NiagaraStatelessEmitterDataBuildContext.h"
-#include "Stateless/NiagaraStatelessModuleShaderParameters.h"
-#include "Stateless/NiagaraStatelessParticleSimContext.h"
 
 #include "NiagaraStatelessModule_SubUVAnimation.generated.h"
 
 UENUM()
 enum class ENSMSubUVAnimation_Mode
 {
+	DirectSet,
 	InfiniteLoop,
 	Linear,
 	Random,
@@ -24,26 +22,13 @@ class UNiagaraStatelessModule_SubUVAnimation : public UNiagaraStatelessModule
 {
 	GENERATED_BODY()
 
-	struct FModuleBuiltData
-	{
-		int32	Mode					= 0;
-		float	NumFrames				= 0.0f;
-		float	InitialFrameScale		= 0.0f;
-		float	InitialFrameBias		= 0.0f;
-		float	InitialFrameRateChange	= 0.0f;
-		float	AnimFrameStart			= 0.0f;
-		float	AnimFrameRange			= 0.0f;
-		float	RateScale				= 0.0f;
-
-		int32	SubImageIndexVariableOffset = INDEX_NONE;
-	};
-
 public:
-	using FParameters = NiagaraStateless::FSubUVAnimationModule_ShaderParameters;
-
 	UPROPERTY(EditAnywhere, Category = "Parameters")
 	int32	NumFrames = 16;
 
+	UPROPERTY(EditAnywhere, Category = "Parameters", meta = (EditConditionHides, EditCondition = "AnimationMode == ENSMSubUVAnimation_Mode::DirectSet", DisableRange))
+	FNiagaraDistributionRangeInt FrameIndex = FNiagaraDistributionRangeInt(0);
+
 	UPROPERTY(EditAnywhere, Category = "Parameters", meta = (InlineEditConditionToggle))
 	bool	bStartFrameRangeOverride_Enabled = false;
 
@@ -72,119 +57,15 @@ public:
 	UPROPERTY(EditAnywhere, Category = "Parameters", meta = (EditConditionHides, EditCondition = "AnimationMode == ENSMSubUVAnimation_Mode::Random"))
 	float RandomChangeInterval = 0.1f;
 
-	virtual void BuildEmitterData(const FNiagaraStatelessEmitterDataBuildContext& BuildContext) const override
-	{
-		FModuleBuiltData* BuiltData = BuildContext.AllocateBuiltData<FModuleBuiltData>();
-		if (!IsModuleEnabled())
-		{
-			return;
-		}
-
-		const FNiagaraStatelessGlobals& StatelessGlobals = FNiagaraStatelessGlobals::Get();
-		BuiltData->SubImageIndexVariableOffset = BuildContext.FindParticleVariableIndex(StatelessGlobals.SubImageIndexVariable);
-		if (BuiltData->SubImageIndexVariableOffset == INDEX_NONE)
-		{
-			return;
-		}
-
-		const float FrameRangeStart	= bStartFrameRangeOverride_Enabled ? FMath::Clamp(float(StartFrameRangeOverride) / float(NumFrames - 1), 0.0f, 1.0f) : 0.0f;
-		const float FrameRangeEnd	= bEndFrameRangeOverride_Enabled   ? FMath::Clamp(float(EndFrameRangeOverride)   / float(NumFrames - 1), 0.0f, 1.0f) : 1.0f;
-
-		BuiltData->Mode			= (int)AnimationMode;
-		BuiltData->NumFrames	= float(NumFrames);
-		switch (AnimationMode)
-		{
-			case ENSMSubUVAnimation_Mode::InfiniteLoop:
-				BuiltData->InitialFrameScale		= 0.0f;
-				BuiltData->InitialFrameBias			= 0.0f;
-				BuiltData->InitialFrameRateChange	= 0.0f;
-				BuiltData->AnimFrameStart			= FrameRangeStart;
-				BuiltData->AnimFrameRange			= FrameRangeEnd - FrameRangeStart;
-				BuiltData->RateScale				= LoopsPerSecond;
-				break;
-
-			case ENSMSubUVAnimation_Mode::Linear:
-				BuiltData->InitialFrameScale		= 0.0f;
-				BuiltData->InitialFrameBias			= 0.0f;
-				BuiltData->InitialFrameRateChange	= 0.0f;
-				BuiltData->AnimFrameStart			= FrameRangeStart;
-				BuiltData->AnimFrameRange			= FrameRangeEnd - FrameRangeStart;
-				BuiltData->RateScale				= 1.0f;
-				break;
-
-			case ENSMSubUVAnimation_Mode::Random:
-				BuiltData->InitialFrameScale		= FrameRangeEnd - FrameRangeStart;
-				BuiltData->InitialFrameBias			= FrameRangeStart;
-				BuiltData->InitialFrameRateChange	= RandomChangeInterval > 0.0f ? 1.0f / RandomChangeInterval : 0.0f;
-				BuiltData->AnimFrameStart			= 0.0f;
-				BuiltData->AnimFrameRange			= 0.0f;
-				BuiltData->RateScale				= 0.0f;
-				break;
-		}
-
-		BuildContext.AddParticleSimulationExecSimulate(&UNiagaraStatelessModule_SubUVAnimation::ParticleSimulate);
-	}
-
-	virtual void BuildShaderParameters(FNiagaraStatelessShaderParametersBuilder& ShaderParametersBuilder) const override
-	{
-		ShaderParametersBuilder.AddParameterNestedStruct<FParameters>();
-	}
-
-	virtual void SetShaderParameters(const FNiagaraStatelessSetShaderParameterContext& SetShaderParameterContext) const override
-	{
-		FParameters* Parameters = SetShaderParameterContext.GetParameterNestedStruct<FParameters>();
-		const FModuleBuiltData* ModuleBuiltData = SetShaderParameterContext.ReadBuiltData<FModuleBuiltData>();
-
-		Parameters->SubUVAnimation_Mode						= ModuleBuiltData->Mode;
-		Parameters->SubUVAnimation_NumFrames				= ModuleBuiltData->NumFrames;
-		Parameters->SubUVAnimation_InitialFrameScale		= ModuleBuiltData->InitialFrameScale;
-		Parameters->SubUVAnimation_InitialFrameBias			= ModuleBuiltData->InitialFrameBias;
-		Parameters->SubUVAnimation_InitialFrameRateChange	= ModuleBuiltData->InitialFrameRateChange;
-		Parameters->SubUVAnimation_AnimFrameStart			= ModuleBuiltData->AnimFrameStart;
-		Parameters->SubUVAnimation_AnimFrameRange			= ModuleBuiltData->AnimFrameRange;
-		Parameters->SubUVAnimation_RateScale				= ModuleBuiltData->RateScale;
-	}
-
-	static void ParticleSimulate(const NiagaraStateless::FParticleSimulationContext& ParticleSimulationContext)
-	{
-		using namespace NiagaraStateless;
-
-		const FModuleBuiltData* ModuleBuiltData = ParticleSimulationContext.ReadBuiltData<FModuleBuiltData>();
-		const float* ParticleAges = ParticleSimulationContext.GetParticleAge();
-		const float* ParticleNormalizedAges = ParticleSimulationContext.GetParticleNormalizedAge();
-
-		for (uint32 i = 0; i < ParticleSimulationContext.GetNumInstances(); ++i)
-		{
-			const float Particle_Age = ParticleAges[i];
-			const float Particle_NormalizedAge = ParticleNormalizedAges[i];
-
-			const uint32 SeedOffset = uint32(Particle_Age * ModuleBuiltData->InitialFrameRateChange);
-			float Frame = ParticleSimulationContext.RandomFloat(i, SeedOffset) * ModuleBuiltData->InitialFrameScale + ModuleBuiltData->InitialFrameBias;
-
-			if (ModuleBuiltData->Mode == 0)
-			{
-				const float Interp = Particle_Age * ModuleBuiltData->RateScale;
-				Frame = FMath::Fractional(Frame + ModuleBuiltData->AnimFrameStart + (Interp * ModuleBuiltData->AnimFrameRange));
-			}
-			else if (ModuleBuiltData->Mode == 1)
-			{
-				const float Interp = Particle_NormalizedAge * ModuleBuiltData->RateScale;
-				Frame = FMath::Clamp(Frame + ModuleBuiltData->AnimFrameStart + (Interp * ModuleBuiltData->AnimFrameRange), 0.0f, 1.0f);
-			}
-			const float SubImageIndex = Frame * ModuleBuiltData->NumFrames;
-
-			ParticleSimulationContext.WriteParticleVariable(ModuleBuiltData->SubImageIndexVariableOffset, i, SubImageIndex);
-		}
-	}
+	virtual void BuildEmitterData(const FNiagaraStatelessEmitterDataBuildContext& BuildContext) const override;
+	virtual void BuildShaderParameters(FNiagaraStatelessShaderParametersBuilder& ShaderParametersBuilder) const override;
+	virtual void SetShaderParameters(const FNiagaraStatelessSetShaderParameterContext& SetShaderParameterContext) const override;
 
 #if WITH_EDITOR
 	virtual bool CanDisableModule() const override { return true; }
 #endif
+
 #if WITH_EDITORONLY_DATA
-	virtual void GetOutputVariables(TArray<FNiagaraVariableBase>& OutVariables) const override
-	{
-		const FNiagaraStatelessGlobals& StatelessGlobals = FNiagaraStatelessGlobals::Get();
-		OutVariables.AddUnique(StatelessGlobals.SubImageIndexVariable);
-	}
+	virtual void GetOutputVariables(TArray<FNiagaraVariableBase>& OutVariables) const override;
 #endif
 };
diff --git a/Engine/Plugins/FX/Niagara/Source/Niagara/Internal/Stateless/NiagaraStatelessParticleSimContext.h b/Engine/Plugins/FX/Niagara/Source/Niagara/Internal/Stateless/NiagaraStatelessParticleSimContext.h
--- a/Engine/Plugins/FX/Niagara/Source/Niagara/Internal/Stateless/NiagaraStatelessParticleSimContext.h
+++ b/Engine/Plugins/FX/Niagara/Source/Niagara/Internal/Stateless/NiagaraStatelessParticleSimContext.h
@@ -176,6 +176,7 @@ public:
 	FVector2f		SafeNormalize(const FVector2f& v) const { return SafeNormalize(v, FVector2f(1.0f, 0.0f)); }
 	FVector3f		SafeNormalize(const FVector3f& v) const { return SafeNormalize(v, FVector3f(1.0f, 0.0f, 0.0f)); }
 
+	void			ConvertRangeToScaleBias(const FNiagaraStatelessRangeInt&     Range, int32& OutScale, int32& OutBias) const { OutScale = Range.ParameterOffset == INDEX_NONE ? Range.GetScale() : 0; OutBias = Range.ParameterOffset == INDEX_NONE ? Range.Min : GetParameterBufferValue<int32>(Range.ParameterOffset, 0); }
 	void			ConvertRangeToScaleBias(const FNiagaraStatelessRangeFloat&   Range, float& OutScale, float& OutBias) const { OutScale = Range.ParameterOffset == INDEX_NONE ? Range.GetScale() : 0.0f; OutBias = Range.ParameterOffset == INDEX_NONE ? Range.Min : GetParameterBufferValue<float>(Range.ParameterOffset, 0); }
 	void			ConvertRangeToScaleBias(const FNiagaraStatelessRangeVector2& Range, FVector2f& OutScale, FVector2f& OutBias) const { OutScale = Range.ParameterOffset == INDEX_NONE ? Range.GetScale() : FVector2f::ZeroVector; OutBias = Range.ParameterOffset == INDEX_NONE ? Range.Min : GetParameterBufferValue<FVector2f>(Range.ParameterOffset, 0); }
 	void			ConvertRangeToScaleBias(const FNiagaraStatelessRangeVector3& Range, FVector3f& OutScale, FVector3f& OutBias) const { OutScale = Range.ParameterOffset == INDEX_NONE ? Range.GetScale() : FVector3f::ZeroVector; OutBias = Range.ParameterOffset == INDEX_NONE ? Range.Min : GetParameterBufferValue<FVector3f>(Range.ParameterOffset, 0); }
diff --git a/Engine/Plugins/FX/Niagara/Source/Niagara/Private/Stateless/Modules/NiagaraStatelessModule_SubUVAnimation.cpp b/Engine/Plugins/FX/Niagara/Source/Niagara/Private/Stateless/Modules/NiagaraStatelessModule_SubUVAnimation.cpp
new file mode 100644
--- /dev/null
+++ b/Engine/Plugins/FX/Niagara/Source/Niagara/Private/Stateless/Modules/NiagaraStatelessModule_SubUVAnimation.cpp
@@ -0,0 +1,170 @@
+// Copyright Epic Games, Inc. All Rights Reserved.
+
+#include "Stateless/Modules/NiagaraStatelessModule_SubUVAnimation.h"
+
+#include "Stateless/NiagaraStatelessEmitterDataBuildContext.h"
+#include "Stateless/NiagaraStatelessModuleShaderParameters.h"
+#include "Stateless/NiagaraStatelessParticleSimContext.h"
+
+namespace NSMSubUVAnimationPrivate
+{
+	using FParameters = NiagaraStateless::FSubUVAnimationModule_ShaderParameters;
+
+	struct FModuleBuiltData
+	{
+		int32						Mode = 0;
+		int32						NumFrames = 0;
+		FNiagaraStatelessRangeInt	InitialFrame;
+		float						InitialFrameRateChange = 0.0f;
+		int32						AnimFrameStart = 0;
+		int32						AnimFrameRange = 0;
+		float						RateScale = 0.0f;
+
+		int32						SubImageIndexVariableOffset = INDEX_NONE;
+	};
+
+	static_assert(int(ENSMSubUVAnimation_Mode::InfiniteLoop) == 1 && int(ENSMSubUVAnimation_Mode::Linear) == 2, "Enum is out of sync with simulation");
+
+	void ParticleSimulate(const NiagaraStateless::FParticleSimulationContext& ParticleSimulationContext)
+	{
+		using namespace NiagaraStateless;
+
+		const FModuleBuiltData* ModuleBuiltData = ParticleSimulationContext.ReadBuiltData<FModuleBuiltData>();
+		const FParameters* ShaderParameters = ParticleSimulationContext.ReadParameterNestedStruct<FParameters>();
+
+		const float* ParticleAges = ParticleSimulationContext.GetParticleAge();
+		const float* ParticleNormalizedAges = ParticleSimulationContext.GetParticleNormalizedAge();
+
+		for (uint32 i = 0; i < ParticleSimulationContext.GetNumInstances(); ++i)
+		{
+			const float Particle_Age = ParticleAges[i];
+			const float Particle_NormalizedAge = ParticleNormalizedAges[i];
+
+			const uint32 SeedOffset = uint32(Particle_Age * ShaderParameters->SubUVAnimation_InitialFrameRateChange);
+			float Frame = ParticleSimulationContext.RandomFloat(i, SeedOffset) * ShaderParameters->SubUVAnimation_InitialFrameScale + ShaderParameters->SubUVAnimation_InitialFrameBias;
+
+			if (ShaderParameters->SubUVAnimation_Mode == 1)
+			{
+				const float Interp = Particle_Age * ShaderParameters->SubUVAnimation_RateScale;
+				Frame = FMath::Fractional(Frame + ShaderParameters->SubUVAnimation_AnimFrameStart + (Interp * ShaderParameters->SubUVAnimation_AnimFrameRange));
+			}
+			else if (ShaderParameters->SubUVAnimation_Mode == 2)
+			{
+				const float Interp = Particle_NormalizedAge * ShaderParameters->SubUVAnimation_RateScale;
+				Frame = FMath::Clamp(Frame + ShaderParameters->SubUVAnimation_AnimFrameStart + (Interp * ShaderParameters->SubUVAnimation_AnimFrameRange), 0.0f, 1.0f);
+			}
+			const float SubImageIndex = Frame * ShaderParameters->SubUVAnimation_NumFrames;
+
+			ParticleSimulationContext.WriteParticleVariable(ModuleBuiltData->SubImageIndexVariableOffset, i, SubImageIndex);
+		}
+	}
+}
+
+void UNiagaraStatelessModule_SubUVAnimation::BuildEmitterData(const FNiagaraStatelessEmitterDataBuildContext& BuildContext) const
+{
+	using namespace NSMSubUVAnimationPrivate;
+
+	FModuleBuiltData* BuiltData = BuildContext.AllocateBuiltData<FModuleBuiltData>();
+	if (!IsModuleEnabled())
+	{
+		return;
+	}
+
+	const FNiagaraStatelessGlobals& StatelessGlobals = FNiagaraStatelessGlobals::Get();
+	BuiltData->SubImageIndexVariableOffset = BuildContext.FindParticleVariableIndex(StatelessGlobals.SubImageIndexVariable);
+	if (BuiltData->SubImageIndexVariableOffset == INDEX_NONE)
+	{
+		return;
+	}
+
+	if (NumFrames <= 1)
+	{
+		return;
+	}
+
+	const int32 NumFramesMinusOne	= NumFrames - 1;
+	const int32 FrameRangeStart		= bStartFrameRangeOverride_Enabled ? FMath::Clamp(StartFrameRangeOverride, 0, NumFramesMinusOne) : 0;
+	const int32 FrameRangeEnd		= bEndFrameRangeOverride_Enabled ? FMath::Clamp(EndFrameRangeOverride, 0, NumFramesMinusOne) : NumFramesMinusOne;
+
+	BuiltData->Mode			= int(AnimationMode);
+	BuiltData->NumFrames	= float(NumFrames);
+	switch (AnimationMode)
+	{
+		case ENSMSubUVAnimation_Mode::DirectSet:
+			BuiltData->InitialFrame	 = BuildContext.ConvertDistributionToRange(FrameIndex, 0);
+			break;
+
+		case ENSMSubUVAnimation_Mode::InfiniteLoop:
+			BuiltData->AnimFrameStart	= FrameRangeStart;
+			BuiltData->AnimFrameRange	= FrameRangeEnd - FrameRangeStart;
+			BuiltData->RateScale		= LoopsPerSecond;
+			break;
+
+		case ENSMSubUVAnimation_Mode::Linear:
+			BuiltData->AnimFrameStart	= FrameRangeStart;
+			BuiltData->AnimFrameRange	= FrameRangeEnd - FrameRangeStart;
+			BuiltData->RateScale		= 1.0f;
+			break;
+
+		case ENSMSubUVAnimation_Mode::Random:
+			BuiltData->InitialFrame				= FNiagaraStatelessRangeInt(FrameRangeStart, FrameRangeEnd);
+			BuiltData->InitialFrameRateChange	= RandomChangeInterval > 0.0f ? 1.0f / RandomChangeInterval : 0.0f;
+			break;
+	}
+
+	BuildContext.AddParticleSimulationExecSimulate(&NSMSubUVAnimationPrivate::ParticleSimulate);
+}
+
+void UNiagaraStatelessModule_SubUVAnimation::BuildShaderParameters(FNiagaraStatelessShaderParametersBuilder& ShaderParametersBuilder) const
+{
+	using namespace NSMSubUVAnimationPrivate;
+	
+	ShaderParametersBuilder.AddParameterNestedStruct<FParameters>();
+}
+
+void UNiagaraStatelessModule_SubUVAnimation::SetShaderParameters(const FNiagaraStatelessSetShaderParameterContext& SetShaderParameterContext) const
+{
+	using namespace NSMSubUVAnimationPrivate;
+
+	FParameters* Parameters = SetShaderParameterContext.GetParameterNestedStruct<FParameters>();
+	const FModuleBuiltData* ModuleBuiltData = SetShaderParameterContext.ReadBuiltData<FModuleBuiltData>();
+
+	if (ModuleBuiltData->NumFrames > 1)
+	{
+		//const float fNumFramesMinusOne = ModuleBuiltData->NumFrames > 1 ? float(ModuleBuiltData->NumFrames - 1) : 1.0f;
+		const float fNumFrames = float(ModuleBuiltData->NumFrames);
+
+		int32 InitialFrameScale = 0;
+		int32 InitialFrameBias = 0;
+		SetShaderParameterContext.ConvertRangeToScaleBias(ModuleBuiltData->InitialFrame, InitialFrameScale, InitialFrameBias);
+
+		Parameters->SubUVAnimation_Mode						= ModuleBuiltData->Mode;
+		Parameters->SubUVAnimation_NumFrames				= fNumFrames;
+		Parameters->SubUVAnimation_InitialFrameScale		= float(InitialFrameScale) / fNumFrames;
+		Parameters->SubUVAnimation_InitialFrameBias			= float(InitialFrameBias) / fNumFrames;
+		Parameters->SubUVAnimation_InitialFrameRateChange	= ModuleBuiltData->InitialFrameRateChange;
+		Parameters->SubUVAnimation_AnimFrameStart			= ModuleBuiltData->AnimFrameStart / fNumFrames;
+		Parameters->SubUVAnimation_AnimFrameRange			= ModuleBuiltData->AnimFrameRange / fNumFrames;
+		Parameters->SubUVAnimation_RateScale				= ModuleBuiltData->RateScale;
+	}
+	else
+	{
+		Parameters->SubUVAnimation_Mode						= 0;
+		Parameters->SubUVAnimation_NumFrames				= 0.0f;
+		Parameters->SubUVAnimation_InitialFrameScale		= 0.0f;
+		Parameters->SubUVAnimation_InitialFrameBias			= 0.0f;
+		Parameters->SubUVAnimation_InitialFrameRateChange	= 0.0f;
+		Parameters->SubUVAnimation_AnimFrameStart			= 0.0f;
+		Parameters->SubUVAnimation_AnimFrameRange			= 0.0f;
+		Parameters->SubUVAnimation_RateScale				= 0.0f;
+	}
+}
+
+
+#if WITH_EDITORONLY_DATA
+void UNiagaraStatelessModule_SubUVAnimation::GetOutputVariables(TArray<FNiagaraVariableBase>& OutVariables) const
+{
+	const FNiagaraStatelessGlobals& StatelessGlobals = FNiagaraStatelessGlobals::Get();
+	OutVariables.AddUnique(StatelessGlobals.SubImageIndexVariable);
+}
+#endif
