- Disable compression on Position types to fix precision issues when compression is enabled

diff --git a/Engine/Plugins/FX/Niagara/Source/NiagaraEditor/Private/NiagaraHlslTranslator.cpp b/Engine/Plugins/FX/Niagara/Source/NiagaraEditor/Private/NiagaraHlslTranslator.cpp
--- a/Engine/Plugins/FX/Niagara/Source/NiagaraEditor/Private/NiagaraHlslTranslator.cpp
+++ b/Engine/Plugins/FX/Niagara/Source/NiagaraEditor/Private/NiagaraHlslTranslator.cpp
@@ -1425,7 +1425,6 @@ static void ConvertFloatToHalf(const FNiagaraCompileOptions& InCompileOptions, T
 			{ FNiagaraTypeDefinition::GetFloatDef(), FNiagaraTypeDefinition::GetHalfDef() },
 			{ FNiagaraTypeDefinition::GetVec2Def(), FNiagaraTypeDefinition::GetHalfVec2Def() },
 			{ FNiagaraTypeDefinition::GetVec3Def(), FNiagaraTypeDefinition::GetHalfVec3Def() },
-			{ FNiagaraTypeDefinition::GetPositionDef(), FNiagaraTypeDefinition::GetHalfVec3Def() },
 			{ FNiagaraTypeDefinition::GetVec4Def(), FNiagaraTypeDefinition::GetHalfVec4Def() },
 			{ FNiagaraTypeDefinition::GetColorDef(), FNiagaraTypeDefinition::GetHalfVec4Def() },
 			{ FNiagaraTypeDefinition::GetQuatDef(), FNiagaraTypeDefinition::GetHalfVec4Def() },
@@ -1433,7 +1432,6 @@ static void ConvertFloatToHalf(const FNiagaraCompileOptions& InCompileOptions, T
 
 		TArray<FNiagaraVariable> ConvertExceptions =
 		{
-			SYS_PARAM_ENGINE_POSITION,
 			SYS_PARAM_ENGINE_INV_DELTA_TIME,
 			SYS_PARAM_ENGINE_TIME,
 			SYS_PARAM_ENGINE_REAL_TIME,
@@ -1447,14 +1445,12 @@ static void ConvertFloatToHalf(const FNiagaraCompileOptions& InCompileOptions, T
 			SYS_PARAM_ENGINE_SYSTEM_NUM_EMITTERS,
 			SYS_PARAM_ENGINE_NUM_SYSTEM_INSTANCES,
 			SYS_PARAM_ENGINE_EMITTER_NUM_PARTICLES,
-			SYS_PARAM_ENGINE_EMITTER_SIMULATION_POSITION,
 			SYS_PARAM_ENGINE_EMITTER_TOTAL_SPAWNED_PARTICLES,
 			SYS_PARAM_PARTICLES_UNIQUE_ID,
 			SYS_PARAM_PARTICLES_ID,
 			SYS_PARAM_EMITTER_AGE,
 			SYS_PARAM_EMITTER_RANDOM_SEED,
 			SYS_PARAM_ENGINE_EMITTER_INSTANCE_SEED,
-			SYS_PARAM_PARTICLES_POSITION,
 			SYS_PARAM_PARTICLES_LIFETIME,
 		};
 
diff --git a/Engine/Source/Runtime/Core/Private/UObject/FortniteMainBranchObjectVersion.cpp b/Engine/Source/Runtime/Core/Private/UObject/FortniteMainBranchObjectVersion.cpp
--- a/Engine/Source/Runtime/Core/Private/UObject/FortniteMainBranchObjectVersion.cpp
+++ b/Engine/Source/Runtime/Core/Private/UObject/FortniteMainBranchObjectVersion.cpp
@@ -11,7 +11,7 @@ TMap<FGuid, FGuid> FFortniteMainBranchObjectVersion::GetSystemGuids()
 	SystemGuids.Add(DevGuids.MATERIALSHADERMAP_DERIVEDDATA_VER, FGuid("04B8C817071A453C9EBEFBFB51D5F785"));
 	SystemGuids.Add(DevGuids.NANITE_DERIVEDDATA_VER, FGuid("42C46B24E1586986880AFE62B0F575AB"));
 	SystemGuids.Add(DevGuids.NIAGARASHADERMAP_DERIVEDDATA_VER, FGuid("6360A977062842A29ED30E3A7ACB0E64"));
-	SystemGuids.Add(DevGuids.Niagara_LatestScriptCompileVersion, FGuid("041F479DAAB746DAA556A4C1ABB07E04"));
+	SystemGuids.Add(DevGuids.Niagara_LatestScriptCompileVersion, FGuid("F83165088DBF08469F4E45D0D1E806E3"));
 	SystemGuids.Add(DevGuids.SkeletalMeshDerivedDataVersion, FGuid("716424BCEC314E969ACFE36D2C2ECE49"));
 	SystemGuids.Add(DevGuids.STATICMESH_DERIVEDDATA_VER, FGuid("3ABF32149F9D4D838750368AB95573BE"));
 	SystemGuids.Add(DevGuids.MaterialTranslationDDCVersion, FGuid("84CF0B79D69D4A4BBE8CAFD4B3D48AB7"));
