- Fix a crash when attempting to fix up outers if the parameter store data is out of sync

diff --git a/Engine/Plugins/FX/Niagara/Source/Niagara/Private/NiagaraComponent.cpp b/Engine/Plugins/FX/Niagara/Source/Niagara/Private/NiagaraComponent.cpp
--- a/Engine/Plugins/FX/Niagara/Source/Niagara/Private/NiagaraComponent.cpp
+++ b/Engine/Plugins/FX/Niagara/Source/Niagara/Private/NiagaraComponent.cpp
@@ -3471,10 +3471,11 @@ void UNiagaraComponent::FixDataInterfaceOuters()
 	FixParameterOverrides(InstanceParameterOverrides, VariableToFixedDataInterfaces);
 
 	// Update data interfaces in use in the override parameters using the replacement map.
+	// Note: the IndexOf for the data interface can return an index OOB, this will be fixed during SanityCheckData in PostLoad so we must guard against reading OOB here
 	for (TPair<FNiagaraVariableBase, UNiagaraDataInterface*>& VariableFixedDataInterfacePair : VariableToFixedDataInterfaces)
 	{
 		int32 OverrideParametersIndex = OverrideParameters.IndexOf(VariableFixedDataInterfacePair.Key);
-		if (OverrideParametersIndex != INDEX_NONE)
+		if (OverrideParameters.GetDataInterfaces().IsValidIndex(OverrideParametersIndex))
 		{
 			OverrideParameters.SetDataInterface(VariableFixedDataInterfacePair.Value, OverrideParametersIndex);
 		}
