Set default values for struct members

#rnx

diff --git a/Engine/Plugins/FX/Niagara/Source/Niagara/Public/NiagaraDataChannelPublic.h b/Engine/Plugins/FX/Niagara/Source/Niagara/Public/NiagaraDataChannelPublic.h
--- a/Engine/Plugins/FX/Niagara/Source/Niagara/Public/NiagaraDataChannelPublic.h
+++ b/Engine/Plugins/FX/Niagara/Source/Niagara/Public/NiagaraDataChannelPublic.h
@@ -27,13 +27,13 @@ struct NIAGARA_API FNiagaraDataChannelUpdateContext
 	TObjectPtr<class UNiagaraDataChannelReader> Reader;
 
 	UPROPERTY(BlueprintReadOnly, Category = "Data Channel")
-	int32 FirstNewDataIndex;
+	int32 FirstNewDataIndex = 0;
 
 	UPROPERTY(BlueprintReadOnly, Category = "Data Channel")
-	int32 LastNewDataIndex;
+	int32 LastNewDataIndex = 0;
 
 	UPROPERTY(BlueprintReadOnly, Category = "Data Channel")
-	int32 NewElementCount;
+	int32 NewElementCount = 0;
 };
 
 DECLARE_DYNAMIC_DELEGATE_OneParam(FOnNewNiagaraDataChannelPublish, const FNiagaraDataChannelUpdateContext&, NewDataContext);
