[AutoRTFM] Allow Niagara arrays to be serialized from within Verse.

diff --git a/Engine/Plugins/FX/Niagara/Source/Niagara/Classes/NiagaraDataInterfaceArrayImpl.h b/Engine/Plugins/FX/Niagara/Source/Niagara/Classes/NiagaraDataInterfaceArrayImpl.h
--- a/Engine/Plugins/FX/Niagara/Source/Niagara/Classes/NiagaraDataInterfaceArrayImpl.h
+++ b/Engine/Plugins/FX/Niagara/Source/Niagara/Classes/NiagaraDataInterfaceArrayImpl.h
@@ -12,6 +12,7 @@
 
 #include "Async/Async.h"
 #include "Containers/ArrayView.h"
+#include "Misc/TransactionallySafeRWLock.h"
 #include "RenderGraphUtils.h"
 #include "ShaderParameterUtils.h"
 #include "ShaderCompilerCore.h"
@@ -211,11 +212,11 @@ struct FNiagaraDataInterfaceArrayImplInternal
 template<typename TArrayType>
 struct FNDIArrayInstanceData_GameThread
 {
-	FNiagaraSystemInstance*	OwnerInstance = nullptr;
-	bool					bIsModified = false;		// True if the array has ever been modified and we are reading instance data
-	bool					bIsRenderDirty = true;		// True if we have made modifications that could be pushed to the render thread
-	mutable FRWLock			ArrayRWGuard;
-	TArray<TArrayType>		ArrayData;					// Modified array data
+	FNiagaraSystemInstance*				OwnerInstance = nullptr;
+	bool								bIsModified = false;		// True if the array has ever been modified and we are reading instance data
+	bool								bIsRenderDirty = true;		// True if we have made modifications that could be pushed to the render thread
+	mutable FTransactionallySafeRWLock	ArrayRWGuard;
+	TArray<TArrayType>					ArrayData;					// Modified array data
 };
 
 struct FNDIArrayInstanceData_RenderThreadBase
@@ -314,7 +315,7 @@ struct FNDIArrayProxyImpl : public INDIArrayProxyBase
 		const TArray<TArrayType>& GetArray() { return *ArrayData; }
 
 	private:
-		FRWLock*					LockObject = nullptr;
+		FTransactionallySafeRWLock* LockObject = nullptr;
 		const TArray<TArrayType>*	ArrayData = nullptr;
 		UE_NONCOPYABLE(FReadArrayRef);
 	};
@@ -352,8 +353,8 @@ struct FNDIArrayProxyImpl : public INDIArrayProxyBase
 		TArray<TArrayType>& GetArray() { return *ArrayData; }
 
 	private:
-		FRWLock*			LockObject = nullptr;
-		TArray<TArrayType>*	ArrayData = nullptr;
+		FTransactionallySafeRWLock*	LockObject = nullptr;
+		TArray<TArrayType>*	        ArrayData = nullptr;
 		UE_NONCOPYABLE(FWriteArrayRef);
 	};
 
