[Niagara] Added PostEditChange() after SetParentEditorOnly call for UMaterialInstanceConstant

#rnx

diff --git a/Engine/Plugins/FX/Niagara/Source/Niagara/Private/NiagaraRendererProperties.cpp b/Engine/Plugins/FX/Niagara/Source/Niagara/Private/NiagaraRendererProperties.cpp
--- a/Engine/Plugins/FX/Niagara/Source/Niagara/Private/NiagaraRendererProperties.cpp
+++ b/Engine/Plugins/FX/Niagara/Source/Niagara/Private/NiagaraRendererProperties.cpp
@@ -554,15 +554,18 @@ void UNiagaraRendererProperties::UpdateMaterialParametersMIC(const FNiagaraRende
 		if (InOutMIC == nullptr)
 		{
 			InOutMIC = NewObject<UMaterialInstanceConstant>(this, FName(NameBuilder));
+			InOutMIC->PreEditChange(nullptr);
 			InOutMIC->SetParentEditorOnly(InOutMaterial);
 		}
 		else if (InOutMIC->Parent != InOutMaterial)
 		{
 			InOutMIC->Rename(NameBuilder.ToString());
+			InOutMIC->PreEditChange(nullptr);
 			InOutMIC->SetParentEditorOnly(InOutMaterial);
 		}
 
 		InOutMIC->UpdateStaticPermutation(MaterialParameterSet);
+		InOutMIC->PostEditChange();
 	}
 	else if (InOutMIC)
 	{
@@ -624,6 +627,8 @@ void UNiagaraRendererProperties::UpdateMaterialParametersMIC(const FNiagaraRende
 					bNeedsRename = true;
 					MIC = MICPool.Pop();
 				}
+				
+				MIC->PreEditChange(nullptr);
 				if (MIC->Parent != Material)
 				{
 					MIC->SetParentEditorOnly(Material);
@@ -632,10 +637,12 @@ void UNiagaraRendererProperties::UpdateMaterialParametersMIC(const FNiagaraRende
 			else
 			{
 				MIC = NewObject<UMaterialInstanceConstant>(this, FName(NameBuilder));
+				MIC->PreEditChange(nullptr);
 				MIC->SetParentEditorOnly(Material);
 			}
 
 			MIC->UpdateStaticPermutation(MaterialParameterSet);
+			MIC->PostEditChange();
 
 			InOutMICs.SetNum(i + 1);
 			InOutMICs[i] = MIC;
