- Move renderer defaults to project settings

#okforightub public
#rnx

diff --git a/Engine/Plugins/FX/Niagara/Config/BaseNiagara.ini b/Engine/Plugins/FX/Niagara/Config/BaseNiagara.ini
--- a/Engine/Plugins/FX/Niagara/Config/BaseNiagara.ini
+++ b/Engine/Plugins/FX/Niagara/Config/BaseNiagara.ini
@@ -81,6 +81,10 @@ DefaultDynamicInputScript=/Niagara/DefaultAssets/DefaultDynamicInputs.DefaultDyn
 DefaultFunctionScript=/Niagara/DefaultAssets/DefaultFunction.DefaultFunction
 DefaultModuleScript=/Niagara/DefaultAssets/DefaultModule.DefaultModule
 DefaultEmptyEmitter=/Niagara/DefaultAssets/Templates/Emitters/Minimal.Minimal
+DefaultMeshRendererMesh=/Niagara/DefaultAssets/S_Gnomon.S_Gnomon
+DefaultDecalRendererMaterial=/Niagara/DefaultAssets/DefaultDecalMaterial.DefaultDecalMaterial
+DefaultRibbonRendererMaterial=/Niagara/DefaultAssets/DefaultRibbonMaterial.DefaultRibbonMaterial
+DefaultSpriteRendererMaterial=/Niagara/DefaultAssets/DefaultSpriteMaterial.DefaultSpriteMaterial
 RequiredSystemUpdateScript=/Niagara/Modules/System/SystemState.SystemState
 +GraphCreationShortcuts=(Name="Numeric::Add",Input=(Key=A,bShift=False,bCtrl=False,bAlt=False,bCmd=False))
 +GraphCreationShortcuts=(Name="Numeric::Divide",Input=(Key=D,bShift=False,bCtrl=False,bAlt=False,bCmd=False))
diff --git a/Engine/Plugins/FX/Niagara/Source/NiagaraEditor/Private/NiagaraEditorModule.cpp b/Engine/Plugins/FX/Niagara/Source/NiagaraEditor/Private/NiagaraEditorModule.cpp
--- a/Engine/Plugins/FX/Niagara/Source/NiagaraEditor/Private/NiagaraEditorModule.cpp
+++ b/Engine/Plugins/FX/Niagara/Source/NiagaraEditor/Private/NiagaraEditorModule.cpp
@@ -2173,11 +2173,12 @@ void FNiagaraEditorModule::RegisterDefaultRendererFactories()
 		FNiagaraRendererCreationInfo::FRendererFactory::CreateLambda([](UObject* OuterEmitter)
 		{
 			UNiagaraMeshRendererProperties* NewRenderer = NewObject<UNiagaraMeshRendererProperties>(OuterEmitter, NAME_None, RF_Transactional);
+
 			// we have an empty entry in the constructor. Due to CDO default value propagation being unwanted, we have to keep it in there
 			if(ensure(NewRenderer->Meshes.Num() == 1))
 			{
-				FSoftObjectPath DefaultMesh(TEXT("StaticMesh'/Niagara/DefaultAssets/S_Gnomon.S_Gnomon'"));
-				NewRenderer->Meshes[0].Mesh = Cast<UStaticMesh>(DefaultMesh.TryLoad());
+				const UNiagaraEditorSettings* NiagaraEditorSettings = GetDefault<UNiagaraEditorSettings>();
+				NewRenderer->Meshes[0].Mesh = Cast<UStaticMesh>(NiagaraEditorSettings->DefaultMeshRendererMesh.TryLoad());
 			}
 			return NewRenderer;
 		})));
@@ -2189,8 +2190,8 @@ void FNiagaraEditorModule::RegisterDefaultRendererFactories()
 		FNiagaraRendererCreationInfo::FRendererFactory::CreateLambda([](UObject* OuterEmitter)
 		{
 			UNiagaraSpriteRendererProperties* NewRenderer = NewObject<UNiagaraSpriteRendererProperties>(OuterEmitter, NAME_None, RF_Transactional);
-			FSoftObjectPath DefaultMaterial(TEXT("Material'/Niagara/DefaultAssets/DefaultSpriteMaterial.DefaultSpriteMaterial'"));
-			NewRenderer->Material = Cast<UMaterialInterface>(DefaultMaterial.TryLoad());
+			const UNiagaraEditorSettings* NiagaraEditorSettings = GetDefault<UNiagaraEditorSettings>();
+			NewRenderer->Material = Cast<UMaterialInterface>(NiagaraEditorSettings->DefaultSpriteRendererMaterial.TryLoad());
 			return NewRenderer;
 		})));
 	
@@ -2201,8 +2202,8 @@ void FNiagaraEditorModule::RegisterDefaultRendererFactories()
 		FNiagaraRendererCreationInfo::FRendererFactory::CreateLambda([](UObject* OuterEmitter)
 		{
 			UNiagaraRibbonRendererProperties* NewRenderer = NewObject<UNiagaraRibbonRendererProperties>(OuterEmitter, NAME_None, RF_Transactional);
-			FSoftObjectPath DefaultMaterial(TEXT("Material'/Niagara/DefaultAssets/DefaultRibbonMaterial.DefaultRibbonMaterial'"));
-			NewRenderer->Material = Cast<UMaterialInterface>(DefaultMaterial.TryLoad());
+			const UNiagaraEditorSettings* NiagaraEditorSettings = GetDefault<UNiagaraEditorSettings>();
+			NewRenderer->Material = Cast<UMaterialInterface>(NiagaraEditorSettings->DefaultRibbonRendererMaterial.TryLoad());
 			return NewRenderer;
 		})));
 
@@ -2233,8 +2234,8 @@ void FNiagaraEditorModule::RegisterDefaultRendererFactories()
 		FNiagaraRendererCreationInfo::FRendererFactory::CreateLambda([](UObject* OuterEmitter)
 		{
 			UNiagaraDecalRendererProperties* NewRenderer = NewObject<UNiagaraDecalRendererProperties>(OuterEmitter, NAME_None, RF_Transactional);
-			FSoftObjectPath DefaultMaterial(TEXT("Material'/Niagara/DefaultAssets/DefaultDecalMaterial.DefaultDecalMaterial'"));
-			NewRenderer->Material = Cast<UMaterialInterface>(DefaultMaterial.TryLoad());
+			const UNiagaraEditorSettings* NiagaraEditorSettings = GetDefault<UNiagaraEditorSettings>();
+			NewRenderer->Material = Cast<UMaterialInterface>(NiagaraEditorSettings->DefaultDecalRendererMaterial.TryLoad());
 			return NewRenderer;
 		})));
 
@@ -2245,8 +2246,6 @@ void FNiagaraEditorModule::RegisterDefaultRendererFactories()
 		FNiagaraRendererCreationInfo::FRendererFactory::CreateLambda([](UObject* OuterEmitter)
 		{
 			UNiagaraVolumeRendererProperties* NewRenderer = NewObject<UNiagaraVolumeRendererProperties>(OuterEmitter, NAME_None, RF_Transactional);
-			//FSoftObjectPath DefaultMaterial(TEXT("Material'/Niagara/DefaultAssets/DefaultDecalMaterial.DefaultDecalMaterial'"));
-			//NewRenderer->Material = Cast<UMaterialInterface>(DefaultMaterial.TryLoad());
 			return NewRenderer;
 		})));
 }
diff --git a/Engine/Plugins/FX/Niagara/Source/NiagaraEditor/Public/NiagaraEditorSettings.h b/Engine/Plugins/FX/Niagara/Source/NiagaraEditor/Public/NiagaraEditorSettings.h
--- a/Engine/Plugins/FX/Niagara/Source/NiagaraEditor/Public/NiagaraEditorSettings.h
+++ b/Engine/Plugins/FX/Niagara/Source/NiagaraEditor/Public/NiagaraEditorSettings.h
@@ -349,6 +349,22 @@ public:
 	UPROPERTY(config, EditAnywhere, DisplayName="Minimal Emitter", Category = Niagara, meta=(AllowedClasses="/Script/Niagara.NiagaraEmitter"))
 	FSoftObjectPath DefaultEmptyEmitter;
 
+	/** Default material to use with new decal renderers. */
+	UPROPERTY(config, EditAnywhere, Category = "Niagara Renderer", meta = (AllowedClasses = "/Script/Engine.MaterialInterface"))
+	FSoftObjectPath DefaultDecalRendererMaterial;
+
+	/** Default mesh to use with new mesh renderers. */
+	UPROPERTY(config, EditAnywhere, Category = "Niagara Renderer", meta = (AllowedClasses = "/Script/Engine.StaticMesh"))
+	FSoftObjectPath DefaultMeshRendererMesh;
+
+	/** Default material to use with new ribbon renderers. */
+	UPROPERTY(config, EditAnywhere, Category = "Niagara Renderer", meta = (AllowedClasses = "/Script/Engine.MaterialInterface"))
+	FSoftObjectPath DefaultRibbonRendererMaterial;
+
+	/** Default material to use with new sprite renderers. */
+	UPROPERTY(config, EditAnywhere, Category = "Niagara Renderer", meta = (AllowedClasses = "/Script/Engine.MaterialInterface"))
+	FSoftObjectPath DefaultSpriteRendererMaterial;
+
 	/** Niagara script which is required in the system update script to control system state. */
 	UPROPERTY(config, EditAnywhere, Category = Niagara)
 	FSoftObjectPath RequiredSystemUpdateScript;
