Added StackObjectCustomization & implemented it for NiagaraStatelessModule_DynamicMaterialParameters

diff --git a/Engine/Plugins/FX/Niagara/Source/NiagaraEditor/Private/Customizations/Stack/NiagaraStackObjectPropertyCustomization.cpp b/Engine/Plugins/FX/Niagara/Source/NiagaraEditor/Private/Customizations/Stack/NiagaraStackObjectPropertyCustomization.cpp
new file mode 100644
--- /dev/null
+++ b/Engine/Plugins/FX/Niagara/Source/NiagaraEditor/Private/Customizations/Stack/NiagaraStackObjectPropertyCustomization.cpp
@@ -0,0 +1,3 @@
+﻿// Copyright Epic Games, Inc. All Rights Reserved.
+
+#include "Customizations/NiagaraStackObjectPropertyCustomization.h"
diff --git a/Engine/Plugins/FX/Niagara/Source/NiagaraEditor/Private/Customizations/Stack/NiagaraStackObjectPropertyCustomization_StatelessModule_DynamicMaterialParameters.cpp b/Engine/Plugins/FX/Niagara/Source/NiagaraEditor/Private/Customizations/Stack/NiagaraStackObjectPropertyCustomization_StatelessModule_DynamicMaterialParameters.cpp
new file mode 100644
--- /dev/null
+++ b/Engine/Plugins/FX/Niagara/Source/NiagaraEditor/Private/Customizations/Stack/NiagaraStackObjectPropertyCustomization_StatelessModule_DynamicMaterialParameters.cpp
@@ -0,0 +1,260 @@
+﻿// Copyright Epic Games, Inc. All Rights Reserved.
+
+#include "NiagaraStackObjectPropertyCustomization_StatelessModule_DynamicMaterialParameters.h"
+#include "Widgets/Text/STextBlock.h"
+#include "ViewModels/Stack/NiagaraStackPropertyRow.h"
+#include "DetailTreeNode.h"
+#include "Containers/StaticBitArray.h"
+#include "Materials/MaterialExpressionDynamicParameter.h"
+#include "Stateless/NiagaraStatelessEmitter.h"
+#include "ViewModels/NiagaraEmitterHandleViewModel.h"
+#include "ViewModels/NiagaraEmitterViewModel.h"
+#include "ViewModels/NiagaraSystemViewModel.h"
+#include "Materials/Material.h"
+
+FNiagaraStackObjectPropertyCustomization_StatelessModule_DynamicMaterialParameters::
+FNiagaraStackObjectPropertyCustomization_StatelessModule_DynamicMaterialParameters()
+{
+	ParameterIndexMap.Add(FName("Parameter0"), 0);
+	ParameterIndexMap.Add(FName("Parameter1"), 1);
+	ParameterIndexMap.Add(FName("Parameter2"), 2);
+	ParameterIndexMap.Add(FName("Parameter3"), 3);
+	ParameterChannelMap.Add(FName("XChannelDistribution"), 0);
+	ParameterChannelMap.Add(FName("YChannelDistribution"), 1);
+	ParameterChannelMap.Add(FName("ZChannelDistribution"), 2);
+	ParameterChannelMap.Add(FName("WChannelDistribution"), 3);
+}
+
+TSharedRef<FNiagaraStackObjectPropertyCustomization> FNiagaraStackObjectPropertyCustomization_StatelessModule_DynamicMaterialParameters::MakeInstance()
+{
+	return MakeShared<FNiagaraStackObjectPropertyCustomization_StatelessModule_DynamicMaterialParameters>();
+}
+
+TOptional<TSharedPtr<SWidget>> FNiagaraStackObjectPropertyCustomization_StatelessModule_DynamicMaterialParameters::GenerateNameWidget(UNiagaraStackPropertyRow* StackPropertyRow) const
+{
+	TSharedRef<FDetailTreeNode> DetailTreeNode = StaticCastSharedRef<FDetailTreeNode>(StackPropertyRow->GetDetailTreeNode());
+
+	int32 ParameterIndex = INDEX_NONE;
+	int32 ParameterChannel = INDEX_NONE;
+	GetParameterIndexAndChannel(DetailTreeNode, ParameterIndex, ParameterChannel);
+
+	if(ParameterIndex == INDEX_NONE || ParameterChannel == INDEX_NONE)
+	{
+		return TOptional<TSharedPtr<SWidget>>();	
+	}
+	
+	TWeakPtr<FDetailTreeNode> ParentNode = DetailTreeNode->GetParentNode();
+	TSharedPtr<IDetailPropertyRow> PropertyRow = DetailTreeNode->GetRow();
+
+	TSharedPtr<FNiagaraSystemViewModel> SystemViewModel = StackPropertyRow->GetSystemViewModel();
+	TSharedPtr<FNiagaraEmitterHandleViewModel> EmitterHandleViewModel = SystemViewModel->GetEmitterHandleViewModelForEmitter(StackPropertyRow->GetEmitterViewModel()->GetEmitter());
+
+	TOptional<FText> DisplayName = TryGetDisplayNameForDynamicMaterialParameter(EmitterHandleViewModel, ParameterIndex, ParameterChannel);
+
+	if(DisplayName.IsSet())
+	{
+		TSharedPtr<SWidget> Widget = SNew(STextBlock).Text(DisplayName.GetValue());
+		return Widget;
+	}
+	
+	return TOptional<TSharedPtr<SWidget>>();
+}
+
+TOptional<FText> FNiagaraStackObjectPropertyCustomization_StatelessModule_DynamicMaterialParameters::TryGetDisplayNameForDynamicMaterialParameter(TSharedPtr<FNiagaraEmitterHandleViewModel> EmitterHandleViewModel, int32 ParameterIndex, int32 ParameterChannel) const
+{
+	UNiagaraStatelessEmitter* OwningStatelessEmitter = EmitterHandleViewModel->GetEmitterHandle()->GetStatelessEmitter();
+
+	TArray<UMaterial*> Materials;
+	FNiagaraEmitterInstance* ThisEmitterInstance = EmitterHandleViewModel->GetEmitterViewModel()->GetSimulation().Pin().Get();
+	Materials = GetMaterialsFromEmitter(*OwningStatelessEmitter, ThisEmitterInstance);
+	
+	TMap<UMaterialExpressionDynamicParameter*, TStaticBitArray<4>> DynamicParameterExpressionToOutputMaskMap;
+	TArray<FExpressionInput*> ExpressionInputsToProcess;
+	
+	// Visit each material and gather all expression inputs for each expression.
+	for (UMaterial* Material : Materials)
+	{
+		if (Material == nullptr)
+		{
+			continue;
+		}
+
+		for (int32 MaterialPropertyIndex = 0; MaterialPropertyIndex < MP_MAX; ++MaterialPropertyIndex)
+		{
+			FExpressionInput* ExpressionInput = Material->GetExpressionInputForProperty(EMaterialProperty(MaterialPropertyIndex));
+			if (ExpressionInput != nullptr)
+			{
+				ExpressionInputsToProcess.Add(ExpressionInput);
+			}
+		}
+		
+		TArray<UMaterialExpression*> Expressions;
+		Material->GetAllExpressionsInMaterialAndFunctionsOfType<UMaterialExpression>(Expressions);
+		for (UMaterialExpression* Expression : Expressions)
+		{
+			for (FExpressionInputIterator It{ Expression}; It; ++It)
+			{
+				ExpressionInputsToProcess.Push(It.Input);
+			}
+		}
+	}
+	
+	// Then we check these expression inputs to see dynamic parameter node/expression to check whether its corresponding parameter channel is used or not
+	bool bAnyDynamicParametersFound = false;
+	for (FExpressionInput* ExpressionInput : ExpressionInputsToProcess)
+	{
+		if (ExpressionInput->Expression == nullptr)
+		{
+			continue;
+		}
+	
+		UMaterialExpressionDynamicParameter* DynamicParameterExpression = Cast<UMaterialExpressionDynamicParameter>(ExpressionInput->Expression);
+		if (DynamicParameterExpression == nullptr)
+		{
+			continue;
+		}
+	
+		bAnyDynamicParametersFound = true;
+		// Each bit tells us whether the corresponding channel for the parameter is used or not
+		GetChannelUsedBitMask(ExpressionInput, DynamicParameterExpressionToOutputMaskMap.FindOrAdd(DynamicParameterExpression));
+	}
+	
+	FName ParameterName = NAME_None;
+	bool bMultipleNamesFound = false;
+	TOptional<bool> bIsParameterUsed;
+	for (auto It = DynamicParameterExpressionToOutputMaskMap.CreateConstIterator(); It; ++It)
+	{
+		UMaterialExpressionDynamicParameter* ExpressionDynamicParameter = It.Key();
+		const TStaticBitArray<4>& ExpressionOutputMask = It.Value();
+		// We ignore those with a different index, and those where the channel isn't used
+		if (ExpressionDynamicParameter->ParameterIndex != ParameterIndex)
+		{
+			continue;
+		}
+		
+		const FExpressionOutput& Output = ExpressionDynamicParameter->GetOutputs()[ParameterChannel];
+		if (ParameterName == NAME_None)
+		{
+			// We found a name but we continue iterating to check for duplicates
+			ParameterName = Output.OutputName;
+			
+			if(bIsParameterUsed.IsSet() == false)
+			{
+				bIsParameterUsed = ExpressionOutputMask[ParameterChannel];
+			}
+		}
+		// If we already found a name, but now find one with the same index and channel but a different name, we report that to the user
+		else if (ParameterName != Output.OutputName)
+		{
+			bMultipleNamesFound = true;
+		}
+	}
+	
+	// Return the final dynamic param UI name.
+	if (bAnyDynamicParametersFound == false)
+	{
+		return FText::FromString("(No material found using dynamic params)");
+	}
+	else if (ParameterName != NAME_None)
+	{
+		if (bMultipleNamesFound == false && bIsParameterUsed.IsSet())
+		{
+			FText ParameterNameAsText = FText::FromName(ParameterName);
+			if(bIsParameterUsed.GetValue() == true)
+			{
+				return ParameterNameAsText;
+			}
+			else
+			{
+				return FText::FormatOrdered(FText::FromString("{0} - {1}"), ParameterNameAsText, FText::FromString("Unused in Material"));
+			}
+		}
+		else
+		{
+			return FText::FromString(ParameterName.ToString() + TEXT(" (Multiple Aliases Found)"));
+		}
+	}
+
+	return TOptional<FText>();
+}
+
+void FNiagaraStackObjectPropertyCustomization_StatelessModule_DynamicMaterialParameters::GetChannelUsedBitMask(FExpressionInput* Input, TStaticBitArray<4>& ChannelUsedMask) const
+{
+	if (Input->Expression)
+	{
+		TArray<FExpressionOutput>& Outputs = Input->Expression->GetOutputs();
+
+		if (Outputs.Num() > 0)
+		{
+			const bool bOutputIndexIsValid = Outputs.IsValidIndex(Input->OutputIndex)
+				// Attempt to handle legacy connections before OutputIndex was used that had a mask
+				&& (Input->OutputIndex != 0 || Input->Mask == 0);
+
+			for (int32 OutputIndex=0; OutputIndex < Outputs.Num(); ++OutputIndex)
+			{
+				const FExpressionOutput& Output = Outputs[OutputIndex];
+
+				if ((bOutputIndexIsValid && OutputIndex == Input->OutputIndex)
+					|| (!bOutputIndexIsValid
+						&& Output.Mask == Input->Mask
+						&& Output.MaskR == Input->MaskR
+						&& Output.MaskG == Input->MaskG
+						&& Output.MaskB == Input->MaskB
+						&& Output.MaskA == Input->MaskA))
+				{
+					ChannelUsedMask[0] = ChannelUsedMask[0] || (Input->MaskR != 0);
+					ChannelUsedMask[1] = ChannelUsedMask[1] || (Input->MaskG != 0);
+					ChannelUsedMask[2] = ChannelUsedMask[2] || (Input->MaskB != 0);
+					ChannelUsedMask[3] = ChannelUsedMask[3] || (Input->MaskA != 0);
+					return;
+				}
+			}
+		}
+	}
+}
+
+TArray<UMaterial*> FNiagaraStackObjectPropertyCustomization_StatelessModule_DynamicMaterialParameters::GetMaterialsFromEmitter(
+	const UNiagaraStatelessEmitter& InEmitter, const FNiagaraEmitterInstance* InEmitterInstance) const
+{
+	TArray<UMaterial*> ResultMaterials;
+	for (UNiagaraRendererProperties* RenderProperties : InEmitter.GetRenderers())
+	{
+		TArray<UMaterialInterface*> UsedMaterialInteraces;
+		RenderProperties->GetUsedMaterials(InEmitterInstance, UsedMaterialInteraces);
+		for (UMaterialInterface* UsedMaterialInterface : UsedMaterialInteraces)
+		{
+			if (UsedMaterialInterface != nullptr)
+			{
+				UMaterial* UsedMaterial = UsedMaterialInterface->GetBaseMaterial();
+				if (UsedMaterial != nullptr)
+				{
+					ResultMaterials.AddUnique(UsedMaterial);
+					break;
+				}
+			}
+		}
+	}
+
+	return ResultMaterials;
+}
+
+void FNiagaraStackObjectPropertyCustomization_StatelessModule_DynamicMaterialParameters::GetParameterIndexAndChannel(TSharedRef<FDetailTreeNode> DetailTreeNode, int32& OutParameterIndex, int32& OutParameterChannel) const
+{
+	OutParameterIndex = INDEX_NONE;
+	OutParameterChannel = INDEX_NONE;
+
+	FName NodeName = DetailTreeNode->GetNodeName();
+	// First we check against the channel map, the individual values within one 'Parameter'
+	if(ParameterChannelMap.Contains(NodeName))
+	{
+		OutParameterChannel = ParameterChannelMap[NodeName];
+
+		TSharedPtr<FDetailTreeNode> ParentTreeNode = DetailTreeNode->GetParentNode().Pin();
+		FName ParentNodeName = ParentTreeNode->GetNodeName();
+		// If we found one, we can safely assume its parent is valid, which is the entire Parameter
+		if(ensure(ParameterIndexMap.Contains(ParentNodeName)))
+		{
+			OutParameterIndex = ParameterIndexMap[ParentNodeName];
+		}
+	}
+}
diff --git a/Engine/Plugins/FX/Niagara/Source/NiagaraEditor/Private/Customizations/Stack/NiagaraStackObjectPropertyCustomization_StatelessModule_DynamicMaterialParameters.h b/Engine/Plugins/FX/Niagara/Source/NiagaraEditor/Private/Customizations/Stack/NiagaraStackObjectPropertyCustomization_StatelessModule_DynamicMaterialParameters.h
new file mode 100644
--- /dev/null
+++ b/Engine/Plugins/FX/Niagara/Source/NiagaraEditor/Private/Customizations/Stack/NiagaraStackObjectPropertyCustomization_StatelessModule_DynamicMaterialParameters.h
@@ -0,0 +1,38 @@
+﻿// Copyright Epic Games, Inc. All Rights Reserved.
+
+#pragma once
+
+#include "CoreMinimal.h"
+#include "Containers/StaticBitArray.h"
+#include "Customizations/NiagaraStackObjectPropertyCustomization.h"
+
+class FNiagaraEmitterHandleViewModel;
+class FNiagaraEmitterInstance;
+class UNiagaraStatelessEmitter;
+class UNiagaraStackPropertyRow;
+class UMaterial;
+class FDetailTreeNode;
+struct FExpressionInput;
+
+
+class NIAGARAEDITOR_API FNiagaraStackObjectPropertyCustomization_StatelessModule_DynamicMaterialParameters
+	: public FNiagaraStackObjectPropertyCustomization
+{
+public:
+	FNiagaraStackObjectPropertyCustomization_StatelessModule_DynamicMaterialParameters();
+
+	static TSharedRef<FNiagaraStackObjectPropertyCustomization> MakeInstance();
+	
+	virtual TOptional<TSharedPtr<SWidget>> GenerateNameWidget(UNiagaraStackPropertyRow* PropertyRow) const override;
+
+private:
+	TOptional<FText> TryGetDisplayNameForDynamicMaterialParameter(TSharedPtr<FNiagaraEmitterHandleViewModel> EmitterHandleViewModel, int32 ParameterIndex, int32 ParameterChannel) const;
+	
+	void GetChannelUsedBitMask(FExpressionInput* Input, TStaticBitArray<4>& ChannelUsedMask) const;
+	TArray<UMaterial*> GetMaterialsFromEmitter(const UNiagaraStatelessEmitter& InEmitter, const FNiagaraEmitterInstance* InEmitterInstance) const;
+	void GetParameterIndexAndChannel(TSharedRef<FDetailTreeNode> DetailTreeNode, int32& OutParameterIndex, int32& OutParameterChannel) const;
+
+private:
+	TMap<FName, int32> ParameterIndexMap;
+	TMap<FName, int32> ParameterChannelMap;
+};
diff --git a/Engine/Plugins/FX/Niagara/Source/NiagaraEditor/Private/NiagaraEditorModule.cpp b/Engine/Plugins/FX/Niagara/Source/NiagaraEditor/Private/NiagaraEditorModule.cpp
--- a/Engine/Plugins/FX/Niagara/Source/NiagaraEditor/Private/NiagaraEditorModule.cpp
+++ b/Engine/Plugins/FX/Niagara/Source/NiagaraEditor/Private/NiagaraEditorModule.cpp
@@ -183,6 +183,8 @@
 #include "Engine/AssetManager.h"
 #include "Widgets/AssetBrowser/SNiagaraSelectedAssetDetails.h"
 #include "NiagaraRecentAndFavoritesManager.h"
+#include "Customizations/Stack/NiagaraStackObjectPropertyCustomization_StatelessModule_DynamicMaterialParameters.h"
+#include "Stateless/Modules/NiagaraStatelessModule_DynamicMaterialParameters.h"
 #include "Widgets/DataChannel/NiagaraDataChannelWizard.h"
 
 #include "TraversalCache/TraversalBuilder.h"
@@ -457,6 +459,26 @@ private:
 	TMap<FName, FCreateGraphPin> MiscSubCategoryToCreatePinDelegateMap;
 };
 
+void FNiagaraStackObjectCustomizationRegistry::RegisterStackObjectCustomization(UClass& Class, FOnGetStackObjectCustomizationInstance CustomizationFactory)
+{
+	Customizations.Add(&Class) = CustomizationFactory.Execute();
+}
+
+void FNiagaraStackObjectCustomizationRegistry::UnregisterStackObjectCustomization(UClass& Class)
+{
+	Customizations.Remove(&Class);
+}
+
+TSharedPtr<const FNiagaraStackObjectPropertyCustomization> FNiagaraStackObjectCustomizationRegistry::GetCustomizationForStackObject(UNiagaraStackObject& StackObject)
+{
+	if(Customizations.Contains(StackObject.GetObject()->GetClass()))
+	{
+		return Customizations[StackObject.GetObject()->GetClass()];
+	}
+
+	return nullptr;
+}
+
 FNiagaraEditorModule::FNiagaraEditorModule() 
 	: SequencerSettings(nullptr)
 	, TestCompileScriptCommand(nullptr)
@@ -1304,6 +1326,8 @@ void FNiagaraEditorModule::StartupModule()
 	//Register Stack Object Issue Generators.
 	RegisterStackIssueGenerator(FNiagaraPlatformSet::StaticStruct()->GetFName(), new FNiagaraPlatformSetIssueGenerator());
 
+	RegisterDefaultStackObjectCustomizations();
+	
 	NiagaraEditorMenuHelpers::RegisterToolMenus();
 	NiagaraEditorMenuHelpers::RegisterMenuExtensions();
 
@@ -2371,6 +2395,12 @@ bool FNiagaraEditorModule::DeferredDestructObjects(float InDeltaTime)
 	return false;
 }
 
+void FNiagaraEditorModule::RegisterDefaultStackObjectCustomizations()
+{
+	StackObjectCustomizationRegistry.RegisterStackObjectCustomization(*UNiagaraStatelessModule_DynamicMaterialParameters::StaticClass(),
+		FNiagaraStackObjectCustomizationRegistry::FOnGetStackObjectCustomizationInstance::CreateStatic(&FNiagaraStackObjectPropertyCustomization_StatelessModule_DynamicMaterialParameters::MakeInstance));
+}
+
 void FNiagaraEditorModule::OnAssetRegistryLoadComplete()
 {
 	FAssetRegistryModule& AssetRegistryModule = FModuleManager::LoadModuleChecked<FAssetRegistryModule>("AssetRegistry");
diff --git a/Engine/Plugins/FX/Niagara/Source/NiagaraEditor/Private/ViewModels/Stack/NiagaraStackObject.cpp b/Engine/Plugins/FX/Niagara/Source/NiagaraEditor/Private/ViewModels/Stack/NiagaraStackObject.cpp
--- a/Engine/Plugins/FX/Niagara/Source/NiagaraEditor/Private/ViewModels/Stack/NiagaraStackObject.cpp
+++ b/Engine/Plugins/FX/Niagara/Source/NiagaraEditor/Private/ViewModels/Stack/NiagaraStackObject.cpp
@@ -379,7 +379,8 @@ void UNiagaraStackObject::RefreshChildrenInternal(const TArray<UNiagaraStackEntr
 		{
 			RootTreeNodes = CollapsedRootTreeNodes;
 		}
-		
+
+		TSharedPtr<const FNiagaraStackObjectPropertyCustomization> Customization = NiagaraEditorModule->GetStackObjectCustomizationRegistry().GetCustomizationForStackObject(*this);
 		for (TSharedRef<IDetailTreeNode> RootTreeNode : RootTreeNodes)
 		{
 			if (RootTreeNode->GetNodeType() == EDetailNodeType::Advanced)
@@ -399,6 +400,10 @@ void UNiagaraStackObject::RefreshChildrenInternal(const TArray<UNiagaraStackEntr
 				{
 					ChildRow->SetOnFilterDetailNodes(OnFilterDetailNodesDelegate);
 				}
+				if(Customization)
+				{
+					ChildRow->SetPropertyCustomization(Customization);
+				}
 			}
 
 			NewChildren.Add(ChildRow);
diff --git a/Engine/Plugins/FX/Niagara/Source/NiagaraEditor/Private/ViewModels/Stack/NiagaraStackPropertyRow.cpp b/Engine/Plugins/FX/Niagara/Source/NiagaraEditor/Private/ViewModels/Stack/NiagaraStackPropertyRow.cpp
--- a/Engine/Plugins/FX/Niagara/Source/NiagaraEditor/Private/ViewModels/Stack/NiagaraStackPropertyRow.cpp
+++ b/Engine/Plugins/FX/Niagara/Source/NiagaraEditor/Private/ViewModels/Stack/NiagaraStackPropertyRow.cpp
@@ -9,6 +9,7 @@
 #include "IDetailPropertyRow.h"
 #include "ScopedTransaction.h"
 #include "DataHierarchyViewModelBase.h"
+#include "Customizations/NiagaraStackObjectPropertyCustomization.h"
 
 #include UE_INLINE_GENERATED_CPP_BY_NAME(NiagaraStackPropertyRow)
 
@@ -66,6 +67,11 @@ bool UNiagaraStackPropertyRow::GetIsEnabled() const
 	return OwningNiagaraNode == nullptr || OwningNiagaraNode->GetDesiredEnabledState() == ENodeEnabledState::Enabled;
 }
 
+void UNiagaraStackPropertyRow::SetPropertyCustomization(TSharedPtr<const FNiagaraStackObjectPropertyCustomization> Customization)
+{
+	PropertyCustomization = Customization;
+}
+
 UNiagaraStackEntry::EStackRowStyle UNiagaraStackPropertyRow::GetStackRowStyle() const
 {
 	return RowStyle;
@@ -183,6 +189,8 @@ void UNiagaraStackPropertyRow::FinalizeInternal()
 
 void UNiagaraStackPropertyRow::RefreshChildrenInternal(const TArray<UNiagaraStackEntry*>& CurrentChildren, TArray<UNiagaraStackEntry*>& NewChildren, TArray<FStackIssue>& NewIssues)
 {
+	RefreshCustomization();
+	
 	TArray<TSharedRef<IDetailTreeNode>> AllNodeChildren;
 	DetailTreeNode->GetChildren(AllNodeChildren);
 
@@ -213,6 +221,10 @@ void UNiagaraStackPropertyRow::RefreshChildrenInternal(const TArray<UNiagaraStac
 			ChildRow = NewObject<UNiagaraStackPropertyRow>(this);
 			ChildRow->Initialize(CreateDefaultChildRequiredData(), NodeChild, bChildIsTopLevelProperty, bHideTopLevelCategories, GetOwnerStackItemEditorDataKey(), GetStackEditorDataKey(), OwningNiagaraNode);
 			ChildRow->SetOwnerGuid(OwnerGuid);
+			if(PropertyCustomization.IsValid())
+			{
+				ChildRow->SetPropertyCustomization(PropertyCustomization);
+			}
 		}
 
 		NewChildren.Add(ChildRow);
@@ -231,6 +243,14 @@ void UNiagaraStackPropertyRow::RefreshChildrenInternal(const TArray<UNiagaraStac
 	}
 }
 
+void UNiagaraStackPropertyRow::RefreshCustomization()
+{
+	if(PropertyCustomization.IsValid())
+	{
+		//PropertyCustomization->GenerateNameWidget(DetailTreeNode.ToSharedRef());
+	}
+}
+
 int32 UNiagaraStackPropertyRow::GetChildIndentLevel() const
 {
 	// We want to keep inputs under a top level category at the same indent level as the category.
diff --git a/Engine/Plugins/FX/Niagara/Source/NiagaraEditor/Public/Customizations/NiagaraStackObjectPropertyCustomization.h b/Engine/Plugins/FX/Niagara/Source/NiagaraEditor/Public/Customizations/NiagaraStackObjectPropertyCustomization.h
new file mode 100644
--- /dev/null
+++ b/Engine/Plugins/FX/Niagara/Source/NiagaraEditor/Public/Customizations/NiagaraStackObjectPropertyCustomization.h
@@ -0,0 +1,22 @@
+﻿// Copyright Epic Games, Inc. All Rights Reserved.
+
+#pragma once
+
+#include "CoreMinimal.h"
+#include "UObject/Object.h"
+#include "Widgets/SWidget.h"
+
+class UNiagaraStackPropertyRow;
+
+/**
+ *  This class allows you to customize the UI for a UNiagaraStackObject's properties.
+ *  Register this customization in the NiagaraEditorModule by providing the class it is customizing and your customization class.
+ */
+class NIAGARAEDITOR_API FNiagaraStackObjectPropertyCustomization : public TSharedFromThis<FNiagaraStackObjectPropertyCustomization>
+{
+public:
+	virtual ~FNiagaraStackObjectPropertyCustomization() = default;
+	
+	/** If specified, will generate a custom name widget for a given property row. Only applicable for rows that generate typical name & value widgets. */
+	virtual TOptional<TSharedPtr<SWidget>> GenerateNameWidget(UNiagaraStackPropertyRow* PropertyRow) const { return TOptional<TSharedPtr<SWidget>>(); }
+};
diff --git a/Engine/Plugins/FX/Niagara/Source/NiagaraEditor/Public/NiagaraEditorModule.h b/Engine/Plugins/FX/Niagara/Source/NiagaraEditor/Public/NiagaraEditorModule.h
--- a/Engine/Plugins/FX/Niagara/Source/NiagaraEditor/Public/NiagaraEditorModule.h
+++ b/Engine/Plugins/FX/Niagara/Source/NiagaraEditor/Public/NiagaraEditorModule.h
@@ -14,8 +14,11 @@
 #include "NiagaraDebuggerCommon.h"
 #include "NiagaraRendererProperties.h"
 #include "Customizations/NiagaraDataInterfaceSimCacheVisualizer.h"
+#include "Customizations/NiagaraStackObjectPropertyCustomization.h"
 #include "NiagaraEditorModule.generated.h"
 
+class UNiagaraStackObject;
+
 namespace UE::Niagara::Wizard
 {
 	class FModuleWizardGenerator;
@@ -143,6 +146,18 @@ struct FNiagaraRendererCreationInfo
 	FRendererFactory RendererFactory;
 };
 
+struct FNiagaraStackObjectCustomizationRegistry
+{
+	DECLARE_DELEGATE_RetVal(TSharedRef<FNiagaraStackObjectPropertyCustomization>, FOnGetStackObjectCustomizationInstance);
+public:
+	void RegisterStackObjectCustomization(UClass& Class, FOnGetStackObjectCustomizationInstance Customization);
+	void UnregisterStackObjectCustomization(UClass& Class);
+
+	TSharedPtr<const FNiagaraStackObjectPropertyCustomization> GetCustomizationForStackObject(UNiagaraStackObject& StackObject);
+private:
+	TMap<UClass*, TSharedPtr<FNiagaraStackObjectPropertyCustomization>> Customizations;
+};
+
 FORCEINLINE uint32 GetTypeHash(const FReservedParameter& ReservedParameter) { return GetTypeHash(ReservedParameter.GetParameter().GetName()); };
 
 /** Niagara Editor module */
@@ -216,7 +231,7 @@ public:
 	NIAGARAEDITOR_API void RegisterDataInterfaceCacheVisualizer(UClass* DataInterfaceClass, TSharedRef<INiagaraDataInterfaceSimCacheVisualizer> InCacheVisualizer);
 	NIAGARAEDITOR_API void UnregisterDataInterfaceCacheVisualizer(UClass* DataInterfaceClass, TSharedRef<INiagaraDataInterfaceSimCacheVisualizer> InCacheVisualizer);
 	TArrayView<TSharedRef<INiagaraDataInterfaceSimCacheVisualizer>> FindDataInterfaceCacheVisualizer(UClass* DataInterfaceClass);
-
+	
 	TSharedRef<FNiagaraScriptMergeManager> GetScriptMergeManager() const;
 
 	// Object pooling methods used to prevent unnecessary object allocation during compiles
@@ -249,6 +264,8 @@ public:
 	NIAGARAEDITOR_API const FNiagaraEditorCommands& GetCommands() const;
 
 	void InvalidateCachedScriptAssetData();
+
+	FNiagaraStackObjectCustomizationRegistry& GetStackObjectCustomizationRegistry() { return StackObjectCustomizationRegistry; }
 	
 	const TArray<UNiagaraScript*>& GetCachedTypeConversionScripts() const;
 
@@ -413,6 +430,8 @@ private:
 		StackIssueGenerators.Add(StructName) = Generator;
 	}
 
+	void RegisterDefaultStackObjectCustomizations();
+
 	void OnAssetRegistryLoadComplete();
 	
 	void OnAssetsPreDelete(const TArray<UObject*>& Objects);
@@ -474,6 +493,8 @@ private:
 	TSharedPtr<FNiagaraComponentBroker> NiagaraComponentBroker;
 
 	TMap<const UScriptStruct*, FOnCreateMovieSceneTrackForParameter> TypeToParameterTrackCreatorMap;
+
+	FNiagaraStackObjectCustomizationRegistry StackObjectCustomizationRegistry;
 	
 	IConsoleCommand* TestCompileScriptCommand;
 	IConsoleCommand* ValidateScriptVariableGuidsCommand;
diff --git a/Engine/Plugins/FX/Niagara/Source/NiagaraEditor/Public/ViewModels/Stack/NiagaraStackPropertyRow.h b/Engine/Plugins/FX/Niagara/Source/NiagaraEditor/Public/ViewModels/Stack/NiagaraStackPropertyRow.h
--- a/Engine/Plugins/FX/Niagara/Source/NiagaraEditor/Public/ViewModels/Stack/NiagaraStackPropertyRow.h
+++ b/Engine/Plugins/FX/Niagara/Source/NiagaraEditor/Public/ViewModels/Stack/NiagaraStackPropertyRow.h
@@ -6,6 +6,7 @@
 #include "ViewModels/Stack/NiagaraStackObjectShared.h"
 #include "NiagaraStackPropertyRow.generated.h"
 
+class FNiagaraStackObjectPropertyCustomization;
 class IDetailTreeNode;
 class UNiagaraNode;
 
@@ -26,6 +27,9 @@ public:
 	
 	NIAGARAEDITOR_API void SetOnFilterDetailNodes(FNiagaraStackObjectShared::FOnFilterDetailNodes InOnFilterDetailNodes) { OnFilterDetailNodes = InOnFilterDetailNodes; }
 
+	void SetPropertyCustomization(TSharedPtr<const FNiagaraStackObjectPropertyCustomization> Customization);
+	TSharedPtr<const FNiagaraStackObjectPropertyCustomization> GetCustomization() const { return PropertyCustomization; }
+	
 	NIAGARAEDITOR_API virtual EStackRowStyle GetStackRowStyle() const override;
 
 	NIAGARAEDITOR_API virtual bool GetShouldShowInStack() const override;
@@ -61,7 +65,8 @@ public:
 protected:
 	NIAGARAEDITOR_API virtual void FinalizeInternal() override;
 
-	NIAGARAEDITOR_API virtual void RefreshChildrenInternal(const TArray<UNiagaraStackEntry*>& CurrentChildren, TArray<UNiagaraStackEntry*>& NewChildren, TArray<FStackIssue>& NewIssues) override;
+	virtual void RefreshChildrenInternal(const TArray<UNiagaraStackEntry*>& CurrentChildren, TArray<UNiagaraStackEntry*>& NewChildren, TArray<FStackIssue>& NewIssues) override;
+	void RefreshCustomization();
 
 	NIAGARAEDITOR_API virtual int32 GetChildIndentLevel() const override;
 
@@ -77,6 +82,7 @@ private:
 	TSharedPtr<IDetailTreeNode> DetailTreeNode;
 	UNiagaraNode* OwningNiagaraNode;
 	EStackRowStyle RowStyle;
+	TSharedPtr<const FNiagaraStackObjectPropertyCustomization> PropertyCustomization;
 
 	UPROPERTY()
 	TObjectPtr<UNiagaraStackSpacer> CategorySpacer;
diff --git a/Engine/Plugins/FX/Niagara/Source/NiagaraEditorWidgets/Private/SNiagaraStack.cpp b/Engine/Plugins/FX/Niagara/Source/NiagaraEditorWidgets/Private/SNiagaraStack.cpp
--- a/Engine/Plugins/FX/Niagara/Source/NiagaraEditorWidgets/Private/SNiagaraStack.cpp
+++ b/Engine/Plugins/FX/Niagara/Source/NiagaraEditorWidgets/Private/SNiagaraStack.cpp
@@ -1070,12 +1070,24 @@ SNiagaraStack::FRowWidgets SNiagaraStack::ConstructNameAndValueWidgetsForItem(UN
 		}
 		else
 		{
+			TSharedPtr<const FNiagaraStackObjectPropertyCustomization> Customization = PropertyRow->GetCustomization();
+			if(Customization)
+			{
+				TOptional<TSharedPtr<SWidget>> CustomizedNameWidget = Customization->GenerateNameWidget(PropertyRow);
+
+				if(CustomizedNameWidget.IsSet())
+				{
+					PropertyRowWidgets.NameWidget = CustomizedNameWidget.GetValue();
+				}
+			}
+			
 			Container->SetOverrideNameWidth(PropertyRowWidgets.NameWidgetLayoutData.MinWidth, PropertyRowWidgets.NameWidgetLayoutData.MaxWidth);
 			Container->SetOverrideValueWidth(PropertyRowWidgets.ValueWidgetLayoutData.MinWidth, PropertyRowWidgets.ValueWidgetLayoutData.MaxWidth);
 			Container->SetOverrideNameAlignment(HAlign_Left, VAlign_Center);
 			PropertyRowWidgets.NameWidget->SetEnabled(IsEnabled);
 			PropertyRowWidgets.ValueWidget->SetEnabled(IsEnabled);
 
+			
 			TSharedPtr<SWidget> NameWidget;
 			if (PropertyRowWidgets.EditConditionWidget.IsValid())
 			{
@@ -1088,7 +1100,7 @@ SNiagaraStack::FRowWidgets SNiagaraStack::ConstructNameAndValueWidgetsForItem(UN
 				]
 				+ SHorizontalBox::Slot()
 				[
-					PropertyRowWidgets.NameWidget.ToSharedRef()
+					 PropertyRowWidgets.NameWidget.ToSharedRef()
 				];
 			}
 			else
