- Fix missing external access for texture sample data interfaces

diff --git a/Engine/Plugins/FX/Niagara/Source/Niagara/Private/NiagaraDataInterface2DArrayTexture.cpp b/Engine/Plugins/FX/Niagara/Source/Niagara/Private/NiagaraDataInterface2DArrayTexture.cpp
--- a/Engine/Plugins/FX/Niagara/Source/Niagara/Private/NiagaraDataInterface2DArrayTexture.cpp
+++ b/Engine/Plugins/FX/Niagara/Source/Niagara/Private/NiagaraDataInterface2DArrayTexture.cpp
@@ -10,6 +10,7 @@
 #include "Engine/TextureRenderTarget2DArray.h"
 #include "RenderGraphBuilder.h"
 #include "RHIStaticStates.h"
+#include "RenderGraphUtils.h"
 #include "ShaderCompilerCore.h"
 #include "TextureResource.h"
 
@@ -287,6 +288,7 @@ void UNiagaraDataInterface2DArrayTexture::SetShaderParameters(const FNiagaraData
 {
 	FNiagaraDataInterfaceProxyTexture2DArray& TextureProxy = Context.GetProxy<FNiagaraDataInterfaceProxyTexture2DArray>();
 	FNDITexture2DArrayInstanceData_RenderThread* InstanceData = TextureProxy.InstanceData_RT.Find(Context.GetSystemInstanceID());
+	FRDGBuilder& GraphBuilder = Context.GetGraphBuilder();
 
 	FShaderParameters* Parameters = Context.GetParameterNestedStruct<FShaderParameters>();
 	if (InstanceData && InstanceData->TextureReferenceRHI.IsValid())
@@ -301,23 +303,26 @@ void UNiagaraDataInterface2DArrayTexture::SetShaderParameters(const FNiagaraData
 				FTextureRHIRef ResolvedTextureRHI = InstanceData->TextureReferenceRHI->GetReferencedTexture();
 				if (ResolvedTextureRHI.IsValid())
 				{
-					InstanceData->TransientRDGTexture = Context.GetGraphBuilder().FindExternalTexture(ResolvedTextureRHI);
+					InstanceData->TransientRDGTexture = GraphBuilder.FindExternalTexture(ResolvedTextureRHI);
 					if (InstanceData->TransientRDGTexture == nullptr)
 					{
-						InstanceData->TransientRDGTexture = Context.GetGraphBuilder().RegisterExternalTexture(CreateRenderTarget(ResolvedTextureRHI, TEXT("NiagaraTexture2DArray")));
+						InstanceData->TransientRDGTexture = GraphBuilder.RegisterExternalTexture(CreateRenderTarget(ResolvedTextureRHI, TEXT("NiagaraTexture2DArray")));
+
+						GraphBuilder.UseInternalAccessMode(InstanceData->TransientRDGTexture);
+						Context.GetRDGExternalAccessQueue().Add(InstanceData->TransientRDGTexture);
 					}
 					RDGTexture = InstanceData->TransientRDGTexture;
 				}
 			}
 
-			Parameters->Texture = RDGTexture ? RDGTexture : Context.GetComputeDispatchInterface().GetBlackTexture(Context.GetGraphBuilder(), ETextureDimension::Texture2DArray);
+			Parameters->Texture = RDGTexture ? RDGTexture : Context.GetComputeDispatchInterface().GetBlackTexture(GraphBuilder, ETextureDimension::Texture2DArray);
 		}
 	}
 	else
 	{
 		Parameters->TextureSize = FVector3f::ZeroVector;
 		Parameters->TextureSampler = TStaticSamplerState<SF_Point>::GetRHI();
-		Parameters->Texture = Context.GetComputeDispatchInterface().GetBlackTexture(Context.GetGraphBuilder(), ETextureDimension::Texture2DArray);
+		Parameters->Texture = Context.GetComputeDispatchInterface().GetBlackTexture(GraphBuilder, ETextureDimension::Texture2DArray);
 	}
 }
 
diff --git a/Engine/Plugins/FX/Niagara/Source/Niagara/Private/NiagaraDataInterfaceCubeTexture.cpp b/Engine/Plugins/FX/Niagara/Source/Niagara/Private/NiagaraDataInterfaceCubeTexture.cpp
--- a/Engine/Plugins/FX/Niagara/Source/Niagara/Private/NiagaraDataInterfaceCubeTexture.cpp
+++ b/Engine/Plugins/FX/Niagara/Source/Niagara/Private/NiagaraDataInterfaceCubeTexture.cpp
@@ -10,6 +10,7 @@
 #include "Engine/TextureRenderTargetCube.h"
 #include "RenderGraphBuilder.h"
 #include "RHIStaticStates.h"
+#include "RenderGraphUtils.h"
 #include "ShaderCompilerCore.h"
 #include "TextureResource.h"
 
@@ -280,6 +281,7 @@ void UNiagaraDataInterfaceCubeTexture::SetShaderParameters(const FNiagaraDataInt
 {
 	FNiagaraDataInterfaceProxyCubeTexture& TextureProxy = Context.GetProxy<FNiagaraDataInterfaceProxyCubeTexture>();
 	FNDICubeTextureInstanceData_RenderThread* InstanceData = TextureProxy.InstanceData_RT.Find(Context.GetSystemInstanceID());
+	FRDGBuilder& GraphBuilder = Context.GetGraphBuilder();
 
 	FShaderParameters* Parameters = Context.GetParameterNestedStruct<FShaderParameters>();
 	if (InstanceData && InstanceData->TextureReferenceRHI.IsValid())
@@ -294,23 +296,26 @@ void UNiagaraDataInterfaceCubeTexture::SetShaderParameters(const FNiagaraDataInt
 				FTextureRHIRef ResolvedTextureRHI = InstanceData->TextureReferenceRHI->GetReferencedTexture();
 				if (ResolvedTextureRHI.IsValid())
 				{
-					InstanceData->TransientRDGTexture = Context.GetGraphBuilder().FindExternalTexture(ResolvedTextureRHI);
+					InstanceData->TransientRDGTexture = GraphBuilder.FindExternalTexture(ResolvedTextureRHI);
 					if (InstanceData->TransientRDGTexture == nullptr)
 					{
-						InstanceData->TransientRDGTexture = Context.GetGraphBuilder().RegisterExternalTexture(CreateRenderTarget(ResolvedTextureRHI, TEXT("NiagaraTextureCube")));
+						InstanceData->TransientRDGTexture = GraphBuilder.RegisterExternalTexture(CreateRenderTarget(ResolvedTextureRHI, TEXT("NiagaraTextureCube")));
+
+						GraphBuilder.UseInternalAccessMode(InstanceData->TransientRDGTexture);
+						Context.GetRDGExternalAccessQueue().Add(InstanceData->TransientRDGTexture);
 					}
 					RDGTexture = InstanceData->TransientRDGTexture;
 				}
 			}
 
-			Parameters->Texture = RDGTexture ? RDGTexture : Context.GetComputeDispatchInterface().GetBlackTexture(Context.GetGraphBuilder(), ETextureDimension::TextureCube);
+			Parameters->Texture = RDGTexture ? RDGTexture : Context.GetComputeDispatchInterface().GetBlackTexture(GraphBuilder, ETextureDimension::TextureCube);
 		}
 	}
 	else
 	{
 		Parameters->TextureSize = FIntPoint::ZeroValue;
 		Parameters->TextureSampler = TStaticSamplerState<SF_Point>::GetRHI();
-		Parameters->Texture = Context.GetComputeDispatchInterface().GetBlackTexture(Context.GetGraphBuilder(), ETextureDimension::TextureCube);
+		Parameters->Texture = Context.GetComputeDispatchInterface().GetBlackTexture(GraphBuilder, ETextureDimension::TextureCube);
 	}
 }
 
diff --git a/Engine/Plugins/FX/Niagara/Source/Niagara/Private/NiagaraDataInterfaceTexture.cpp b/Engine/Plugins/FX/Niagara/Source/Niagara/Private/NiagaraDataInterfaceTexture.cpp
--- a/Engine/Plugins/FX/Niagara/Source/Niagara/Private/NiagaraDataInterfaceTexture.cpp
+++ b/Engine/Plugins/FX/Niagara/Source/Niagara/Private/NiagaraDataInterfaceTexture.cpp
@@ -11,6 +11,7 @@
 #include "NiagaraSystemInstance.h"
 #include "RenderGraphBuilder.h"
 #include "RHIStaticStates.h"
+#include "RenderGraphUtils.h"
 #include "ShaderCompilerCore.h"
 #include "TextureResource.h"
 
@@ -453,6 +454,7 @@ void UNiagaraDataInterfaceTexture::SetShaderParameters(const FNiagaraDataInterfa
 {
 	FNiagaraDataInterfaceProxyTexture& TextureProxy = Context.GetProxy<FNiagaraDataInterfaceProxyTexture>();
 	FNDITextureInstanceData_RenderThread* InstanceData = TextureProxy.InstanceData_RT.Find(Context.GetSystemInstanceID());
+	FRDGBuilder& GraphBuilder = Context.GetGraphBuilder();
 
 	FShaderParameters* Parameters = Context.GetParameterNestedStruct<FShaderParameters>();
 	if (InstanceData && InstanceData->TextureReferenceRHI.IsValid())
@@ -468,16 +470,19 @@ void UNiagaraDataInterfaceTexture::SetShaderParameters(const FNiagaraDataInterfa
 				FTextureRHIRef ResolvedTextureRHI = InstanceData->TextureReferenceRHI->GetReferencedTexture();
 				if (ResolvedTextureRHI.IsValid())
 				{
-					InstanceData->TransientRDGTexture = Context.GetGraphBuilder().FindExternalTexture(ResolvedTextureRHI);
+					InstanceData->TransientRDGTexture = GraphBuilder.FindExternalTexture(ResolvedTextureRHI);
 					if (InstanceData->TransientRDGTexture == nullptr)
 					{
-						InstanceData->TransientRDGTexture = Context.GetGraphBuilder().RegisterExternalTexture(CreateRenderTarget(ResolvedTextureRHI, TEXT("NiagaraTexture2D")));
+						InstanceData->TransientRDGTexture = GraphBuilder.RegisterExternalTexture(CreateRenderTarget(ResolvedTextureRHI, TEXT("NiagaraTexture2D")));
+
+						GraphBuilder.UseInternalAccessMode(InstanceData->TransientRDGTexture);
+						Context.GetRDGExternalAccessQueue().Add(InstanceData->TransientRDGTexture);
 					}
 					RDGTexture = InstanceData->TransientRDGTexture;
 				}
 			}
 
-			Parameters->Texture = RDGTexture ? RDGTexture : Context.GetComputeDispatchInterface().GetBlackTexture(Context.GetGraphBuilder(), ETextureDimension::Texture2D);
+			Parameters->Texture = RDGTexture ? RDGTexture : Context.GetComputeDispatchInterface().GetBlackTexture(GraphBuilder, ETextureDimension::Texture2D);
 		}
 	}
 	else
@@ -485,7 +490,7 @@ void UNiagaraDataInterfaceTexture::SetShaderParameters(const FNiagaraDataInterfa
 		Parameters->TextureSize = FIntPoint(0, 0);
 		Parameters->MipLevels = 0;
 		Parameters->TextureSampler = TStaticSamplerState<SF_Point>::GetRHI();
-		Parameters->Texture = Context.GetComputeDispatchInterface().GetBlackTexture(Context.GetGraphBuilder(), ETextureDimension::Texture2D);
+		Parameters->Texture = Context.GetComputeDispatchInterface().GetBlackTexture(GraphBuilder, ETextureDimension::Texture2D);
 	}
 }
 
diff --git a/Engine/Plugins/FX/Niagara/Source/Niagara/Private/NiagaraDataInterfaceVolumeTexture.cpp b/Engine/Plugins/FX/Niagara/Source/Niagara/Private/NiagaraDataInterfaceVolumeTexture.cpp
--- a/Engine/Plugins/FX/Niagara/Source/Niagara/Private/NiagaraDataInterfaceVolumeTexture.cpp
+++ b/Engine/Plugins/FX/Niagara/Source/Niagara/Private/NiagaraDataInterfaceVolumeTexture.cpp
@@ -9,6 +9,7 @@
 #include "Engine/VolumeTexture.h"
 #include "RenderGraphBuilder.h"
 #include "RHIStaticStates.h"
+#include "RenderGraphUtils.h"
 #include "ShaderCompilerCore.h"
 #include "TextureResource.h"
 
@@ -304,6 +305,7 @@ void UNiagaraDataInterfaceVolumeTexture::SetShaderParameters(const FNiagaraDataI
 {
 	FNiagaraDataInterfaceProxyVolumeTexture& TextureProxy = Context.GetProxy<FNiagaraDataInterfaceProxyVolumeTexture>();
 	FNDIVolumeTextureInstanceData_RenderThread* InstanceData = TextureProxy.InstanceData_RT.Find(Context.GetSystemInstanceID());
+	FRDGBuilder& GraphBuilder = Context.GetGraphBuilder();
 
 	FShaderParameters* Parameters = Context.GetParameterNestedStruct<FShaderParameters>();
 	if (InstanceData && InstanceData->TextureReferenceRHI.IsValid())
@@ -318,23 +320,26 @@ void UNiagaraDataInterfaceVolumeTexture::SetShaderParameters(const FNiagaraDataI
 				FTextureRHIRef ResolvedTextureRHI = InstanceData->TextureReferenceRHI->GetReferencedTexture();
 				if (ResolvedTextureRHI.IsValid())
 				{
-					InstanceData->TransientRDGTexture = Context.GetGraphBuilder().FindExternalTexture(ResolvedTextureRHI);
+					InstanceData->TransientRDGTexture = GraphBuilder.FindExternalTexture(ResolvedTextureRHI);
 					if (InstanceData->TransientRDGTexture == nullptr)
 					{
-						InstanceData->TransientRDGTexture = Context.GetGraphBuilder().RegisterExternalTexture(CreateRenderTarget(ResolvedTextureRHI, TEXT("NiagaraTextureVolume")));
+						InstanceData->TransientRDGTexture = GraphBuilder.RegisterExternalTexture(CreateRenderTarget(ResolvedTextureRHI, TEXT("NiagaraTextureVolume")));
+
+						GraphBuilder.UseInternalAccessMode(InstanceData->TransientRDGTexture);
+						Context.GetRDGExternalAccessQueue().Add(InstanceData->TransientRDGTexture);
 					}
 					RDGTexture = InstanceData->TransientRDGTexture;
 				}
 			}
 
-			Parameters->Texture = RDGTexture ? RDGTexture : Context.GetComputeDispatchInterface().GetBlackTexture(Context.GetGraphBuilder(), ETextureDimension::Texture3D);
+			Parameters->Texture = RDGTexture ? RDGTexture : Context.GetComputeDispatchInterface().GetBlackTexture(GraphBuilder, ETextureDimension::Texture3D);
 		}
 	}
 	else
 	{
 		Parameters->TextureSize = FVector3f::ZeroVector;
 		Parameters->TextureSampler = TStaticSamplerState<SF_Point>::GetRHI();
-		Parameters->Texture = Context.GetComputeDispatchInterface().GetBlackTexture(Context.GetGraphBuilder(), ETextureDimension::Texture3D);
+		Parameters->Texture = Context.GetComputeDispatchInterface().GetBlackTexture(GraphBuilder, ETextureDimension::Texture3D);
 	}
 }
 
