- Expose options for mesh exporting to baker
- Static Mesh building requires various attributes so make sure we set with defaults if not exported

diff --git a/Engine/Plugins/FX/Niagara/Source/Niagara/Classes/NiagaraBakerOutputStaticMesh.h b/Engine/Plugins/FX/Niagara/Source/Niagara/Classes/NiagaraBakerOutputStaticMesh.h
--- a/Engine/Plugins/FX/Niagara/Source/Niagara/Classes/NiagaraBakerOutputStaticMesh.h
+++ b/Engine/Plugins/FX/Niagara/Source/Niagara/Classes/NiagaraBakerOutputStaticMesh.h
@@ -3,6 +3,7 @@
 #pragma once
 
 #include "NiagaraBakerOutput.h"
+#include "NiagaraRendererReadback.h"
 #include "NiagaraBakerOutputStaticMesh.generated.h"
 
 UCLASS(meta = (DisplayName = "Static Mesh Output"), MinimalAPI)
@@ -18,8 +19,8 @@ public:
 	UPROPERTY(EditAnywhere, Category = "Settings")
 	FString FramesAssetPathFormat = TEXT("{AssetFolder}/{AssetName}_BakedFrame_{OutputName}_{FrameIndex}");
 
-	//UPROPERTY(EditAnywhere, Category = "Settings")
-	//FNiagaraRendererReadbackParameters CreateParameters;
+	UPROPERTY(EditAnywhere, Category = "Settings")
+	FNiagaraRendererReadbackParameters ExportParameters;
 
 	NIAGARA_API virtual bool Equals(const UNiagaraBakerOutput& Other) const override;
 
diff --git a/Engine/Plugins/FX/Niagara/Source/Niagara/Internal/NiagaraRendererReadback.h b/Engine/Plugins/FX/Niagara/Source/Niagara/Internal/NiagaraRendererReadback.h
--- a/Engine/Plugins/FX/Niagara/Source/Niagara/Internal/NiagaraRendererReadback.h
+++ b/Engine/Plugins/FX/Niagara/Source/Niagara/Internal/NiagaraRendererReadback.h
@@ -5,6 +5,8 @@
 #include "PrimitiveComponentId.h"
 #include "MeshBatch.h"
 
+#include "NiagaraRendererReadback.generated.h"
+
 class UNiagaraComponent;
 class FNiagaraSceneProxy;
 class FPrimitiveSceneProxy;
@@ -16,17 +18,44 @@ class FSceneView;
 	#define WITH_NIAGARA_RENDERER_READBACK WITH_EDITOR
 #endif
 
-#if WITH_NIAGARA_RENDERER_READBACK
+USTRUCT(BlueprintType)
 struct FNiagaraRendererReadbackParameters
 {
-	bool bExportPosition		= true;
-	bool bExportTangentBasis	= true;
-	bool bExportColor			= true;
-	uint32 ExportNumTexCoords	= 1;
+	GENERATED_BODY()
+
+	// When enabled vertex positions will be exported
+	UPROPERTY(EditAnywhere, BlueprintReadWrite, Category = "Parameters")
+	bool bExportPosition = true;
+
+	// When enabled the vertex tangent basis will be exported
+	UPROPERTY(EditAnywhere, BlueprintReadWrite, Category = "Parameters")
+	bool bExportTangentBasis = true;
+
+	// When enabled vertex colors will be exported
+	UPROPERTY(EditAnywhere, BlueprintReadWrite, Category = "Parameters")
+	bool bExportColor = true;
+	//-TODO: Add Color mode (i.e. Vertex Color / Particle Color / Disabled)
 
-	bool bApplyWPO				= false;
+	// How many vertex texture coordinates to export
+	UPROPERTY(EditAnywhere, BlueprintReadWrite, Category = "Parameters", meta=(UIMin="0", UIMax="8", ClampMin="0", ClampMax="8"));
+	int32 ExportNumTexCoords = 1;
+
+	// When enabled we will attempt to export the materials used with each section
+	// When disabled no materials will be assigned, so the default material will be used
+	UPROPERTY(EditAnywhere, BlueprintReadWrite, Category = "Parameters")
+	bool bExportMaterials = true;
+
+	// When enabled material WPO will be included in the exported data
+	UPROPERTY(EditAnywhere, BlueprintReadWrite, Category = "Parameters")
+	bool bApplyWPO = false;
+
+	// When set we capture the batches from the view index provided
+	// When unset we capture batches from all views
+	UPROPERTY(EditAnywhere, BlueprintReadWrite, Category = "Parameters")
+	TOptional<int> ViewIndexToCapture = 0;
 };
 
+#if WITH_NIAGARA_RENDERER_READBACK
 struct FNiagaraRendererReadbackResult
 {
 	struct FSection
diff --git a/Engine/Plugins/FX/Niagara/Source/Niagara/Private/NiagaraRendererReadback.cpp b/Engine/Plugins/FX/Niagara/Source/Niagara/Private/NiagaraRendererReadback.cpp
--- a/Engine/Plugins/FX/Niagara/Source/Niagara/Private/NiagaraRendererReadback.cpp
+++ b/Engine/Plugins/FX/Niagara/Source/Niagara/Private/NiagaraRendererReadback.cpp
@@ -15,6 +15,7 @@
 #include "Materials/Material.h"
 #include "Materials/MaterialRenderProxy.h"
 #include "Async/Async.h"
+#include "Components.h"
 #include "MeshMaterialShader.h"
 #include "MeshPassProcessor.h"
 #include "RenderGraphUtils.h"
@@ -48,7 +49,9 @@ namespace NiagaraRendererReadback
 
 	struct FRendererReadbackRequest
 	{
+		bool													bExportMaterials = true;
 		bool													bApplyWPO = false;
+		TOptional<int>											ViewIndexToCapture;
 		FNiagaraRendererReadbackResult							Result;
 
 		UE::FMutex												MeshBatchesMutex;
@@ -211,9 +214,12 @@ namespace NiagaraRendererReadback
 					const uint32 MaxOrNumVertices = ReadbackMeshBatch.NumInstances * ReadbackMeshBatch.NumVerticesPerInstance;
 
 					FNiagaraRendererReadbackResult::FSection& Section = Request->Result.Sections.AddDefaulted_GetRef();
-					Section.WeakMaterialInterface	= MeshBatch.MaterialRenderProxy->GetMaterialInterface();
-					Section.FirstTriangle			= Request->Result.NumVertices / 3;
-					Section.NumTriangles			= MaxOrNumVertices / 3;
+					if (Request->bExportMaterials)
+					{
+						Section.WeakMaterialInterface	= MeshBatch.MaterialRenderProxy->GetMaterialInterface();
+					}
+					Section.FirstTriangle	= Request->Result.NumVertices / 3;
+					Section.NumTriangles	= MaxOrNumVertices / 3;
 					Request->Result.NumVertices += MaxOrNumVertices;
 				}
 
@@ -360,6 +366,16 @@ namespace NiagaraRendererReadback
 					return;
 				}
 
+				if (Request->ViewIndexToCapture.IsSet())
+				{
+					const int32 CaptureViewIndex = Request->ViewIndexToCapture.GetValue();
+					const FSceneView* CaptureView = View && View->Family && View->Family->Views.IsValidIndex(CaptureViewIndex) ? View->Family->Views[CaptureViewIndex] : nullptr;
+					if (View != CaptureView)
+					{
+						continue;
+					}
+				}
+
 				UE::TScopeLock ScopeLock(Request->MeshBatchesMutex);
 				FReadbackMeshBatch& ReadbackMeshBatch = Request->MeshBatches.AddDefaulted_GetRef();
 				ReadbackMeshBatch.View						= View;
@@ -432,7 +448,9 @@ void NiagaraRendererReadback::EnqueueReadback(TConstArrayView<UNiagaraComponent*
 	}
 
 	// Copy over parameters
-	ReadbackRequest->bApplyWPO = Parameters.bApplyWPO;
+	ReadbackRequest->bExportMaterials	= Parameters.bExportMaterials;
+	ReadbackRequest->bApplyWPO			= Parameters.bApplyWPO;
+	ReadbackRequest->ViewIndexToCapture	= Parameters.ViewIndexToCapture;
 
 	// Build Vertex Output
 	{
@@ -454,9 +472,10 @@ void NiagaraRendererReadback::EnqueueReadback(TConstArrayView<UNiagaraComponent*
 		}
 		if (Parameters.ExportNumTexCoords > 0)
 		{
-			Result.VertexTexCoordNum = Parameters.ExportNumTexCoords;
+			const int32 NumTexCoords = FMath::Min<int32>(Parameters.ExportNumTexCoords, MAX_STATIC_TEXCOORDS);
+			Result.VertexTexCoordNum = NumTexCoords;
 			Result.VertexTexCoordOffset = Result.VertexStride;
-			Result.VertexStride += sizeof(FVector2f) * Parameters.ExportNumTexCoords;
+			Result.VertexStride += sizeof(FVector2f) * NumTexCoords;
 		}
 	}
 
diff --git a/Engine/Plugins/FX/Niagara/Source/NiagaraEditor/Private/NiagaraBakerRendererOutputStaticMesh.cpp b/Engine/Plugins/FX/Niagara/Source/NiagaraEditor/Private/NiagaraBakerRendererOutputStaticMesh.cpp
--- a/Engine/Plugins/FX/Niagara/Source/NiagaraEditor/Private/NiagaraBakerRendererOutputStaticMesh.cpp
+++ b/Engine/Plugins/FX/Niagara/Source/NiagaraEditor/Private/NiagaraBakerRendererOutputStaticMesh.cpp
@@ -140,19 +140,20 @@ void FNiagaraBakerRendererOutputStaticMesh::BakeFrame(FNiagaraBakerFeedbackConte
 			{
 				const FVertexID VertexID = MeshDescription.CreateVertex();
 				check(VertexID.GetValue() == iVertex);
-				VertexPositions[VertexID] = Result.GetPosition(iVertex);
-
-				const FVector3f TangentX = Result.GetTangentX(iVertex);
-				const FVector3f TangentY = Result.GetTangentY(iVertex);
-				const FVector3f TangentZ = Result.GetTangentZ(iVertex);
-				const float TangentSign = FVector3f::DotProduct(FVector3f::CrossProduct(TangentX, TangentZ), TangentY) < 0.0f ? -1.0f : 1.0f;
+				VertexPositions[VertexID] = Result.HasPosition() ? Result.GetPosition(iVertex) : FVector3f::ZeroVector;
 
 				FVertexInstanceID VertexInstanceID = MeshDescription.CreateVertexInstance(VertexID);
 				VertexInstances.Add(VertexInstanceID);
+
+				const FVector3f TangentX = Result.HasTangentBasis() ? Result.GetTangentX(iVertex) : FVector3f::XAxisVector;
+				const FVector3f TangentY = Result.HasTangentBasis() ? Result.GetTangentY(iVertex) : FVector3f::YAxisVector;
+				const FVector3f TangentZ = Result.HasTangentBasis() ? Result.GetTangentZ(iVertex) : FVector3f::ZAxisVector;
+				const float TangentSign = FVector3f::DotProduct(FVector3f::CrossProduct(TangentX, TangentZ), TangentY) < 0.0f ? -1.0f : 1.0f;
 				VertexInstanceNormals.Set(VertexInstanceID, TangentX);
 				VertexInstanceTangents.Set(VertexInstanceID, TangentZ);
 				VertexInstanceBinormalSigns.Set(VertexInstanceID, TangentSign);
-				VertexInstanceColors.Set(VertexInstanceID, Result.GetColor(iVertex));
+
+				VertexInstanceColors.Set(VertexInstanceID, Result.HasColor() ? Result.GetColor(iVertex) : FLinearColor::White);
 				for (uint32 iTexCoord=0; iTexCoord < NumTexCoords; ++iTexCoord)
 				{
 					VertexInstanceUVs.Set(VertexInstanceID, iTexCoord, Result.GetTexCoord(iVertex, iTexCoord));
@@ -182,6 +183,13 @@ void FNiagaraBakerRendererOutputStaticMesh::BakeFrame(FNiagaraBakerFeedbackConte
 				}
 			}
 
+			StaticMesh->SetNumSourceModels(1);
+			{
+				FMeshBuildSettings& MeshBuildSettings = StaticMesh->GetSourceModel(0).BuildSettings;
+				MeshBuildSettings.bRecomputeNormals = false;
+				MeshBuildSettings.bRecomputeTangents = false;
+			}
+
 			StaticMesh->SetStaticMaterials(StaticMaterials);
 
 			// Build Mesh
@@ -193,7 +201,8 @@ void FNiagaraBakerRendererOutputStaticMesh::BakeFrame(FNiagaraBakerFeedbackConte
 			Params.bAllowCpuAccess = false;
 
 			StaticMesh->BuildFromMeshDescriptions({ &MeshDescription }, Params);
-		}
+		},
+		BakerOutput->ExportParameters
 	);
 
 	BakerRenderer.RenderSceneCapture(BakeRenderTarget, ESceneCaptureSource::SCS_SceneColorHDR);
