Remove int type pill icon from 'enum style' integer module inputs

diff --git a/Engine/Plugins/FX/Niagara/Source/NiagaraEditor/Private/Widgets/SNiagaraNumericDropDown.h b/Engine/Plugins/FX/Niagara/Source/NiagaraEditor/Private/Widgets/SNiagaraNumericDropDown.h
--- a/Engine/Plugins/FX/Niagara/Source/NiagaraEditor/Private/Widgets/SNiagaraNumericDropDown.h
+++ b/Engine/Plugins/FX/Niagara/Source/NiagaraEditor/Private/Widgets/SNiagaraNumericDropDown.h
@@ -131,6 +131,7 @@ public:
 				SNew(SImage)
 				.ColorAndOpacity(UEdGraphSchema_Niagara::GetTypeColor(InArgs._PillType.Get(FNiagaraTypeDefinition::GetFloatDef())))
 				.Image(FNiagaraEditorStyle::Get().GetBrush("NiagaraEditor.Module.TypeIconPill"))
+				.Visibility(bAllowTyping ? EVisibility::Visible : EVisibility::Collapsed)
 			]
 			+ SHorizontalBox::Slot()
 			.AutoWidth()
