- Fix for SKM pre skinned extents / half extends being half the value

#tests Create a NS to display the preskinned local bounds, verify against a material that outputs the same values

diff --git a/Engine/Plugins/FX/Niagara/Source/Niagara/Private/NiagaraDataInterfaceSkeletalMesh.cpp b/Engine/Plugins/FX/Niagara/Source/Niagara/Private/NiagaraDataInterfaceSkeletalMesh.cpp
--- a/Engine/Plugins/FX/Niagara/Source/Niagara/Private/NiagaraDataInterfaceSkeletalMesh.cpp
+++ b/Engine/Plugins/FX/Niagara/Source/Niagara/Private/NiagaraDataInterfaceSkeletalMesh.cpp
@@ -1811,7 +1811,7 @@ bool FNDISkeletalMesh_InstanceData::Init(UNiagaraDataInterfaceSkeletalMesh* Inte
 		{
 			const FBoxSphereBounds LocalBounds = Mesh->GetBounds();
 			PreSkinnedLocalBoundsCenter = FVector3f(LocalBounds.Origin);
-			PreSkinnedLocalBoundsExtents = FVector3f(LocalBounds.BoxExtent);
+			PreSkinnedLocalBoundsExtents = FVector3f(LocalBounds.BoxExtent) * 2.0f;
 		}
 	}
 
