- Fix for sim cache bool comparison not working on all platforms

#rnx

diff --git a/Engine/Plugins/FX/Niagara/Source/Niagara/Private/NiagaraDataInterfaceArrayImplHelpers.h b/Engine/Plugins/FX/Niagara/Source/Niagara/Private/NiagaraDataInterfaceArrayImplHelpers.h
--- a/Engine/Plugins/FX/Niagara/Source/Niagara/Private/NiagaraDataInterfaceArrayImplHelpers.h
+++ b/Engine/Plugins/FX/Niagara/Source/Niagara/Private/NiagaraDataInterfaceArrayImplHelpers.h
@@ -538,6 +538,11 @@ struct FNDIArrayImplHelper<bool> : public FNDIArrayImplHelperBase<bool>
 	{
 		OutString.Append(Value ? TEXT("True") : TEXT("False"));
 	}
+
+	static bool IsNearlyEqual(bool Lhs, bool Rhs, float Tolerance)
+	{
+		return !!Lhs == !!Rhs;
+	}
 };
 
 template<>
