[HWRT/Niagara] Fix FNiagaraRendererMeshes::GetDynamicRayTracingInstances(...) not setting up vertex factory using ray tracing proxy buffers.

diff --git a/Engine/Plugins/FX/Niagara/Source/Niagara/Private/NiagaraMeshRendererProperties.cpp b/Engine/Plugins/FX/Niagara/Source/Niagara/Private/NiagaraMeshRendererProperties.cpp
--- a/Engine/Plugins/FX/Niagara/Source/Niagara/Private/NiagaraMeshRendererProperties.cpp
+++ b/Engine/Plugins/FX/Niagara/Source/Niagara/Private/NiagaraMeshRendererProperties.cpp
@@ -172,6 +172,15 @@ public:
 		InVertexFactory.SetData(RHICmdList, Data);
 	}
 
+#if RHI_RAYTRACING
+	void SetupRayTraceVertexFactory(FRHICommandListBase& RHICmdList, class FNiagaraMeshVertexFactory& InVertexFactory, const FLODModelData& LODModelData) const override
+	{
+		FStaticMeshDataType Data;
+		InitVertexFactoryComponents(*RenderData->RayTracingProxy->LODs[LODModelData.LODIndex].VertexBuffers, &InVertexFactory, Data);
+		InVertexFactory.SetData(RHICmdList, Data);
+	}
+#endif
+
 	void GetUsedMaterials(TArray<UMaterialInterface*>& OutMaterials) const override
 	{
 		const UStaticMesh* StaticMesh = WeakStaticMesh.Get();
diff --git a/Engine/Plugins/FX/Niagara/Source/Niagara/Private/NiagaraRendererMeshes.cpp b/Engine/Plugins/FX/Niagara/Source/Niagara/Private/NiagaraRendererMeshes.cpp
--- a/Engine/Plugins/FX/Niagara/Source/Niagara/Private/NiagaraRendererMeshes.cpp
+++ b/Engine/Plugins/FX/Niagara/Source/Niagara/Private/NiagaraRendererMeshes.cpp
@@ -1681,7 +1681,7 @@ void FNiagaraRendererMeshes::GetDynamicRayTracingInstances(FRayTracingInstanceCo
 		}
 		VertexFactory.EnablePrimitiveIDElement(ParticleMeshRenderData.bUseGPUScene);
 		VertexFactory.InitResource(RHICmdList);
-		MeshData.RenderableMesh->SetupVertexFactory(RHICmdList, VertexFactory, LODModel);
+		MeshData.RenderableMesh->SetupRayTraceVertexFactory(RHICmdList, VertexFactory, LODModel);
 
 		PreparePerMeshData(ParticleMeshRenderData, VertexFactory, *SceneProxy, MeshData);
 
diff --git a/Engine/Plugins/FX/Niagara/Source/Niagara/Public/NiagaraRenderableMeshInterface.h b/Engine/Plugins/FX/Niagara/Source/Niagara/Public/NiagaraRenderableMeshInterface.h
--- a/Engine/Plugins/FX/Niagara/Source/Niagara/Public/NiagaraRenderableMeshInterface.h
+++ b/Engine/Plugins/FX/Niagara/Source/Niagara/Public/NiagaraRenderableMeshInterface.h
@@ -48,6 +48,8 @@ public:
 	virtual void GetRayTraceLODModelData(FLODModelData& OutLODModelData, int32 LODLevel) const { GetLODModelData(OutLODModelData, LODLevel); }
 	// Setup the vertex factory for the mesh
 	virtual void SetupVertexFactory(FRHICommandListBase& RHICmdList, class FNiagaraMeshVertexFactory& InVertexFactory, const FLODModelData& LODModelData) const = 0;
+	// Setup the vertex factory for the mesh in the raytracing scene
+	virtual void SetupRayTraceVertexFactory(FRHICommandListBase& RHICmdList, class FNiagaraMeshVertexFactory& InVertexFactory, const FLODModelData& LODModelData) const { SetupVertexFactory(RHICmdList, InVertexFactory, LODModelData); }
 
 	UE_DEPRECATED(5.4, "SetupVertexFactory requires an RHI command list.")
 	virtual void SetupVertexFactory(class FNiagaraMeshVertexFactory& InVertexFactory, const FLODModelData& LODModelData) const final {}
