- Extend spawn infos in stateless to 8
- 4 was just too tight with loops and long particle lifetimes
- Add an on screen warning about overflowing the spawn info cap

#rnx
#tests Make an emitter with 4 spawn infos, loop of 1 second, lifetime 2 seconds, observe the particle counts when 1.5 second in should be 8 rather than 4.

diff --git a/Engine/Plugins/FX/Niagara/Shaders/Private/Stateless/NiagaraStatelessSimulationTemplate.ush b/Engine/Plugins/FX/Niagara/Shaders/Private/Stateless/NiagaraStatelessSimulationTemplate.ush
--- a/Engine/Plugins/FX/Niagara/Shaders/Private/Stateless/NiagaraStatelessSimulationTemplate.ush
+++ b/Engine/Plugins/FX/Niagara/Shaders/Private/Stateless/NiagaraStatelessSimulationTemplate.ush
@@ -38,10 +38,6 @@
 
 /////////////////////////////////////////////////////////////////////////////////////////////////////////////////////////
 
-#define NIAGARA_MAX_GPU_SPAWN_INFOS 4
-
-/////////////////////////////////////////////////////////////////////////////////////////////////////////////////////////
-
 uint	Common_RandomSeed;
 float	Common_SimulationTime;
 float	Common_SimulationDeltaTime;
diff --git a/Engine/Plugins/FX/Niagara/Source/Niagara/Internal/Stateless/NiagaraStatelessEmitterData.h b/Engine/Plugins/FX/Niagara/Source/Niagara/Internal/Stateless/NiagaraStatelessEmitterData.h
--- a/Engine/Plugins/FX/Niagara/Source/Niagara/Internal/Stateless/NiagaraStatelessEmitterData.h
+++ b/Engine/Plugins/FX/Niagara/Source/Niagara/Internal/Stateless/NiagaraStatelessEmitterData.h
@@ -33,6 +33,11 @@ struct FNiagaraStatelessEmitterData
 	FNiagaraStatelessEmitterData() = default;
 	~FNiagaraStatelessEmitterData();
 
+#if WITH_NIAGARA_DEBUG_EMITTER_NAME
+	FName											DebugSimulationName;
+	FName											DebugEmitterName;
+#endif
+
 	TSharedPtr<FNiagaraDataSetCompiledData>			ParticleDataSetCompiledData;
 	TArray<int32>									ComponentOffsets;
 
diff --git a/Engine/Plugins/FX/Niagara/Source/Niagara/Private/Stateless/NiagaraStatelessEmitter.cpp b/Engine/Plugins/FX/Niagara/Source/Niagara/Private/Stateless/NiagaraStatelessEmitter.cpp
--- a/Engine/Plugins/FX/Niagara/Source/Niagara/Private/Stateless/NiagaraStatelessEmitter.cpp
+++ b/Engine/Plugins/FX/Niagara/Source/Niagara/Private/Stateless/NiagaraStatelessEmitter.cpp
@@ -15,6 +15,7 @@
 
 #include "Algo/Copy.h"
 #include "Interfaces/ITargetPlatform.h"
+#include "UObject/Package.h"
 #include "UObject/UObjectIterator.h"
 
 namespace NiagaraStatelessInternal
@@ -295,6 +296,11 @@ const UNiagaraStatelessEmitterTemplate* UNiagaraStatelessEmitter::GetEmitterTemp
 void UNiagaraStatelessEmitter::CacheFromCompiledData()
 {
 	StatelessEmitterData = MakeShareable(new FNiagaraStatelessEmitterData(), FNiagaraStatelessEmitterData::FDeleter());
+
+#if WITH_NIAGARA_DEBUG_EMITTER_NAME
+	StatelessEmitterData->DebugSimulationName = GetPackage()->GetFName();
+	StatelessEmitterData->DebugEmitterName = GetFName();
+#endif
 	StatelessEmitterData->EmitterTemplate = GetEmitterTemplate();
 
 	// Build data set
diff --git a/Engine/Plugins/FX/Niagara/Source/Niagara/Private/Stateless/NiagaraStatelessEmitterData.cpp b/Engine/Plugins/FX/Niagara/Source/Niagara/Private/Stateless/NiagaraStatelessEmitterData.cpp
--- a/Engine/Plugins/FX/Niagara/Source/Niagara/Private/Stateless/NiagaraStatelessEmitterData.cpp
+++ b/Engine/Plugins/FX/Niagara/Source/Niagara/Private/Stateless/NiagaraStatelessEmitterData.cpp
@@ -5,6 +5,8 @@
 #include "Stateless/NiagaraStatelessParticleSimExecData.h"
 #include "Stateless/NiagaraStatelessSimulationShader.h"
 
+#include "Engine.h"
+
 void FNiagaraStatelessEmitterData::FDeleter::operator()(FNiagaraStatelessEmitterData* EmitterData) const
 {
 	if (IsInRenderingThread())
@@ -86,52 +88,65 @@ uint32 FNiagaraStatelessEmitterData::CalculateActiveParticles(int32 InRandomSeed
 		uint32 SpawnInfoTotalParticles = 0;
 
 		const bool bIsValidForAge = !Age.IsSet() || (Age.GetValue() >= SpawnInfo.SpawnTimeStart && Age.GetValue() < SpawnInfo.SpawnTimeEnd + SpawnInfo.LifetimeMax);
-		if (bIsValidForAge && GpuSpawnIndex < NiagaraStateless::MaxGpuSpawnInfos)
+		if (bIsValidForAge == false)
 		{
-			uint32	NumActive		= SpawnInfo.Amount;
-			uint32	ParticleOffset	= 0;
-			float	SpawnRate		= 0.0f;
-			float	SpawnTimeStart	= SpawnInfo.SpawnTimeStart;
-			switch (SpawnInfo.Type)
-			{
-				case ENiagaraStatelessSpawnInfoType::Burst:
-					break;
-
-				case ENiagaraStatelessSpawnInfoType::Rate:
-					// If the age is set we can make a narrowed number of active particles calculation
-					if (Age.IsSet())
-					{
-						const uint32 MaxActive = SpawnInfo.Amount;
-						ParticleOffset = FMath::FloorToInt(FMath::Max(Age.GetValue() - SpawnInfo.SpawnTimeStart - SpawnInfo.LifetimeMax, 0.0f) * SpawnInfo.Rate);
-						ParticleOffset = FMath::Min(ParticleOffset, MaxActive);
-						NumActive = FMath::FloorToInt(FMath::Max(Age.GetValue() - SpawnInfo.SpawnTimeStart, 0.0f) * SpawnInfo.Rate);
-						NumActive = FMath::Min(NumActive, MaxActive);
-						NumActive -= ParticleOffset;
-					}
-					SpawnRate = 1.0f / SpawnInfo.Rate;
-					SpawnTimeStart += SpawnRate;
-					break;
-
-				default:
-					checkNoEntry();
-					break;
-			}
-			if (NumActive > 0)
-			{
-				if (SpawnParameters)
+			continue;
+		}
+
+		uint32	NumActive		= SpawnInfo.Amount;
+		uint32	ParticleOffset	= 0;
+		float	SpawnRate		= 0.0f;
+		float	SpawnTimeStart	= SpawnInfo.SpawnTimeStart;
+		switch (SpawnInfo.Type)
+		{
+			case ENiagaraStatelessSpawnInfoType::Burst:
+				break;
+
+			case ENiagaraStatelessSpawnInfoType::Rate:
+				// If the age is set we can make a narrowed number of active particles calculation
+				if (Age.IsSet())
 				{
-					GET_SCALAR_ARRAY_ELEMENT(SpawnParameters->SpawnInfo_NumActive, GpuSpawnIndex)		= NumActive;
-					GET_SCALAR_ARRAY_ELEMENT(SpawnParameters->SpawnInfo_ParticleOffset, GpuSpawnIndex)	= ParticleOffset;
-					GET_SCALAR_ARRAY_ELEMENT(SpawnParameters->SpawnInfo_UniqueOffset, GpuSpawnIndex)	= SpawnInfo.UniqueOffset;
-					GET_SCALAR_ARRAY_ELEMENT(SpawnParameters->SpawnInfo_Time, GpuSpawnIndex)			= SpawnTimeStart;
-					GET_SCALAR_ARRAY_ELEMENT(SpawnParameters->SpawnInfo_Rate, GpuSpawnIndex)			= SpawnRate;
-					GET_SCALAR_ARRAY_ELEMENT(SpawnParameters->SpawnInfo_LifetimeScale, GpuSpawnIndex)	= SpawnInfo.LifetimeMax - SpawnInfo.LifetimeMin;
-					GET_SCALAR_ARRAY_ELEMENT(SpawnParameters->SpawnInfo_LifetimeBias, GpuSpawnIndex)	= SpawnInfo.LifetimeMin;
-					++GpuSpawnIndex;
+					const uint32 MaxActive = SpawnInfo.Amount;
+					ParticleOffset = FMath::FloorToInt(FMath::Max(Age.GetValue() - SpawnInfo.SpawnTimeStart - SpawnInfo.LifetimeMax, 0.0f) * SpawnInfo.Rate);
+					ParticleOffset = FMath::Min(ParticleOffset, MaxActive);
+					NumActive = FMath::FloorToInt(FMath::Max(Age.GetValue() - SpawnInfo.SpawnTimeStart, 0.0f) * SpawnInfo.Rate);
+					NumActive = FMath::Min(NumActive, MaxActive);
+					NumActive -= ParticleOffset;
 				}
-				TotalActiveParticles += NumActive;
-			}
+				SpawnRate = 1.0f / SpawnInfo.Rate;
+				SpawnTimeStart += SpawnRate;
+				break;
+
+			default:
+				checkNoEntry();
+				break;
+		}
+
+		if (NumActive == 0)
+		{
+			continue;
+		}
+
+		if (GpuSpawnIndex >= NiagaraStateless::MaxGpuSpawnInfos)
+		{
+		#if WITH_EDITOR && WITH_NIAGARA_DEBUG_EMITTER_NAME
+			GEngine->AddOnScreenDebugMessage(uint64(this), 1.f, FColor::White, *FString::Printf(TEXT("Stateless Simulation (%s:%s) has run out of space to store spawn infos results may flicker.  Reduce the number of spawn infos / loop time / particle age / increase code limit to fix."), *DebugSimulationName.ToString(), *DebugEmitterName.ToString()));
+		#endif
+			break;
+		}
+
+		if (SpawnParameters)
+		{
+			GET_SCALAR_ARRAY_ELEMENT(SpawnParameters->SpawnInfo_NumActive, GpuSpawnIndex)		= NumActive;
+			GET_SCALAR_ARRAY_ELEMENT(SpawnParameters->SpawnInfo_ParticleOffset, GpuSpawnIndex)	= ParticleOffset;
+			GET_SCALAR_ARRAY_ELEMENT(SpawnParameters->SpawnInfo_UniqueOffset, GpuSpawnIndex)	= SpawnInfo.UniqueOffset;
+			GET_SCALAR_ARRAY_ELEMENT(SpawnParameters->SpawnInfo_Time, GpuSpawnIndex)			= SpawnTimeStart;
+			GET_SCALAR_ARRAY_ELEMENT(SpawnParameters->SpawnInfo_Rate, GpuSpawnIndex)			= SpawnRate;
+			GET_SCALAR_ARRAY_ELEMENT(SpawnParameters->SpawnInfo_LifetimeScale, GpuSpawnIndex)	= SpawnInfo.LifetimeMax - SpawnInfo.LifetimeMin;
+			GET_SCALAR_ARRAY_ELEMENT(SpawnParameters->SpawnInfo_LifetimeBias, GpuSpawnIndex)	= SpawnInfo.LifetimeMin;
+			++GpuSpawnIndex;
 		}
+		TotalActiveParticles += NumActive;
 	}
 
 	if (SpawnParameters)
diff --git a/Engine/Plugins/FX/Niagara/Source/NiagaraShader/Internal/Stateless/NiagaraStatelessSimulationShader.h b/Engine/Plugins/FX/Niagara/Source/NiagaraShader/Internal/Stateless/NiagaraStatelessSimulationShader.h
--- a/Engine/Plugins/FX/Niagara/Source/NiagaraShader/Internal/Stateless/NiagaraStatelessSimulationShader.h
+++ b/Engine/Plugins/FX/Niagara/Source/NiagaraShader/Internal/Stateless/NiagaraStatelessSimulationShader.h
@@ -9,7 +9,7 @@
 
 namespace NiagaraStateless
 {
-	static constexpr uint32 MaxGpuSpawnInfos = 4;
+	static constexpr uint32 MaxGpuSpawnInfos = 8;
 
 	BEGIN_SHADER_PARAMETER_STRUCT(FSpawnInfoShaderParameters, NIAGARASHADER_API)
 		SHADER_PARAMETER_SCALAR_ARRAY(uint32,	SpawnInfo_NumActive, [MaxGpuSpawnInfos])
