- Change include to not use the monolithic header

#rnx

diff --git a/Engine/Plugins/FX/Niagara/Source/Niagara/Private/Stateless/NiagaraStatelessEmitterData.cpp b/Engine/Plugins/FX/Niagara/Source/Niagara/Private/Stateless/NiagaraStatelessEmitterData.cpp
--- a/Engine/Plugins/FX/Niagara/Source/Niagara/Private/Stateless/NiagaraStatelessEmitterData.cpp
+++ b/Engine/Plugins/FX/Niagara/Source/Niagara/Private/Stateless/NiagaraStatelessEmitterData.cpp
@@ -5,7 +5,7 @@
 #include "Stateless/NiagaraStatelessParticleSimExecData.h"
 #include "Stateless/NiagaraStatelessSimulationShader.h"
 
-#include "Engine.h"
+#include "Engine/Engine.h"
 
 void FNiagaraStatelessEmitterData::FDeleter::operator()(FNiagaraStatelessEmitterData* EmitterData) const
 {
