Added icons for emitter and script context actions

diff --git a/Engine/Plugins/FX/Niagara/Source/NiagaraEditor/Private/AssetDefinitions/AssetDefinition_NiagaraEmitter.cpp b/Engine/Plugins/FX/Niagara/Source/NiagaraEditor/Private/AssetDefinitions/AssetDefinition_NiagaraEmitter.cpp
--- a/Engine/Plugins/FX/Niagara/Source/NiagaraEditor/Private/AssetDefinitions/AssetDefinition_NiagaraEmitter.cpp
+++ b/Engine/Plugins/FX/Niagara/Source/NiagaraEditor/Private/AssetDefinitions/AssetDefinition_NiagaraEmitter.cpp
@@ -3,6 +3,7 @@
 #include "AssetDefinition_NiagaraEmitter.h"
 
 #include "AssetToolsModule.h"
+#include "Styling/SlateIconFinder.h"
 #include "ContentBrowserMenuContexts.h"
 #include "NiagaraEmitter.h"
 #include "NiagaraScriptSourceBase.h"
@@ -150,7 +151,7 @@ namespace MenuExtension_NiagaraEmitter
 					{
 						const TAttribute<FText> Label = LOCTEXT("Emitter_NewNiagaraSystem", "Create Niagara System");
 						const TAttribute<FText> ToolTip = LOCTEXT("Emitter_NewNiagaraSystemTooltip", "Creates a niagara system using this emitter as a base.");
-						const FSlateIcon Icon = FSlateIcon(FAppStyle::GetAppStyleSetName(), "ClassIcon.ParticleSystem");
+						const FSlateIcon Icon = FSlateIconFinder::FindIconForClass(UNiagaraSystem::StaticClass());
 
 						FToolUIAction UIAction;
 						UIAction.ExecuteAction = FToolMenuExecuteAction::CreateStatic(&ExecuteNewNiagaraSystem);
@@ -159,7 +160,7 @@ namespace MenuExtension_NiagaraEmitter
 					{
 						const TAttribute<FText> Label = LOCTEXT("Emitter_CreateDuplicateParent", "Create Duplicate Parent");
 						const TAttribute<FText> ToolTip = LOCTEXT("Emitter_CreateDuplicateParentTooltip", "Duplicate this emitter and set this emitter's parent to the new emitter.");
-						const FSlateIcon Icon = FSlateIcon();
+						const FSlateIcon Icon = FSlateIcon(FAppStyle::GetAppStyleSetName(), "Icons.Duplicate");
 
 						FToolUIAction UIAction;
 						UIAction.ExecuteAction = FToolMenuExecuteAction::CreateStatic(&ExecuteCreateDuplicateParent);
@@ -168,7 +169,7 @@ namespace MenuExtension_NiagaraEmitter
 					{
 						const TAttribute<FText> Label = LOCTEXT("MarkDependentCompilableAssetsDirtyLabel", "Mark Dependent Compilable Assets Dirty");
 						const TAttribute<FText> ToolTip = LOCTEXT("MarkDependentCompilableAssetsDirtyToolTip", "Finds all niagara assets which depend on this asset either directly or indirectly, and marks them dirty so they can be saved with the latest version.");
-						const FSlateIcon Icon = FSlateIcon();
+						const FSlateIcon Icon = FSlateIcon(FAppStyle::GetAppStyleSetName(), "Icons.DirtyBadge");
 					
 						FToolUIAction UIAction;
 						UIAction.ExecuteAction = FToolMenuExecuteAction::CreateStatic(&ExecuteMarkDependentCompilableAssetsDirty);
diff --git a/Engine/Plugins/FX/Niagara/Source/NiagaraEditor/Private/AssetDefinitions/AssetDefinition_NiagaraScript.cpp b/Engine/Plugins/FX/Niagara/Source/NiagaraEditor/Private/AssetDefinitions/AssetDefinition_NiagaraScript.cpp
--- a/Engine/Plugins/FX/Niagara/Source/NiagaraEditor/Private/AssetDefinitions/AssetDefinition_NiagaraScript.cpp
+++ b/Engine/Plugins/FX/Niagara/Source/NiagaraEditor/Private/AssetDefinitions/AssetDefinition_NiagaraScript.cpp
@@ -89,7 +89,7 @@ namespace MenuExtension_NiagaraScript
 				{
 					const TAttribute<FText> Label = LOCTEXT("MarkDependentCompilableAssetsDirtyLabel", "Mark Dependent Compilable Assets Dirty");
 					const TAttribute<FText> ToolTip = LOCTEXT("MarkDependentCompilableAssetsDirtyToolTip", "Finds all niagara assets which depend on this asset either directly or indirectly, and marks them dirty so they can be saved with the latest version.");
-					const FSlateIcon Icon = FSlateIcon();
+					const FSlateIcon Icon = FSlateIcon(FAppStyle::GetAppStyleSetName(), "Icons.DirtyBadge");
 					
 					FToolUIAction UIAction;
 					UIAction.ExecuteAction = FToolMenuExecuteAction::CreateStatic(&ExecuteMarkDependentCompilableAssetsDirty);
