Fixed a bug where the "Show Only Issues" filter toggle would hide selected modules and could not be turned off again.

diff --git a/Engine/Plugins/FX/Niagara/Source/NiagaraEditorWidgets/Private/SNiagaraStack.cpp b/Engine/Plugins/FX/Niagara/Source/NiagaraEditorWidgets/Private/SNiagaraStack.cpp
--- a/Engine/Plugins/FX/Niagara/Source/NiagaraEditorWidgets/Private/SNiagaraStack.cpp
+++ b/Engine/Plugins/FX/Niagara/Source/NiagaraEditorWidgets/Private/SNiagaraStack.cpp
@@ -59,6 +59,8 @@
 #include "Styling/StyleColors.h"
 #include "SResetToDefaultPropertyEditor.h"
 #include "Stack/SNiagaraStackScriptHierarchyRoot.h"
+
+#include "ViewModels/NiagaraSystemSelectionViewModel.h"
 #include "ViewModels/Stack/NiagaraStackNote.h"
 
 #define LOCTEXT_NAMESPACE "NiagaraStack"
@@ -580,11 +582,31 @@ TSharedRef<SWidget> SNiagaraStack::GetViewOptionsMenu() const
 
 	MenuBuilder.AddMenuEntry(
 		LOCTEXT("ShowIssuesLabel", "Show Only Issues"),
-		LOCTEXT("ShowIssuesToolTip", "Hides all modules except those that have unresolved issues."),
+		LOCTEXT("ShowIssuesToolTip", "Hides all modules except those that have unresolved issues. This action is only enabled when system and emitter nodes are selected in the graph, as it could potentially hide any selected modules."),
 		FSlateIcon(),
 		FUIAction(
 			FExecuteAction::CreateLambda([this]() { StackViewModel->SetShowOnlyIssues(!StackViewModel->GetShowOnlyIssues()); }),
-			FCanExecuteAction(),
+			FCanExecuteAction::CreateLambda([this]()
+			{
+				if (UNiagaraStackEntry* RootEntry = StackViewModel->GetRootEntry())
+				{
+					if (UNiagaraSystemSelectionViewModel* SelectionViewModel = RootEntry->GetSystemViewModel()->GetSelectionViewModel())
+					{
+						// only allow the user to toggle this option when a system or emitter node is selected
+						TArray<UNiagaraStackEntry*> SelectedEntries;
+						SelectionViewModel->GetSelectedEntries(SelectedEntries);
+						for (UNiagaraStackEntry* SelectedEntry : SelectedEntries)
+						{
+							if (!SelectedEntry->IsA<UNiagaraStackRoot>())
+							{
+								return false;
+							}
+						}
+						return SelectedEntries.Num() > 0;
+					}
+				}
+				return false;
+			}),
 			FGetActionCheckState::CreateLambda([this]() { return StackViewModel->GetShowOnlyIssues() ? ECheckBoxState::Checked : ECheckBoxState::Unchecked; })),
 		NAME_None, EUserInterfaceActionType::Check);
 
