Fixed shader parameter name typo in niagara stateless module

#rnx

diff --git a/Engine/Plugins/FX/Niagara/Shaders/Private/Stateless/Modules/NiagaraStatelessModule_SolveVelocitiesAndForces.ush b/Engine/Plugins/FX/Niagara/Shaders/Private/Stateless/Modules/NiagaraStatelessModule_SolveVelocitiesAndForces.ush
--- a/Engine/Plugins/FX/Niagara/Shaders/Private/Stateless/Modules/NiagaraStatelessModule_SolveVelocitiesAndForces.ush
+++ b/Engine/Plugins/FX/Niagara/Shaders/Private/Stateless/Modules/NiagaraStatelessModule_SolveVelocitiesAndForces.ush
@@ -21,7 +21,7 @@ float	SolveVelocitiesAndForces_ConeAngleScale;
 float	SolveVelocitiesAndForces_ConeAngleBias;
 float	SolveVelocitiesAndForces_ConeVelocityFalloff;
 
-uint	SolveVelocitiesAndForces_PontVelocityEnabled;
+uint	SolveVelocitiesAndForces_PointVelocityEnabled;
 float	SolveVelocitiesAndForces_PointVelocityScale;
 float	SolveVelocitiesAndForces_PointVelocityBias;
 float3	SolveVelocitiesAndForces_PointOrigin;
@@ -84,7 +84,7 @@ void SolveVelocitiesAndForces_Simulate(inout FStatelessParticle Particle)
 		ModuleData.Velocity += RotateVectorByQuat(Direction, SolveVelocitiesAndForces_ConeQuat) * VelocityScale;
 	}
 
-	if ( SolveVelocitiesAndForces_PontVelocityEnabled != 0 )
+	if ( SolveVelocitiesAndForces_PointVelocityEnabled != 0 )
 	{
 		const float3 FallbackDir = RandomUnitFloat3(8);
 		const float3 Delta = Particle.Position - SolveVelocitiesAndForces_PointOrigin;
diff --git a/Engine/Plugins/FX/Niagara/Source/Niagara/Private/Stateless/Modules/NiagaraStatelessModule_SolveVelocitiesAndForces.cpp b/Engine/Plugins/FX/Niagara/Source/Niagara/Private/Stateless/Modules/NiagaraStatelessModule_SolveVelocitiesAndForces.cpp
--- a/Engine/Plugins/FX/Niagara/Source/Niagara/Private/Stateless/Modules/NiagaraStatelessModule_SolveVelocitiesAndForces.cpp
+++ b/Engine/Plugins/FX/Niagara/Source/Niagara/Private/Stateless/Modules/NiagaraStatelessModule_SolveVelocitiesAndForces.cpp
@@ -166,7 +166,7 @@ void UNiagaraStatelessModule_SolveVelocitiesAndForces::SetShaderParameters(const
 		Parameters->SolveVelocitiesAndForces_ConeVelocityFalloff	= PhysicsData.ConeVelocityFalloff;
 	}
 
-	Parameters->SolveVelocitiesAndForces_PontVelocityEnabled = PhysicsData.bPointVelocity ? 1 : 0;
+	Parameters->SolveVelocitiesAndForces_PointVelocityEnabled = PhysicsData.bPointVelocity ? 1 : 0;
 	{
 		SetShaderParameterContext.ConvertRangeToScaleBias(PhysicsData.PointVelocityRange, Parameters->SolveVelocitiesAndForces_PointVelocityScale, Parameters->SolveVelocitiesAndForces_PointVelocityBias);
 		Parameters->SolveVelocitiesAndForces_PointOrigin = SpaceTransforms.TransformPosition(PhysicsData.PointCoordinateSpace, PhysicsData.PointOrigin);
diff --git a/Engine/Plugins/FX/Niagara/Source/NiagaraShader/Internal/Stateless/NiagaraStatelessModuleShaderParameters.h b/Engine/Plugins/FX/Niagara/Source/NiagaraShader/Internal/Stateless/NiagaraStatelessModuleShaderParameters.h
--- a/Engine/Plugins/FX/Niagara/Source/NiagaraShader/Internal/Stateless/NiagaraStatelessModuleShaderParameters.h
+++ b/Engine/Plugins/FX/Niagara/Source/NiagaraShader/Internal/Stateless/NiagaraStatelessModuleShaderParameters.h
@@ -165,7 +165,7 @@ namespace NiagaraStateless
 		SHADER_PARAMETER(float,			SolveVelocitiesAndForces_ConeAngleBias)
 		SHADER_PARAMETER(float,			SolveVelocitiesAndForces_ConeVelocityFalloff)
 
-		SHADER_PARAMETER(uint32,		SolveVelocitiesAndForces_PontVelocityEnabled)
+		SHADER_PARAMETER(uint32,		SolveVelocitiesAndForces_PointVelocityEnabled)
 		SHADER_PARAMETER(float,			SolveVelocitiesAndForces_PointVelocityScale)
 		SHADER_PARAMETER(float,			SolveVelocitiesAndForces_PointVelocityBias)
 		SHADER_PARAMETER(FVector3f,		SolveVelocitiesAndForces_PointOrigin)
