- Small optimizations to skip disabled emitters with various validation rules

#rnx

diff --git a/Engine/Plugins/FX/Niagara/Source/NiagaraEditor/Private/NiagaraValidationRules.cpp b/Engine/Plugins/FX/Niagara/Source/NiagaraEditor/Private/NiagaraValidationRules.cpp
--- a/Engine/Plugins/FX/Niagara/Source/NiagaraEditor/Private/NiagaraValidationRules.cpp
+++ b/Engine/Plugins/FX/Niagara/Source/NiagaraEditor/Private/NiagaraValidationRules.cpp
@@ -269,7 +269,7 @@ FString NiagaraValidation::GetPlatformConflictsString(TConstArrayView<FNiagaraPl
 
 FString NiagaraValidation::GetPlatformConflictsString(const FNiagaraPlatformSet& PlatformSetA, const FNiagaraPlatformSet& PlatformSetB, int MaxPlatformsToShow)
 {
-	TArray<const FNiagaraPlatformSet*> CheckSets;
+	TArray<const FNiagaraPlatformSet*, TInlineAllocator<2>> CheckSets;
 	CheckSets.Add(&PlatformSetA);
 	CheckSets.Add(&PlatformSetB);
 
@@ -821,6 +821,11 @@ void UNiagaraValidationRule_RendererSortingEnabled::CheckValidity(const FNiagara
 	TArray<TSharedRef<FNiagaraEmitterHandleViewModel>> EmitterHandleViewModels = Context.ViewModel->GetEmitterHandleViewModels();
 	for (TSharedRef<FNiagaraEmitterHandleViewModel> EmitterHandleModel : EmitterHandleViewModels)
 	{
+		if (EmitterHandleModel->GetIsEnabled() == false)
+		{
+			continue;
+		}
+
 		FVersionedNiagaraEmitterData* EmitterData = EmitterHandleModel.Get().GetEmitterHandle()->GetEmitterData();
 
 		const FString PlatformConflictsString = NiagaraValidation::GetPlatformConflictsString(Platforms, EmitterData->Platforms);
@@ -912,6 +917,11 @@ void UNiagaraValidationRule_RibbonRenderer::CheckValidity(const FNiagaraValidati
 	UNiagaraSystem& System = Context.ViewModel->GetSystem();
 	for (TSharedRef<FNiagaraEmitterHandleViewModel> EmitterHandleModel : Context.ViewModel->GetEmitterHandleViewModels())
 	{
+		if (EmitterHandleModel->GetIsEnabled() == false)
+		{
+			continue;
+		}
+
 		FNiagaraEmitterHandle* EmitterHandle = EmitterHandleModel.Get().GetEmitterHandle();
 		FVersionedNiagaraEmitterData* EmitterData = EmitterHandle->GetEmitterData();
 
@@ -1498,7 +1508,10 @@ void UNiagaraValidationRule_NoMapForOnCpu::CheckValidity(const FNiagaraValidatio
 	TArray<TSharedRef<FNiagaraEmitterHandleViewModel>> EmitterHandleViewModels = Context.ViewModel->GetEmitterHandleViewModels();
 	for (TSharedRef<FNiagaraEmitterHandleViewModel> EmitterHandleModel : EmitterHandleViewModels)
 	{
-		AllModules.Append(NiagaraValidation::GetStackEntries<UNiagaraStackModuleItem>(EmitterHandleModel.Get().GetEmitterStackViewModel()));
+		if (EmitterHandleModel->GetIsEnabled())
+		{
+			AllModules.Append(NiagaraValidation::GetStackEntries<UNiagaraStackModuleItem>(EmitterHandleModel.Get().GetEmitterStackViewModel()));
+		}
 	}
 
 	for (UNiagaraStackModuleItem* Module : AllModules)
diff --git a/Engine/Plugins/FX/Niagara/Source/NiagaraEditor/Public/NiagaraValidationRules.h b/Engine/Plugins/FX/Niagara/Source/NiagaraEditor/Public/NiagaraValidationRules.h
--- a/Engine/Plugins/FX/Niagara/Source/NiagaraEditor/Public/NiagaraValidationRules.h
+++ b/Engine/Plugins/FX/Niagara/Source/NiagaraEditor/Public/NiagaraValidationRules.h
@@ -56,7 +56,10 @@ namespace NiagaraValidation
 		TArray<TSharedRef<FNiagaraEmitterHandleViewModel>> EmitterHandleViewModels = ViewModel->GetEmitterHandleViewModels();
 		for (TSharedRef<FNiagaraEmitterHandleViewModel> EmitterHandleModel : EmitterHandleViewModels)
 		{
-			Results.Append(NiagaraValidation::GetStackEntries<T>(EmitterHandleModel.Get().GetEmitterStackViewModel(), bRefresh));
+			if (EmitterHandleModel->GetIsEnabled())
+			{
+				Results.Append(NiagaraValidation::GetStackEntries<T>(EmitterHandleModel.Get().GetEmitterStackViewModel(), bRefresh));
+			}
 		}
 		return Results;
 	}
