Give explicit and stable names to NiagaraDataInterface objects, these are created at cook time and so we want stable names

#rnx

diff --git a/Engine/Plugins/FX/Niagara/Source/Niagara/Private/NiagaraComponent.cpp b/Engine/Plugins/FX/Niagara/Source/Niagara/Private/NiagaraComponent.cpp
--- a/Engine/Plugins/FX/Niagara/Source/Niagara/Private/NiagaraComponent.cpp
+++ b/Engine/Plugins/FX/Niagara/Source/Niagara/Private/NiagaraComponent.cpp
@@ -3581,13 +3581,13 @@ void UNiagaraComponent::CopyParametersFromAsset(bool bResetExistingOverrideParam
 	// DataInterfaces are unique per component so we need to make new instances if they are not already
 	const EObjectFlags DataInterfaceObjectFlags = UNiagaraDataInterface::BuildObjectFlagsForOwner(this, RF_Transactional);
 
-	for (int32 i = 0; i < AssetExposedParameters.GetDataInterfaces().Num(); ++i)
+	for (int32 I = 0; I < AssetExposedParameters.GetDataInterfaces().Num(); ++I)
 	{
-		UNiagaraDataInterface* AssetDataInterface = AssetExposedParameters.GetDataInterface(i);
-		UNiagaraDataInterface* OverrideDataInterface = OverrideParameters.GetDataInterface(i);
+		UNiagaraDataInterface* AssetDataInterface = AssetExposedParameters.GetDataInterface(I);
+		UNiagaraDataInterface* OverrideDataInterface = OverrideParameters.GetDataInterface(I);
 		if (!AssetDataInterface || !OverrideDataInterface)
 		{
-			const FNiagaraVariableBase* DIVariable = AssetExposedParameters.FindVariableFromDataInterfaceIndex(i);
+			const FNiagaraVariableBase* DIVariable = AssetExposedParameters.FindVariableFromDataInterfaceIndex(I);
 			UE_LOG(LogNiagara, Error, TEXT("null data interface found for Variable(%s) System(%s) will not run."), DIVariable ? *DIVariable->GetName().ToString() : TEXT("Unknown"), *GetNameSafe(Asset));
 			continue;
 		}
@@ -3599,9 +3599,24 @@ void UNiagaraComponent::CopyParametersFromAsset(bool bResetExistingOverrideParam
 		}
 		else
 		{
-			OverrideDataInterface = NewObject<UNiagaraDataInterface>(this, AssetDataInterface->GetClass(), NAME_None, DataInterfaceObjectFlags);
+			const UClass* DataInterfaceClass = AssetDataInterface->GetClass();
+			FName DesiredName(DataInterfaceClass->GetFName(), I + 1); // +1 here to force postfix, match default behavior of NewObject
+			// reserve the name - we are seeing collision issues here, but lets just take a deterministic name
+			// for our created objects and move collisions aside:
+			if(UObject* Collision = StaticFindObjectFast(UObject::StaticClass(), this, DesiredName))
+			{
+				Collision->Rename(
+					*MakeUniqueObjectName(
+						this, 
+						Collision->GetClass(),
+						*FString::Printf(TEXT("MOVED_ASIDE_%s"), *Collision->GetName())).ToString(), 
+					this,
+					REN_DoNotDirty | REN_DontCreateRedirectors | REN_NonTransactional);
+				Collision->ClearFlags(RF_Public); // we don't expect references to these, clearing public will generate GLEO on save
+			}
+			OverrideDataInterface = NewObject<UNiagaraDataInterface>(this, DataInterfaceClass, DesiredName, DataInterfaceObjectFlags);
 			AssetDataInterface->CopyTo(OverrideDataInterface);
-			OverrideParameters.SetDataInterface(OverrideDataInterface, i);
+			OverrideParameters.SetDataInterface(OverrideDataInterface, I);
 		}
 	}
 }
