Make DI Hlsl gen context const.

diff --git a/Engine/Plugins/FX/Niagara/Source/Niagara/Classes/NiagaraDataInterface.h b/Engine/Plugins/FX/Niagara/Source/Niagara/Classes/NiagaraDataInterface.h
--- a/Engine/Plugins/FX/Niagara/Source/Niagara/Classes/NiagaraDataInterface.h
+++ b/Engine/Plugins/FX/Niagara/Source/Niagara/Classes/NiagaraDataInterface.h
@@ -743,9 +743,9 @@ public:
 	{
 	}
 	
-	NIAGARA_API virtual void GetParameterDefinitionHLSL(FNiagaraDataInterfaceHlslGenerationContext& HlslGenContext, FString& OutHLSL);
+	NIAGARA_API virtual void GetParameterDefinitionHLSL(const FNiagaraDataInterfaceHlslGenerationContext& HlslGenContext, FString& OutHLSL);
 
-	NIAGARA_API virtual bool GetFunctionHLSL(FNiagaraDataInterfaceHlslGenerationContext& HlslGenContext, FString& OutHLSL);
+	NIAGARA_API virtual bool GetFunctionHLSL(const FNiagaraDataInterfaceHlslGenerationContext& HlslGenContext, FString& OutHLSL);
 
 	virtual void GetParameterDefinitionHLSL(const FNiagaraDataInterfaceGPUParamInfo& ParamInfo, FString& OutHLSL)
 	{
diff --git a/Engine/Plugins/FX/Niagara/Source/Niagara/Internal/DataInterface/NiagaraDataInterfaceDataChannelCommon.h b/Engine/Plugins/FX/Niagara/Source/Niagara/Internal/DataInterface/NiagaraDataInterfaceDataChannelCommon.h
--- a/Engine/Plugins/FX/Niagara/Source/Niagara/Internal/DataInterface/NiagaraDataInterfaceDataChannelCommon.h
+++ b/Engine/Plugins/FX/Niagara/Source/Niagara/Internal/DataInterface/NiagaraDataInterfaceDataChannelCommon.h
@@ -254,7 +254,7 @@ namespace NDIDataChannelUtilities
 	void SortParameters(TArray<FNiagaraVariableBase>& Parameters);
 
 #if WITH_EDITORONLY_DATA
-	void GenerateDataChannelAccessHlsl(FNiagaraDataInterfaceHlslGenerationContext& HlslGenContext, TConstArrayView<FString> CommonTemplateShaderCode, const TMap<FName, FString>& TemplateShaderMap, FString& OutHLSL);
+	void GenerateDataChannelAccessHlsl(const FNiagaraDataInterfaceHlslGenerationContext& HlslGenContext, TConstArrayView<FString> CommonTemplateShaderCode, const TMap<FName, FString>& TemplateShaderMap, FString& OutHLSL);
 #endif
 
 }
diff --git a/Engine/Plugins/FX/Niagara/Source/Niagara/Internal/DataInterface/NiagaraDataInterfaceDataChannelRead.h b/Engine/Plugins/FX/Niagara/Source/Niagara/Internal/DataInterface/NiagaraDataInterfaceDataChannelRead.h
--- a/Engine/Plugins/FX/Niagara/Source/Niagara/Internal/DataInterface/NiagaraDataInterfaceDataChannelRead.h
+++ b/Engine/Plugins/FX/Niagara/Source/Niagara/Internal/DataInterface/NiagaraDataInterfaceDataChannelRead.h
@@ -147,8 +147,8 @@ public:
 #if WITH_EDITORONLY_DATA
 	NIAGARA_API virtual bool AppendCompileHash(FNiagaraCompileHashVisitor* InVisitor) const override;
 	NIAGARA_API virtual void GetCommonHLSL(FString& OutHLSL)override;
-	NIAGARA_API virtual bool GetFunctionHLSL(FNiagaraDataInterfaceHlslGenerationContext& HlslGenContext, FString& OutHLSL) override;
-	NIAGARA_API virtual void GetParameterDefinitionHLSL(FNiagaraDataInterfaceHlslGenerationContext& HlslGenContext, FString& OutHLSL) override;
+	NIAGARA_API virtual bool GetFunctionHLSL(const FNiagaraDataInterfaceHlslGenerationContext& HlslGenContext, FString& OutHLSL) override;
+	NIAGARA_API virtual void GetParameterDefinitionHLSL(const FNiagaraDataInterfaceHlslGenerationContext& HlslGenContext, FString& OutHLSL) override;
 
 	virtual bool UpgradeFunctionCall(FNiagaraFunctionSignature& FunctionSignature)override;
 
diff --git a/Engine/Plugins/FX/Niagara/Source/Niagara/Internal/DataInterface/NiagaraDataInterfaceDataChannelWrite.h b/Engine/Plugins/FX/Niagara/Source/Niagara/Internal/DataInterface/NiagaraDataInterfaceDataChannelWrite.h
--- a/Engine/Plugins/FX/Niagara/Source/Niagara/Internal/DataInterface/NiagaraDataInterfaceDataChannelWrite.h
+++ b/Engine/Plugins/FX/Niagara/Source/Niagara/Internal/DataInterface/NiagaraDataInterfaceDataChannelWrite.h
@@ -170,8 +170,8 @@ public:
 #if WITH_EDITORONLY_DATA
 	NIAGARA_API virtual bool AppendCompileHash(FNiagaraCompileHashVisitor* InVisitor) const override;
 	NIAGARA_API virtual void GetCommonHLSL(FString& OutHLSL)override;
-	NIAGARA_API virtual bool GetFunctionHLSL(FNiagaraDataInterfaceHlslGenerationContext& HlslGenContext, FString& OutHLSL) override;
-	NIAGARA_API virtual void GetParameterDefinitionHLSL(FNiagaraDataInterfaceHlslGenerationContext& HlslGenContext, FString& OutHLSL) override;
+	NIAGARA_API virtual bool GetFunctionHLSL(const FNiagaraDataInterfaceHlslGenerationContext& HlslGenContext, FString& OutHLSL) override;
+	NIAGARA_API virtual void GetParameterDefinitionHLSL(const FNiagaraDataInterfaceHlslGenerationContext& HlslGenContext, FString& OutHLSL) override;
 
 	NIAGARA_API virtual void PostCompile()override;
 #endif
diff --git a/Engine/Plugins/FX/Niagara/Source/Niagara/Private/DataInterface/NiagaraDataInterfaceDataChannelCommon.cpp b/Engine/Plugins/FX/Niagara/Source/Niagara/Private/DataInterface/NiagaraDataInterfaceDataChannelCommon.cpp
--- a/Engine/Plugins/FX/Niagara/Source/Niagara/Private/DataInterface/NiagaraDataInterfaceDataChannelCommon.cpp
+++ b/Engine/Plugins/FX/Niagara/Source/Niagara/Private/DataInterface/NiagaraDataInterfaceDataChannelCommon.cpp
@@ -333,7 +333,7 @@ namespace NDIDataChannelUtilities
 
 #if WITH_EDITORONLY_DATA
 
-	void GenerateDataChannelAccessHlsl(FNiagaraDataInterfaceHlslGenerationContext& HlslGenContext, TConstArrayView<FString> CommonTemplateShaderCode, const TMap<FName, FString>& FunctionTemplateMap, FString& OutHLSL)
+	void GenerateDataChannelAccessHlsl(const FNiagaraDataInterfaceHlslGenerationContext& HlslGenContext, TConstArrayView<FString> CommonTemplateShaderCode, const TMap<FName, FString>& FunctionTemplateMap, FString& OutHLSL)
 	{
 		//Desc:
 		// This function iterates over all functions called for this DI in each script and builds the correct hlsl.
diff --git a/Engine/Plugins/FX/Niagara/Source/Niagara/Private/DataInterface/NiagaraDataInterfaceDataChannelRead.cpp b/Engine/Plugins/FX/Niagara/Source/Niagara/Private/DataInterface/NiagaraDataInterfaceDataChannelRead.cpp
--- a/Engine/Plugins/FX/Niagara/Source/Niagara/Private/DataInterface/NiagaraDataInterfaceDataChannelRead.cpp
+++ b/Engine/Plugins/FX/Niagara/Source/Niagara/Private/DataInterface/NiagaraDataInterfaceDataChannelRead.cpp
@@ -2023,7 +2023,7 @@ void UNiagaraDataInterfaceDataChannelRead::GetCommonHLSL(FString& OutHLSL)
 	OutHLSL.Appendf(TEXT("#include \"%s\"\n"), NDIDataChannelReadLocal::CommonShaderFile);
 }
 
-bool UNiagaraDataInterfaceDataChannelRead::GetFunctionHLSL(FNiagaraDataInterfaceHlslGenerationContext& HlslGenContext, FString& OutHLSL)
+bool UNiagaraDataInterfaceDataChannelRead::GetFunctionHLSL(const FNiagaraDataInterfaceHlslGenerationContext& HlslGenContext, FString& OutHLSL)
 {
 	return	HlslGenContext.GetFunctionInfo().DefinitionName == GET_FUNCTION_NAME_CHECKED(UNiagaraDataInterfaceDataChannelRead, Num) ||
 		HlslGenContext.GetFunctionInfo().DefinitionName == GET_FUNCTION_NAME_CHECKED(UNiagaraDataInterfaceDataChannelRead, GetNDCSpawnData) ||
@@ -2031,7 +2031,7 @@ bool UNiagaraDataInterfaceDataChannelRead::GetFunctionHLSL(FNiagaraDataInterface
 		HlslGenContext.GetFunctionInfo().DefinitionName == GET_FUNCTION_NAME_CHECKED(UNiagaraDataInterfaceDataChannelRead, Consume);
 }
 
-void UNiagaraDataInterfaceDataChannelRead::GetParameterDefinitionHLSL(FNiagaraDataInterfaceHlslGenerationContext& HlslGenContext, FString& OutHLSL)
+void UNiagaraDataInterfaceDataChannelRead::GetParameterDefinitionHLSL(const FNiagaraDataInterfaceHlslGenerationContext& HlslGenContext, FString& OutHLSL)
 {
 	Super::GetParameterDefinitionHLSL(HlslGenContext, OutHLSL);
 
diff --git a/Engine/Plugins/FX/Niagara/Source/Niagara/Private/DataInterface/NiagaraDataInterfaceDataChannelWrite.cpp b/Engine/Plugins/FX/Niagara/Source/Niagara/Private/DataInterface/NiagaraDataInterfaceDataChannelWrite.cpp
--- a/Engine/Plugins/FX/Niagara/Source/Niagara/Private/DataInterface/NiagaraDataInterfaceDataChannelWrite.cpp
+++ b/Engine/Plugins/FX/Niagara/Source/Niagara/Private/DataInterface/NiagaraDataInterfaceDataChannelWrite.cpp
@@ -1527,14 +1527,14 @@ void UNiagaraDataInterfaceDataChannelWrite::GetCommonHLSL(FString& OutHLSL)
  	OutHLSL.Appendf(TEXT("#include \"%s\"\n"), NDIDataChannelWriteLocal::CommonShaderFile);
 }
 
-bool UNiagaraDataInterfaceDataChannelWrite::GetFunctionHLSL(FNiagaraDataInterfaceHlslGenerationContext& HlslGenContext, FString& OutHLSL)
+bool UNiagaraDataInterfaceDataChannelWrite::GetFunctionHLSL(const FNiagaraDataInterfaceHlslGenerationContext& HlslGenContext, FString& OutHLSL)
 {
  	return	HlslGenContext.GetFunctionInfo().DefinitionName == GET_FUNCTION_NAME_CHECKED(UNiagaraDataInterfaceDataChannelWrite, Num) ||
  		//HlslGenContext.GetFunctionInfo().DefinitionName == GET_FUNCTION_NAME_CHECKED(UNiagaraDataInterfaceDataChannelWrite, Write) ||
  		HlslGenContext.GetFunctionInfo().DefinitionName == GET_FUNCTION_NAME_CHECKED(UNiagaraDataInterfaceDataChannelWrite, Append);
 }
 
-void UNiagaraDataInterfaceDataChannelWrite::GetParameterDefinitionHLSL(FNiagaraDataInterfaceHlslGenerationContext& HlslGenContext, FString& OutHLSL)
+void UNiagaraDataInterfaceDataChannelWrite::GetParameterDefinitionHLSL(const FNiagaraDataInterfaceHlslGenerationContext& HlslGenContext, FString& OutHLSL)
 {
 	Super::GetParameterDefinitionHLSL(HlslGenContext, OutHLSL);
 
diff --git a/Engine/Plugins/FX/Niagara/Source/Niagara/Private/NiagaraDataInterface.cpp b/Engine/Plugins/FX/Niagara/Source/Niagara/Private/NiagaraDataInterface.cpp
--- a/Engine/Plugins/FX/Niagara/Source/Niagara/Private/NiagaraDataInterface.cpp
+++ b/Engine/Plugins/FX/Niagara/Source/Niagara/Private/NiagaraDataInterface.cpp
@@ -235,14 +235,14 @@ bool UNiagaraDataInterface::CopyToInternal(UNiagaraDataInterface* Destination) c
 
 #if WITH_EDITOR
 
-void UNiagaraDataInterface::GetParameterDefinitionHLSL(FNiagaraDataInterfaceHlslGenerationContext& HlslGenContext, FString& OutHLSL)
+void UNiagaraDataInterface::GetParameterDefinitionHLSL(const FNiagaraDataInterfaceHlslGenerationContext& HlslGenContext, FString& OutHLSL)
 {
 	//Base implementation will call into the legacy functions so that all DIs need not be updated immediately.
 	//Do not implement both current and legacy functions in any DI or it's parent classes.
 	GetParameterDefinitionHLSL(HlslGenContext.ParameterInfo, OutHLSL);
 }
 
-bool UNiagaraDataInterface::GetFunctionHLSL(FNiagaraDataInterfaceHlslGenerationContext& HlslGenContext, FString& OutHLSL)
+bool UNiagaraDataInterface::GetFunctionHLSL(const FNiagaraDataInterfaceHlslGenerationContext& HlslGenContext, FString& OutHLSL)
 {
 	//Base implementation will call into the legacy functions so that all DIs need not be updated immediately.
 	//Do not implement both current and legacy functions in any DI or it's parent classes.
diff --git a/Engine/Plugins/FX/Niagara/Source/Niagara/Public/NiagaraDataInterfaceUtilities.h b/Engine/Plugins/FX/Niagara/Source/Niagara/Public/NiagaraDataInterfaceUtilities.h
--- a/Engine/Plugins/FX/Niagara/Source/Niagara/Public/NiagaraDataInterfaceUtilities.h
+++ b/Engine/Plugins/FX/Niagara/Source/Niagara/Public/NiagaraDataInterfaceUtilities.h
@@ -76,11 +76,11 @@ struct FNiagaraDataInterfaceHlslGenerationContext
 
 	const FNiagaraDataInterfaceGeneratedFunction& GetFunctionInfo() const { return ParameterInfo.GeneratedFunctions[FunctionInstanceIndex]; }
 
-	FString GetFunctionSignatureSymbol(const FNiagaraFunctionSignature& Signature) { return GetFunctionSignatureSymbolDelegate.Execute(Signature); }
-	FString GetStructHlslTypeName(const FNiagaraTypeDefinition& Type) { return GetStructHlslTypeNameDelegate.Execute(Type); }
-	FString GetPropertyHlslTypeName(const FProperty* Property) { return GetPropertyHlslTypeNameDelegate.Execute(Property); }
-	FString GetSanitizedSymbolName(FStringView SymbolName, bool bCollapsNamespaces = false) { return GetSanitizedSymbolNameDelegate.Execute(SymbolName, bCollapsNamespaces); }
-	FString GetHlslDefaultForType(const FNiagaraTypeDefinition& Type) { return GetHlslDefaultForTypeDelegate.Execute(Type); }
+	FString GetFunctionSignatureSymbol(const FNiagaraFunctionSignature& Signature)const { return GetFunctionSignatureSymbolDelegate.Execute(Signature); }
+	FString GetStructHlslTypeName(const FNiagaraTypeDefinition& Type)const { return GetStructHlslTypeNameDelegate.Execute(Type); }
+	FString GetPropertyHlslTypeName(const FProperty* Property)const { return GetPropertyHlslTypeNameDelegate.Execute(Property); }
+	FString GetSanitizedSymbolName(FStringView SymbolName, bool bCollapsNamespaces = false)const { return GetSanitizedSymbolNameDelegate.Execute(SymbolName, bCollapsNamespaces); }
+	FString GetHlslDefaultForType(const FNiagaraTypeDefinition& Type)const { return GetHlslDefaultForTypeDelegate.Execute(Type); }
 };
 
 #endif
\ No newline at end of file
