- Avoid flushing on-demand shader compilations when calling WaitForCompilationComplete
- This fixes a stall post live coding modification where it will wait on all compilation jobs.  This would previously flush out all deferred compilation requests from load and could result in a long stall

#rnx
#tests Load into a project which loads a lot of Niagara systems but doesn't render any yet.  Made a code modification and run live coding, before it will compile all the systems, after it will not.

diff --git a/Engine/Plugins/FX/Niagara/Source/Niagara/Classes/NiagaraSystem.h b/Engine/Plugins/FX/Niagara/Source/Niagara/Classes/NiagaraSystem.h
--- a/Engine/Plugins/FX/Niagara/Source/Niagara/Classes/NiagaraSystem.h
+++ b/Engine/Plugins/FX/Niagara/Source/Niagara/Classes/NiagaraSystem.h
@@ -559,6 +559,9 @@ public:
 	/** Returns information on the availability of performance stat scopes */
 	NIAGARA_API ENiagaraPerformanceStateScopeMode GetPerformanceStateScopeMode() const;
 
+	/** Do we have any pending compilation requests or not */
+	bool NeedsRequestCompile() const { return RequestCompileStatus != ERequestCompileStatus::None; }
+
 protected:
 	/** When enable constant values are baked into the scripts while editing the system, this will increase iteration times but improve performance. */
 	UPROPERTY(EditAnywhere, AdvancedDisplay, Category = "Performance", meta=(DisplayName="Bake Rapid Iteration Parameters During Edit"))
@@ -591,12 +594,18 @@ protected:
 	UPROPERTY(EditAnywhere, AdvancedDisplay, Category = "Performance", meta = (DisplayName="Disable Debug Switches"))
 	uint32 bDisableDebugSwitchesOnCook : 1;
 
-	/* When set the system needs to compile before it can be activated. */
-	uint32 bNeedsRequestCompile : 1;
-
 	/** When enabled we compile for the edit path, which can result in slower system performance but faster editor responsiveness. */
 	uint32 bCompileForEdit : 1;
 
+	/** Used to track the status of compilation request, when using on demand they will initially be in RequestPendingOnDemand. */
+	enum class ERequestCompileStatus
+	{
+		None,
+		RequestPending,
+		RequestPendingOnDemand,
+	};
+	ERequestCompileStatus RequestCompileStatus = ERequestCompileStatus::None;
+
 public:
 	/** Subscriptions to definitions of parameters. */
 	UPROPERTY()
@@ -717,10 +726,6 @@ public:
 	/** Cache data & accessors from the compiled data, allows us to avoid per instance. */
 	NIAGARA_API void CacheFromCompiledData();
 
-#if WITH_EDITORONLY_DATA
-	bool NeedsRequestCompile() const { return bNeedsRequestCompile; }
-#endif
-
 	FORCEINLINE TConstArrayView<FNiagaraEmitterExecutionIndex> GetEmitterExecutionOrder() const { return MakeArrayView(EmitterExecutionOrder); }
 	FORCEINLINE TConstArrayView<FNiagaraRendererExecutionIndex> GetRendererPostTickOrder() const { return MakeArrayView(RendererPostTickOrder); }
 	FORCEINLINE TConstArrayView<FNiagaraRendererExecutionIndex> GetRendererCompletionOrder() const { return MakeArrayView(RendererCompletionOrder); }
diff --git a/Engine/Plugins/FX/Niagara/Source/Niagara/Private/NiagaraSystem.cpp b/Engine/Plugins/FX/Niagara/Source/Niagara/Private/NiagaraSystem.cpp
--- a/Engine/Plugins/FX/Niagara/Source/Niagara/Private/NiagaraSystem.cpp
+++ b/Engine/Plugins/FX/Niagara/Source/Niagara/Private/NiagaraSystem.cpp
@@ -260,7 +260,6 @@ UNiagaraSystem::UNiagaraSystem(const FObjectInitializer& ObjectInitializer)
 , bIgnoreParticleReadsForAttributeTrim(false)
 , bDisableDebugSwitches(false)
 , bDisableDebugSwitchesOnCook(true)
-, bNeedsRequestCompile(false)
 , bCompileForEdit(false)
 #endif
 , bSupportLargeWorldCoordinates(true)
@@ -375,7 +374,7 @@ void UNiagaraSystem::BeginCacheForCookedPlatformData(const ITargetPlatform *Targ
 	
 	EnsureFullyLoaded();
 
-	if (bNeedsRequestCompile)
+	if (NeedsRequestCompile())
 	{
 		RequestCompile(false, nullptr, TargetPlatform);
 	}
@@ -760,7 +759,7 @@ void UNiagaraSystem::UpdateSystemAfterLoad()
 
 #if WITH_EDITORONLY_DATA
 	// check if the system needs to be compiled and start a compile task if necessary
-	bNeedsRequestCompile = false;
+	RequestCompileStatus = ERequestCompileStatus::None;
 	if (GetOutermost()->HasAnyPackageFlags(EPackageFlags::PKG_Cooked) == false)
 	{
 		bool bSystemScriptsAreSynchronized = true;
@@ -790,7 +789,8 @@ void UNiagaraSystem::UpdateSystemAfterLoad()
 
 		if (bSystemScriptsAreSynchronized == false || bEmitterScriptsAreSynchronized == false)
 		{
-			bNeedsRequestCompile = true;
+			RequestCompileStatus = ERequestCompileStatus::RequestPendingOnDemand;
+
 			if (IsRunningCommandlet())
 			{
 				// Call modify here so that the system will resave the compile ids and script vm when running the resave
@@ -1947,7 +1947,7 @@ void UNiagaraSystem::GetAssetRegistryTagMetadata(TMap<FName, FAssetRegistryTagMe
 
 bool UNiagaraSystem::HasOutstandingCompilationRequests(bool bIncludingGPUShaders) const
 {
-	if (bNeedsRequestCompile)
+	if (NeedsRequestCompile())
 	{
 		return true;
 	}
@@ -2964,7 +2964,7 @@ void UNiagaraSystem::SetCompileForEdit(bool bNewCompileForEdit)
 	if (bNewCompileForEdit != bCompileForEdit)
 	{
 		bCompileForEdit = bNewCompileForEdit;
-		bNeedsRequestCompile = true;
+		RequestCompileStatus = ERequestCompileStatus::RequestPending;
 	}
 }
 
@@ -2975,7 +2975,13 @@ void UNiagaraSystem::WaitForCompilationComplete(bool bIncludingGPUShaders, bool
 	TRACE_CPUPROFILER_EVENT_SCOPE_TEXT_ON_CHANNEL(*GetPathName(), NiagaraChannel);
 	COOK_STAT(FScopedDurationTimer DurationTimer(NiagaraScriptCookStats::NiagaraSystemWaitForCompilationCompleteTime));
 
-	if (bNeedsRequestCompile)
+	// Early out as we have never issued a compilation request so can avoid doing so
+	if ( RequestCompileStatus == ERequestCompileStatus::RequestPendingOnDemand )
+	{
+		return;
+	}
+
+	if (NeedsRequestCompile())
 	{
 		RequestCompile(false);
 	}
@@ -3060,7 +3066,7 @@ bool UNiagaraSystem::HasActiveCompilations() const
 
 bool UNiagaraSystem::PollForCompilationComplete(bool bFlushRequestCompile)
 {
-	if (bNeedsRequestCompile && bFlushRequestCompile)
+	if (NeedsRequestCompile() && bFlushRequestCompile)
 	{
 		RequestCompile(false);
 	}
@@ -3568,7 +3574,7 @@ bool UNiagaraSystem::RequestCompile(bool bForce, FNiagaraSystemUpdateContext* Op
 	check(IsInGameThread());
 	TRACE_CPUPROFILER_EVENT_SCOPE(UNiagaraSystem::RequestCompile)
 
-	bNeedsRequestCompile = false;
+	RequestCompileStatus = ERequestCompileStatus::None;
 
 	// We remove emitters and scripts on dedicated servers, so skip further work.
 	const bool bIsDedicatedServer = !GIsClient && GIsServer;
