- Add a method to get the sanitized hlsl parameters to FNiagaraDataInterfaceHlslGenerationContext

#tests none required

diff --git a/Engine/Plugins/FX/Niagara/Source/Niagara/Public/NiagaraDataInterfaceUtilities.h b/Engine/Plugins/FX/Niagara/Source/Niagara/Public/NiagaraDataInterfaceUtilities.h
--- a/Engine/Plugins/FX/Niagara/Source/Niagara/Public/NiagaraDataInterfaceUtilities.h
+++ b/Engine/Plugins/FX/Niagara/Source/Niagara/Public/NiagaraDataInterfaceUtilities.h
@@ -53,6 +53,7 @@ namespace FNiagaraDataInterfaceUtilities
 /** Helper context object helping to facilitate data interfaces building their hlsl shader code for GPU simulations. */
 struct FNiagaraDataInterfaceHlslGenerationContext
 {
+	DECLARE_DELEGATE_RetVal_OneParam(FString, FGetSanitizedFunctionParameters, const FNiagaraFunctionSignature& /*Signature*/)
 	DECLARE_DELEGATE_RetVal_OneParam(FString, FGetFunctionSignatureSymbol, const FNiagaraFunctionSignature& /*Signature*/)
 	DECLARE_DELEGATE_RetVal_OneParam(FString, FGetStructHlslTypeName, const FNiagaraTypeDefinition& /*Type*/)
 	DECLARE_DELEGATE_RetVal_OneParam(FString, FGetPropertyHlslTypeName, const FProperty* /*Property*/)
@@ -68,6 +69,7 @@ struct FNiagaraDataInterfaceHlslGenerationContext
 	TArrayView<const FNiagaraFunctionSignature> Signatures;
 	int32 FunctionInstanceIndex = INDEX_NONE;
 
+	FGetSanitizedFunctionParameters GetSanitizedFunctionParametersDelegate;
 	FGetFunctionSignatureSymbol GetFunctionSignatureSymbolDelegate;
 	FGetStructHlslTypeName GetStructHlslTypeNameDelegate;
 	FGetPropertyHlslTypeName GetPropertyHlslTypeNameDelegate;
@@ -75,7 +77,9 @@ struct FNiagaraDataInterfaceHlslGenerationContext
 	FGetHlslDefaultForType GetHlslDefaultForTypeDelegate;
 
 	const FNiagaraDataInterfaceGeneratedFunction& GetFunctionInfo() const { return ParameterInfo.GeneratedFunctions[FunctionInstanceIndex]; }
+	const FNiagaraFunctionSignature& GetFunctionSignature() const { return Signatures[FunctionInstanceIndex]; }
 
+	FString GetSanitizedFunctionParameters(const FNiagaraFunctionSignature& Signature) const { return GetSanitizedFunctionParametersDelegate.Execute(Signature); }
 	FString GetFunctionSignatureSymbol(const FNiagaraFunctionSignature& Signature)const { return GetFunctionSignatureSymbolDelegate.Execute(Signature); }
 	FString GetStructHlslTypeName(const FNiagaraTypeDefinition& Type)const { return GetStructHlslTypeNameDelegate.Execute(Type); }
 	FString GetPropertyHlslTypeName(const FProperty* Property)const { return GetPropertyHlslTypeNameDelegate.Execute(Property); }
diff --git a/Engine/Plugins/FX/Niagara/Source/NiagaraEditor/Private/NiagaraHlslTranslator.cpp b/Engine/Plugins/FX/Niagara/Source/NiagaraEditor/Private/NiagaraHlslTranslator.cpp
--- a/Engine/Plugins/FX/Niagara/Source/NiagaraEditor/Private/NiagaraHlslTranslator.cpp
+++ b/Engine/Plugins/FX/Niagara/Source/NiagaraEditor/Private/NiagaraHlslTranslator.cpp
@@ -381,81 +381,88 @@ FNiagaraVariable ConvertToSimulationVariable(const FNiagaraVariable& Param)
 	return Param;
 }
 
-void AppendFunctionHlslPrototypeString(FStringView InVariableName, const FNiagaraFunctionSignature& FunctionSignature, const TMap<FString, FString>& ReplacementMap, FStringBuilderBase& StringBuilder)
+// Build the specifiers + inputs to the function, i.e. <Key=Value>(in int In_Value, out int Out_Value)
+void AppendFunctionHlslParametersString(const FNiagaraFunctionSignature& FunctionSignature, const TMap<FString, FString>& ReplacementMap, FStringBuilderBase& StringBuilder)
 {
-	if (FunctionSignature.bMemberFunction)
+	// Build specifiers
+	if (FunctionSignature.FunctionSpecifiers.Num())
 	{
-		StringBuilder.Append(InVariableName);
-		StringBuilder.Append(TEXT("."));
-		AppendSanitizedSymbolName(FunctionSignature.Name, false /*bCollapse*/, ReplacementMap, StringBuilder);
-
-		// Build specifiers
-		if (FunctionSignature.FunctionSpecifiers.Num())
+		bool bNeedsComma = false;
+		StringBuilder.AppendChar(TEXT('<'));
+		for (auto SpecifierIt = FunctionSignature.FunctionSpecifiers.CreateConstIterator(); SpecifierIt; ++SpecifierIt)
 		{
-			bool bNeedsComma = false;
-			StringBuilder.AppendChar(TEXT('<'));
-			for (auto SpecifierIt = FunctionSignature.FunctionSpecifiers.CreateConstIterator(); SpecifierIt; ++SpecifierIt)
+			if (bNeedsComma)
 			{
-				if (bNeedsComma)
-				{
-					StringBuilder.Append(TEXT(", "));
-				}
-				bNeedsComma = true;
+				StringBuilder.Append(TEXT(", "));
+			}
+			bNeedsComma = true;
 
-				SpecifierIt.Key().AppendString(StringBuilder);
-				StringBuilder.AppendChar(TEXT('='));
-				StringBuilder.AppendChar(TEXT('"'));
-				if (SpecifierIt.Value().IsNone())
-				{
-					StringBuilder.Append(TEXT("Value"));
-				}
-				else
-				{
-					SpecifierIt.Value().AppendString(StringBuilder);
-				}
-				StringBuilder.AppendChar(TEXT('"'));
+			SpecifierIt.Key().AppendString(StringBuilder);
+			StringBuilder.AppendChar(TEXT('='));
+			StringBuilder.AppendChar(TEXT('"'));
+			if (SpecifierIt.Value().IsNone())
+			{
+				StringBuilder.Append(TEXT("Value"));
 			}
-			StringBuilder.AppendChar(TEXT('>'));
+			else
+			{
+				SpecifierIt.Value().AppendString(StringBuilder);
+			}
+			StringBuilder.AppendChar(TEXT('"'));
 		}
+		StringBuilder.AppendChar(TEXT('>'));
+	}
 
-		// Build function parameters
-		{
-			StringBuilder.Append(TEXT("("));
-			bool bNeedsComma = false;
+	// Build function parameters
+	{
+		StringBuilder.Append(TEXT("("));
+		bool bNeedsComma = false;
 
-			// Inputs
-			for (int i = 1; i < FunctionSignature.Inputs.Num(); ++i)
+		// Inputs
+		for (int i = 1; i < FunctionSignature.Inputs.Num(); ++i)
+		{
+			FNiagaraVariable InputVar = ConvertToSimulationVariable(FunctionSignature.Inputs[i]);
+			if (bNeedsComma)
 			{
-				FNiagaraVariable InputVar = ConvertToSimulationVariable(FunctionSignature.Inputs[i]);
-				if (bNeedsComma)
-				{
-					StringBuilder.Append(TEXT(", "));
-				}
-				bNeedsComma = true;
-
-				StringBuilder.Append(TEXT("in "));
-				AppendStructHlslTypeName(InputVar.GetType(), StringBuilder);
-				StringBuilder.Append(TEXT(" In_"));
-				AppendSanitizedSymbolName(InputVar.GetName(), false /*bCollapse*/, ReplacementMap, StringBuilder);
+				StringBuilder.Append(TEXT(", "));
 			}
+			bNeedsComma = true;
 
-			// Outputs
-			for (int i = 0; i < FunctionSignature.Outputs.Num(); ++i)
-			{
-				FNiagaraVariable OutputVar = ConvertToSimulationVariable(FunctionSignature.Outputs[i]);
-				if (bNeedsComma)
-				{
-					StringBuilder.Append(TEXT(", "));
-				}
-				bNeedsComma = true;
+			StringBuilder.Append(TEXT("in "));
+			AppendStructHlslTypeName(InputVar.GetType(), StringBuilder);
+			StringBuilder.Append(TEXT(" In_"));
+			AppendSanitizedSymbolName(InputVar.GetName(), false /*bCollapse*/, ReplacementMap, StringBuilder);
+		}
 
-				StringBuilder.Append(TEXT("out "));
-				AppendStructHlslTypeName(OutputVar.GetType(), StringBuilder);
-				StringBuilder.Append(TEXT(" Out_"));
-				AppendSanitizedSymbolName(OutputVar.GetName(), false /*bCollapse*/, ReplacementMap, StringBuilder);
+		// Outputs
+		for (int i = 0; i < FunctionSignature.Outputs.Num(); ++i)
+		{
+			FNiagaraVariable OutputVar = ConvertToSimulationVariable(FunctionSignature.Outputs[i]);
+			if (bNeedsComma)
+			{
+				StringBuilder.Append(TEXT(", "));
 			}
-			StringBuilder.Append(TEXT(");"));
+			bNeedsComma = true;
+
+			StringBuilder.Append(TEXT("out "));
+			AppendStructHlslTypeName(OutputVar.GetType(), StringBuilder);
+			StringBuilder.Append(TEXT(" Out_"));
+			AppendSanitizedSymbolName(OutputVar.GetName(), false /*bCollapse*/, ReplacementMap, StringBuilder);
 		}
+		StringBuilder.Append(TEXT(")"));
+	}
+}
+
+void AppendFunctionHlslPrototypeString(FStringView InVariableName, const FNiagaraFunctionSignature& FunctionSignature, const TMap<FString, FString>& ReplacementMap, FStringBuilderBase& StringBuilder)
+{
+	if (FunctionSignature.bMemberFunction)
+	{
+		StringBuilder.Append(InVariableName);
+		StringBuilder.Append(TEXT("."));
+		AppendSanitizedSymbolName(FunctionSignature.Name, false /*bCollapse*/, ReplacementMap, StringBuilder);
+
+		AppendFunctionHlslParametersString(FunctionSignature, ReplacementMap, StringBuilder);
+		StringBuilder.Append(TEXT(";"));
 	}
 	else
 	{
@@ -465,7 +472,6 @@ void AppendFunctionHlslPrototypeString(FStringView InVariableName, const FNiagar
 
 };
 
-
 ENiagaraScriptCompileStatus FNiagaraTranslateResults::TranslateResultsToSummary(const FNiagaraTranslateResults* TranslateResults)
 {
 	ENiagaraScriptCompileStatus SummaryStatus = ENiagaraScriptCompileStatus::NCS_Unknown;
@@ -3410,7 +3416,7 @@ void FNiagaraHlslTranslator::DefineDataInterfaceHLSL(FString& InHlslOutput)
 			ConvertCompileInfoToParamInfo(Info, DIInstanceInfo,GeneratedFunctionSignatures);
 
 			FNiagaraDataInterfaceHlslGenerationContext DIHlslGenContext(DIInstanceInfo, GeneratedFunctionSignatures);
-
+			DIHlslGenContext.GetSanitizedFunctionParametersDelegate.BindStatic(FNiagaraHlslTranslator::GenerateFunctionParametersHlsl);
 			DIHlslGenContext.GetFunctionSignatureSymbolDelegate.BindStatic(&FNiagaraHlslTranslator::GetFunctionSignatureSymbol);
 			DIHlslGenContext.GetStructHlslTypeNameDelegate.BindStatic(&FNiagaraHlslTranslator::GetStructHlslTypeName);
 			DIHlslGenContext.GetPropertyHlslTypeNameDelegate.BindStatic(&FNiagaraHlslTranslator::GetPropertyHlslTypeName);
@@ -9616,6 +9622,14 @@ FString FNiagaraHlslTranslator::GenerateFunctionHlslPrototype(FStringView InVari
 	return StringBuilder.ToString();
 }
 
+FString FNiagaraHlslTranslator::GenerateFunctionParametersHlsl(const FNiagaraFunctionSignature& FunctionSignature)
+{
+	const UNiagaraEditorSettings* EditorSettings = GetDefault<UNiagaraEditorSettings>();
+	TStringBuilder<1024> StringBuilder;
+	NiagaraHlslTranslatorImpl::AppendFunctionHlslParametersString(FunctionSignature, EditorSettings->GetHLSLKeywordReplacementsMap(), StringBuilder);
+	return StringBuilder.ToString();
+}
+
 FName FNiagaraHlslTranslator::GetDataInterfaceName(FName BaseName, const FString& UniqueEmitterName, bool bIsParameterMapDataInterface)
 {
 	if (UniqueEmitterName.IsEmpty() == false)
diff --git a/Engine/Plugins/FX/Niagara/Source/NiagaraEditor/Private/NiagaraHlslTranslator.h b/Engine/Plugins/FX/Niagara/Source/NiagaraEditor/Private/NiagaraHlslTranslator.h
--- a/Engine/Plugins/FX/Niagara/Source/NiagaraEditor/Private/NiagaraHlslTranslator.h
+++ b/Engine/Plugins/FX/Niagara/Source/NiagaraEditor/Private/NiagaraHlslTranslator.h
@@ -408,6 +408,7 @@ public:
 
 	static NIAGARAEDITOR_API FName GetDataInterfaceName(FName BaseName, const FString& UniqueEmitterName, bool bIsParameterMapDataInterface);
 	static NIAGARAEDITOR_API FString GenerateFunctionHlslPrototype(FStringView InVariableName, const FNiagaraFunctionSignature& FunctionSignature);
+	static NIAGARAEDITOR_API FString GenerateFunctionParametersHlsl(const FNiagaraFunctionSignature& FunctionSignature);
 
 	NIAGARAEDITOR_API void SetConstantByStaticVariable(int32& OutValue, const FNiagaraVariable& Var, FString* DebugString = nullptr);
 	NIAGARAEDITOR_API void SetConstantByStaticVariable(FNiagaraVariable& OutValue, const FNiagaraVariable& Var, FString* DebugString = nullptr);
