- Add a render thread deletor class that can be reused

#rnx
#tests Test SimCache Reader / Stateless Emitters

diff --git a/Engine/Plugins/FX/Niagara/Source/Niagara/Internal/NiagaraRenderThreadDeletor.h b/Engine/Plugins/FX/Niagara/Source/Niagara/Internal/NiagaraRenderThreadDeletor.h
new file mode 100644
--- /dev/null
+++ b/Engine/Plugins/FX/Niagara/Source/Niagara/Internal/NiagaraRenderThreadDeletor.h
@@ -0,0 +1,30 @@
+// Copyright Epic Games, Inc. All Rights Reserved.
+#pragma once
+
+#include "NiagaraCommon.h"
+#include "RHICommandList.h"
+#include "RenderingThread.h"
+
+// Use with unique / shared pointers to ensure the class is deleted on the rendering thread once no references exist
+// Example usage: MakeShareable(new FMyClass(), FNiagaraRenderThreadDeletor<FMyClass>())
+template<typename TType>
+struct FNiagaraRenderThreadDeletor
+{
+	void operator()(TType* ObjectToDelete) const
+	{
+		if (IsInRenderingThread())
+		{
+			delete ObjectToDelete;
+		}
+		else
+		{
+			ENQUEUE_RENDER_COMMAND(NiagaraRenderThreadDeletor)
+			(
+				[ObjectToDelete](FRHICommandListImmediate&)
+				{
+					delete ObjectToDelete;
+				}
+			);
+		}
+	}
+};
diff --git a/Engine/Plugins/FX/Niagara/Source/Niagara/Internal/Stateless/NiagaraStatelessEmitterData.h b/Engine/Plugins/FX/Niagara/Source/Niagara/Internal/Stateless/NiagaraStatelessEmitterData.h
--- a/Engine/Plugins/FX/Niagara/Source/Niagara/Internal/Stateless/NiagaraStatelessEmitterData.h
+++ b/Engine/Plugins/FX/Niagara/Source/Niagara/Internal/Stateless/NiagaraStatelessEmitterData.h
@@ -23,11 +23,6 @@ struct FInstancedStruct;
 
 struct FNiagaraStatelessEmitterData
 {
-	struct FDeleter
-	{
-		void operator()(FNiagaraStatelessEmitterData* EmitterData) const;
-	};
-
 	UE_NONCOPYABLE(FNiagaraStatelessEmitterData);
 
 	FNiagaraStatelessEmitterData() = default;
diff --git a/Engine/Plugins/FX/Niagara/Source/Niagara/Private/NiagaraSimCacheGpuResource.cpp b/Engine/Plugins/FX/Niagara/Source/Niagara/Private/NiagaraSimCacheGpuResource.cpp
--- a/Engine/Plugins/FX/Niagara/Source/Niagara/Private/NiagaraSimCacheGpuResource.cpp
+++ b/Engine/Plugins/FX/Niagara/Source/Niagara/Private/NiagaraSimCacheGpuResource.cpp
@@ -4,16 +4,6 @@
 #include "NiagaraSimCache.h"
 #include "RenderingThread.h"
 
-void FNiagaraSimCacheGpuResource::FDeletor::operator()(FNiagaraSimCacheGpuResource* Ptr)
-{
-	ENQUEUE_RENDER_COMMAND(FNiagaraSimCacheGpuResource_FDeletor)(
-		[RT_Ptr=Ptr](FRHICommandListImmediate& RHICmdList)
-		{
-			delete RT_Ptr;
-		}
-	);
-}
-
 FNiagaraSimCacheGpuResource::FNiagaraSimCacheGpuResource(UNiagaraSimCache* SimCache)
 	: WeakSimCache(SimCache)
 {
@@ -32,7 +22,7 @@ FNiagaraSimCacheGpuResource::~FNiagaraSimCacheGpuResource()
 
 FNiagaraSimCacheGpuResourcePtr FNiagaraSimCacheGpuResource::CreateResource(UNiagaraSimCache* SimCache)
 {
-	return MakeShareable(new FNiagaraSimCacheGpuResource(SimCache), FDeletor());
+	return MakeShareable(new FNiagaraSimCacheGpuResource(SimCache), FNiagaraRenderThreadDeletor<FNiagaraSimCacheGpuResource>());
 }
 
 void FNiagaraSimCacheGpuResource::BuildResource(UNiagaraSimCache* SimCache)
diff --git a/Engine/Plugins/FX/Niagara/Source/Niagara/Private/NiagaraSimCacheGpuResource.h b/Engine/Plugins/FX/Niagara/Source/Niagara/Private/NiagaraSimCacheGpuResource.h
--- a/Engine/Plugins/FX/Niagara/Source/Niagara/Private/NiagaraSimCacheGpuResource.h
+++ b/Engine/Plugins/FX/Niagara/Source/Niagara/Private/NiagaraSimCacheGpuResource.h
@@ -4,6 +4,7 @@
 
 #include "Templates/SharedPointer.h"
 
+#include "NiagaraRenderThreadDeletor.h"
 #include "NiagaraRenderer.h"
 
 class UNiagaraSimCache;
@@ -11,14 +12,9 @@ class UNiagaraSimCache;
 struct FNiagaraSimCacheGpuResource
 {
 	UE_NONCOPYABLE(FNiagaraSimCacheGpuResource)
-	friend struct FDeletor;
+	friend FNiagaraRenderThreadDeletor<FNiagaraSimCacheGpuResource>;
 
 private:
-	struct FDeletor
-	{
-		void operator()(FNiagaraSimCacheGpuResource* Pointer);
-	};
-
 	FNiagaraSimCacheGpuResource(UNiagaraSimCache* SimCache);
 	virtual ~FNiagaraSimCacheGpuResource();
 
diff --git a/Engine/Plugins/FX/Niagara/Source/Niagara/Private/Stateless/NiagaraStatelessEmitter.cpp b/Engine/Plugins/FX/Niagara/Source/Niagara/Private/Stateless/NiagaraStatelessEmitter.cpp
--- a/Engine/Plugins/FX/Niagara/Source/Niagara/Private/Stateless/NiagaraStatelessEmitter.cpp
+++ b/Engine/Plugins/FX/Niagara/Source/Niagara/Private/Stateless/NiagaraStatelessEmitter.cpp
@@ -10,6 +10,7 @@
 #include "NiagaraConstants.h"
 #include "NiagaraModule.h"
 #include "NiagaraRendererProperties.h"
+#include "NiagaraRenderThreadDeletor.h"
 #include "NiagaraSettings.h"
 #include "NiagaraSystem.h"
 
@@ -295,7 +296,7 @@ const UNiagaraStatelessEmitterTemplate* UNiagaraStatelessEmitter::GetEmitterTemp
 
 void UNiagaraStatelessEmitter::CacheFromCompiledData()
 {
-	StatelessEmitterData = MakeShareable(new FNiagaraStatelessEmitterData(), FNiagaraStatelessEmitterData::FDeleter());
+	StatelessEmitterData = MakeShareable(new FNiagaraStatelessEmitterData(), FNiagaraRenderThreadDeletor<FNiagaraStatelessEmitterData>());
 
 #if WITH_NIAGARA_DEBUG_EMITTER_NAME
 	StatelessEmitterData->DebugSimulationName = GetPackage()->GetFName();
diff --git a/Engine/Plugins/FX/Niagara/Source/Niagara/Private/Stateless/NiagaraStatelessEmitterData.cpp b/Engine/Plugins/FX/Niagara/Source/Niagara/Private/Stateless/NiagaraStatelessEmitterData.cpp
--- a/Engine/Plugins/FX/Niagara/Source/Niagara/Private/Stateless/NiagaraStatelessEmitterData.cpp
+++ b/Engine/Plugins/FX/Niagara/Source/Niagara/Private/Stateless/NiagaraStatelessEmitterData.cpp
@@ -7,24 +7,6 @@
 
 #include "Engine/Engine.h"
 
-void FNiagaraStatelessEmitterData::FDeleter::operator()(FNiagaraStatelessEmitterData* EmitterData) const
-{
-	if (IsInRenderingThread())
-	{
-		delete EmitterData;
-	}
-	else
-	{
-		ENQUEUE_RENDER_COMMAND(DeleteNiagaraStatelessEmitterData)
-		(
-			[EmitterData](FRHICommandListImmediate&)
-			{
-				delete EmitterData;
-			}
-		);
-	}
-}
-
 FNiagaraStatelessEmitterData::~FNiagaraStatelessEmitterData()
 {
 	check(IsInRenderingThread());
