Updated stateless forces and velocity solver. This fixes a problem where even small increases in mass values could quickly cause float precision issues, making the whole sim unstable after a few frames. This change might alter the look of existing content, especially if non-standard particle mass is used together with acceleration force.
* ignore particle mass for drag module
* more precise calculation when drag = 0
* acceleration force factors in mass now (gravity does not and still acts independent of mass)

#tests Checked and updated enginetests

diff --git a/Engine/Plugins/FX/Niagara/Shaders/Private/Stateless/Modules/NiagaraStatelessModule_ShapeLocation.ush b/Engine/Plugins/FX/Niagara/Shaders/Private/Stateless/Modules/NiagaraStatelessModule_ShapeLocation.ush
--- a/Engine/Plugins/FX/Niagara/Shaders/Private/Stateless/Modules/NiagaraStatelessModule_ShapeLocation.ush
+++ b/Engine/Plugins/FX/Niagara/Shaders/Private/Stateless/Modules/NiagaraStatelessModule_ShapeLocation.ush
@@ -6,7 +6,7 @@ uint4	ShapeLocation_Mode;
 float4	ShapeLocation_Parameters0;
 float4	ShapeLocation_Parameters1;
 
-float3 ShapeLocation_GetLocation(inout FStatelessParticle Particle)
+float3 ShapeLocation_GetLocation()
 {
 	// ENSM_ShapePrimitive::Box | ENSM_ShapePrimitive::Plane
 	BRANCH
@@ -34,10 +34,7 @@ float3 ShapeLocation_GetLocation(inout FStatelessParticle Particle)
 				
 			return Location;
 		}
-		else
-		{
-			return P0 * BoxSize;
-		}
+		return P0 * BoxSize;
 	}
 
 	// ENSM_ShapePrimitive::Cylinder:
@@ -82,7 +79,7 @@ float3 ShapeLocation_GetLocation(inout FStatelessParticle Particle)
 
 void ShapeLocation_Simulate(inout FStatelessParticle Particle)
 {
-	float3 Location = ShapeLocation_GetLocation(Particle);
+	float3 Location = ShapeLocation_GetLocation();
 	Particle.Position += Location;
 	Particle.PreviousPosition += Location;
 }
diff --git a/Engine/Plugins/FX/Niagara/Shaders/Private/Stateless/Modules/NiagaraStatelessModule_SolveVelocitiesAndForces.ush b/Engine/Plugins/FX/Niagara/Shaders/Private/Stateless/Modules/NiagaraStatelessModule_SolveVelocitiesAndForces.ush
--- a/Engine/Plugins/FX/Niagara/Shaders/Private/Stateless/Modules/NiagaraStatelessModule_SolveVelocitiesAndForces.ush
+++ b/Engine/Plugins/FX/Niagara/Shaders/Private/Stateless/Modules/NiagaraStatelessModule_SolveVelocitiesAndForces.ush
@@ -12,6 +12,8 @@ float3	SolveVelocitiesAndForces_WindScale;
 float3	SolveVelocitiesAndForces_WindBias;
 float3	SolveVelocitiesAndForces_AccelerationScale;
 float3	SolveVelocitiesAndForces_AccelerationBias;
+float3	SolveVelocitiesAndForces_GravityScale;
+float3	SolveVelocitiesAndForces_GravityBias;
 
 uint	SolveVelocitiesAndForces_ConeVelocityEnabled;
 float4	SolveVelocitiesAndForces_ConeQuat;
@@ -48,13 +50,20 @@ struct FStatelessModule_SolveVelocitiesAndForces
 
 void SolveVelocitiesAndForces_IntegratePosition(in FStatelessModule_SolveVelocitiesAndForces ModuleData, float Age, inout float3 Position)
 {
-	const float3 IntVelocity = (ModuleData.Velocity - ModuleData.Wind) + (ModuleData.Wind * Age * Age);
-	const float LambdaDragMass = max(ModuleData.Drag * rcp(ModuleData.Mass), 0.0001f);
-	const float LambdaAge = (1.0f - exp(-(LambdaDragMass * Age))) / LambdaDragMass;
-	Position += IntVelocity * LambdaAge;
-	Position += (Age - LambdaAge) * (ModuleData.Acceleration / LambdaDragMass);
+	if (ModuleData.Drag > 0.0001f)
+	{
+		const float3 TerminalVelocity = ModuleData.Acceleration * rcp(ModuleData.Drag) + ModuleData.Wind;
+		const float3 IntVelocity = ModuleData.Velocity - TerminalVelocity;
+		const float LambdaAge = (1.0f - exp(-(Age * ModuleData.Drag))) * rcp(ModuleData.Drag);
+		Position += IntVelocity * LambdaAge + TerminalVelocity * Age;
+	}
+	else
+	{
+		// without drag, we can use the simpler formula for Newtonian motion v*t + 1/2*a*t²
+		Position += Age * (ModuleData.Velocity + ModuleData.Wind) + 0.5f * ModuleData.Acceleration * Age * Age;
+	}
 }
-	
+
 void SolveVelocitiesAndForces_Simulate(inout FStatelessParticle Particle)
 {
 	// Initialize shared parameters
@@ -63,7 +72,9 @@ void SolveVelocitiesAndForces_Simulate(inout FStatelessParticle Particle)
 	ModuleData.Drag = RandomScaleBiasFloat(1, SolveVelocitiesAndForces_DragScale, SolveVelocitiesAndForces_DragBias);
 	ModuleData.Velocity = RandomScaleBiasFloat(2, SolveVelocitiesAndForces_VelocityScale, SolveVelocitiesAndForces_VelocityBias);
 	ModuleData.Wind = RandomScaleBiasFloat(3, SolveVelocitiesAndForces_WindScale, SolveVelocitiesAndForces_WindBias);
-	ModuleData.Acceleration = RandomScaleBiasFloat(4, SolveVelocitiesAndForces_AccelerationScale, SolveVelocitiesAndForces_AccelerationBias);
+	const float MassFactor = rcp(max(ModuleData.Mass, UE_KINDA_SMALL_NUMBER)); // factor in mass for acceleration force
+	ModuleData.Acceleration = MassFactor * RandomScaleBiasFloat(4, SolveVelocitiesAndForces_AccelerationScale, SolveVelocitiesAndForces_AccelerationBias); 
+	ModuleData.Acceleration += RandomScaleBiasFloat(41, SolveVelocitiesAndForces_GravityScale, SolveVelocitiesAndForces_GravityBias);
 
 	if ( SolveVelocitiesAndForces_ConeVelocityEnabled != 0 )
 	{
diff --git a/Engine/Plugins/FX/Niagara/Shaders/Private/Stateless/NiagaraStatelessSimulationTemplate.ush b/Engine/Plugins/FX/Niagara/Shaders/Private/Stateless/NiagaraStatelessSimulationTemplate.ush
--- a/Engine/Plugins/FX/Niagara/Shaders/Private/Stateless/NiagaraStatelessSimulationTemplate.ush
+++ b/Engine/Plugins/FX/Niagara/Shaders/Private/Stateless/NiagaraStatelessSimulationTemplate.ush
@@ -69,8 +69,8 @@ void StatelessMain(uint3 GroupId : SV_GroupID, int GroupThreadIndex : SV_GroupIn
 	uint LinearThreadId = GetUnWrappedDispatchThreadId(GroupId, GroupThreadIndex, THREADGROUP_SIZE);
 
 	// Initialize our particle data
-	FStatelessParticle Particle;
-	Particle.Age = -1.0f;
+	FStatelessParticle Particle = (FStatelessParticle)0;;
+	Particle.Age = -1.0f;	
 
 	// Find which spawn info we belong to and fill out details
 	uint ModuleSeedOffset = 0;
diff --git a/Engine/Plugins/FX/Niagara/Source/Niagara/Internal/Stateless/Modules/NiagaraStatelessModule_AccelerationForce.h b/Engine/Plugins/FX/Niagara/Source/Niagara/Internal/Stateless/Modules/NiagaraStatelessModule_AccelerationForce.h
--- a/Engine/Plugins/FX/Niagara/Source/Niagara/Internal/Stateless/Modules/NiagaraStatelessModule_AccelerationForce.h
+++ b/Engine/Plugins/FX/Niagara/Source/Niagara/Internal/Stateless/Modules/NiagaraStatelessModule_AccelerationForce.h
@@ -9,7 +9,7 @@
 
 #include "NiagaraStatelessModule_AccelerationForce.generated.h"
 
-// Apply acceleration which accumulates per frame
+// Apply an acceleration force each frame (in cm/s). This factors in mass, so particles with a high mass will accelerate slower. 
 UCLASS(MinimalAPI, EditInlineNew, meta = (DisplayName = "Acceleration Force"))
 class UNiagaraStatelessModule_AccelerationForce : public UNiagaraStatelessModule
 {
diff --git a/Engine/Plugins/FX/Niagara/Source/Niagara/Internal/Stateless/Modules/NiagaraStatelessModule_GravityForce.h b/Engine/Plugins/FX/Niagara/Source/Niagara/Internal/Stateless/Modules/NiagaraStatelessModule_GravityForce.h
--- a/Engine/Plugins/FX/Niagara/Source/Niagara/Internal/Stateless/Modules/NiagaraStatelessModule_GravityForce.h
+++ b/Engine/Plugins/FX/Niagara/Source/Niagara/Internal/Stateless/Modules/NiagaraStatelessModule_GravityForce.h
@@ -9,7 +9,7 @@
 
 #include "NiagaraStatelessModule_GravityForce.generated.h"
 
-// Applies a gravitational force (in cm/s)
+// Applies a gravitational force (in cm/s). This acceleration is the same regardless of mass, so particles with high and low mass will reach the same velocity.
 UCLASS(MinimalAPI, EditInlineNew, meta = (DisplayName = "Gravity Force"))
 class UNiagaraStatelessModule_GravityForce : public UNiagaraStatelessModule
 {
diff --git a/Engine/Plugins/FX/Niagara/Source/Niagara/Private/Stateless/Modules/NiagaraStatelessModule_SolveVelocitiesAndForces.cpp b/Engine/Plugins/FX/Niagara/Source/Niagara/Private/Stateless/Modules/NiagaraStatelessModule_SolveVelocitiesAndForces.cpp
--- a/Engine/Plugins/FX/Niagara/Source/Niagara/Private/Stateless/Modules/NiagaraStatelessModule_SolveVelocitiesAndForces.cpp
+++ b/Engine/Plugins/FX/Niagara/Source/Niagara/Private/Stateless/Modules/NiagaraStatelessModule_SolveVelocitiesAndForces.cpp
@@ -17,14 +17,18 @@ namespace NiagaraStatelessModuleSolveVelocitiesAndForcesPrivate
 		int32								PreviousVelocityVariableOffset = INDEX_NONE;
 	};
 
-	static FVector3f IntegratePosition(float Age, float Mass, float Drag, const FVector3f& Velocity, const FVector3f& Wind, const FVector3f& Acceleration)
+	static FVector3f IntegratePosition(float Age, float Drag, const FVector3f& Velocity, const FVector3f& Wind, const FVector3f& Acceleration)
 	{
-		const FVector3f IntVelocity = (Velocity - Wind) + (Wind * Age * Age);
-		const float LambdaDragMass = FMath::Max(Drag * (1.0f / Mass), 0.0001f);
-		const float LambdaAge = (1.0f - FMath::Exp(-(LambdaDragMass * Age))) / LambdaDragMass;
-		FVector3f Position = IntVelocity * LambdaAge;
-		Position += (Age - LambdaAge) * (Acceleration / LambdaDragMass);
-		return Position;
+		if (Drag > 0.0001f)
+		{
+			const FVector3f TerminalVelocity = Acceleration / Drag + Wind;
+			const FVector3f IntVelocity = Velocity - TerminalVelocity;
+			const float LambdaAge = (1.0f - FMath::Exp(-(Age * Drag))) / Drag;
+			return IntVelocity * LambdaAge + TerminalVelocity * Age;
+		}
+
+		// without drag, we can use the simpler formula for Newtonian motion v*t + 1/2*a*t²
+		return Age * (Velocity + Wind) + 0.5f * Acceleration * Age * Age;
 	}
 
 	static void ParticleSimulate(const NiagaraStateless::FParticleSimulationContext& ParticleSimulationContext)
@@ -42,7 +46,10 @@ namespace NiagaraStatelessModuleSolveVelocitiesAndForcesPrivate
 			const float Drag				= ParticleSimulationContext.RandomScaleBiasFloat(i, 1, Parameters->SolveVelocitiesAndForces_DragScale, Parameters->SolveVelocitiesAndForces_DragBias);
 			FVector3f InitialVelocity		= ParticleSimulationContext.RandomScaleBiasFloat(i, 2, Parameters->SolveVelocitiesAndForces_VelocityScale, Parameters->SolveVelocitiesAndForces_VelocityBias);
 			const FVector3f Wind			= ParticleSimulationContext.RandomScaleBiasFloat(i, 3, Parameters->SolveVelocitiesAndForces_WindScale, Parameters->SolveVelocitiesAndForces_WindBias);
-			const FVector3f Acceleration	= ParticleSimulationContext.RandomScaleBiasFloat(i, 4, Parameters->SolveVelocitiesAndForces_AccelerationScale, Parameters->SolveVelocitiesAndForces_AccelerationBias);
+
+			const float MassFactor	= 1 / (FMath::Max(Mass, UE_KINDA_SMALL_NUMBER)); // factor in mass for acceleration force
+			FVector3f Acceleration	= MassFactor * ParticleSimulationContext.RandomScaleBiasFloat(i, 4, Parameters->SolveVelocitiesAndForces_AccelerationScale, Parameters->SolveVelocitiesAndForces_AccelerationBias);
+			Acceleration			+= ParticleSimulationContext.RandomScaleBiasFloat(i, 41, Parameters->SolveVelocitiesAndForces_GravityScale, Parameters->SolveVelocitiesAndForces_GravityBias);
 
 			FVector3f Position			= ParticleSimulationContext.ReadParticleVariable(ModuleBuiltData->PositionVariableOffset, i, FVector3f::ZeroVector);
 			FVector3f PreviousPosition	= ParticleSimulationContext.ReadParticleVariable(ModuleBuiltData->PreviousPositionVariableOffset, i, FVector3f::ZeroVector);
@@ -80,8 +87,8 @@ namespace NiagaraStatelessModuleSolveVelocitiesAndForcesPrivate
 				//-TODO:
 			}
 
-			Position += IntegratePosition(FMath::Max(AgeData[i], UE_KINDA_SMALL_NUMBER), Mass, Drag, InitialVelocity, Wind, Acceleration);
-			PreviousPosition += IntegratePosition(PreviousAgeData[i], Mass, Drag, InitialVelocity, Wind, Acceleration);
+			Position += IntegratePosition(FMath::Max(AgeData[i], UE_KINDA_SMALL_NUMBER), Drag, InitialVelocity, Wind, Acceleration);
+			PreviousPosition += IntegratePosition(PreviousAgeData[i], Drag, InitialVelocity, Wind, Acceleration);
 
 			ParticleSimulationContext.WriteParticleVariable(ModuleBuiltData->PositionVariableOffset, i, Position);
 			ParticleSimulationContext.WriteParticleVariable(ModuleBuiltData->PreviousPositionVariableOffset, i, PreviousPosition);
@@ -137,8 +144,6 @@ void UNiagaraStatelessModule_SolveVelocitiesAndForces::SetShaderParameters(const
 	FParameters* Parameters = SetShaderParameterContext.GetParameterNestedStruct<FParameters>();
 	SetShaderParameterContext.ConvertRangeToScaleBias(PhysicsData.MassRange, Parameters->SolveVelocitiesAndForces_MassScale, Parameters->SolveVelocitiesAndForces_MassBias);
 	SetShaderParameterContext.ConvertRangeToScaleBias(PhysicsData.DragRange, Parameters->SolveVelocitiesAndForces_DragScale, Parameters->SolveVelocitiesAndForces_DragBias);
-	Parameters->SolveVelocitiesAndForces_DragScale	= FMath::Max(Parameters->SolveVelocitiesAndForces_DragScale, 0.01f);
-	Parameters->SolveVelocitiesAndForces_DragBias	= FMath::Max(Parameters->SolveVelocitiesAndForces_DragBias, 0.01f);
 
 	SetShaderParameterContext.TransformVectorRangeToScaleBias(PhysicsData.VelocityRange, PhysicsData.VelocityCoordinateSpace, Parameters->SolveVelocitiesAndForces_VelocityScale, Parameters->SolveVelocitiesAndForces_VelocityBias);
 	const float VelocityScale = SetShaderParameterContext.ConvertRangeToValue(PhysicsData.LinearVelocityScale);
@@ -153,8 +158,8 @@ void UNiagaraStatelessModule_SolveVelocitiesAndForces::SetShaderParameters(const
 		SetShaderParameterContext.ConvertRangeToScaleBias(PhysicsData.GravityRange, GravityScale, GravityBias);
 		Parameters->SolveVelocitiesAndForces_AccelerationScale	= SpaceTransforms.TransformVector(PhysicsData.AccelerationCoordinateSpace, AccelerationScale);
 		Parameters->SolveVelocitiesAndForces_AccelerationBias	= SpaceTransforms.TransformVector(PhysicsData.AccelerationCoordinateSpace, AccelerationBias);
-		Parameters->SolveVelocitiesAndForces_AccelerationScale	+= SpaceTransforms.TransformVector(ENiagaraCoordinateSpace::World, GravityScale);
-		Parameters->SolveVelocitiesAndForces_AccelerationBias	+= SpaceTransforms.TransformVector(ENiagaraCoordinateSpace::World, GravityBias);
+		Parameters->SolveVelocitiesAndForces_GravityScale	= SpaceTransforms.TransformVector(ENiagaraCoordinateSpace::World, GravityScale);
+		Parameters->SolveVelocitiesAndForces_GravityBias	= SpaceTransforms.TransformVector(ENiagaraCoordinateSpace::World, GravityBias);
 	}
 
 	Parameters->SolveVelocitiesAndForces_ConeVelocityEnabled = PhysicsData.bConeVelocity ? 1 : 0;
diff --git a/Engine/Plugins/FX/Niagara/Source/NiagaraShader/Internal/Stateless/NiagaraStatelessModuleShaderParameters.h b/Engine/Plugins/FX/Niagara/Source/NiagaraShader/Internal/Stateless/NiagaraStatelessModuleShaderParameters.h
--- a/Engine/Plugins/FX/Niagara/Source/NiagaraShader/Internal/Stateless/NiagaraStatelessModuleShaderParameters.h
+++ b/Engine/Plugins/FX/Niagara/Source/NiagaraShader/Internal/Stateless/NiagaraStatelessModuleShaderParameters.h
@@ -156,6 +156,8 @@ namespace NiagaraStateless
 		SHADER_PARAMETER(FVector3f,		SolveVelocitiesAndForces_WindBias)
 		SHADER_PARAMETER(FVector3f,		SolveVelocitiesAndForces_AccelerationScale)
 		SHADER_PARAMETER(FVector3f,		SolveVelocitiesAndForces_AccelerationBias)
+		SHADER_PARAMETER(FVector3f,		SolveVelocitiesAndForces_GravityScale)
+		SHADER_PARAMETER(FVector3f,		SolveVelocitiesAndForces_GravityBias)
 
 		SHADER_PARAMETER(uint32,		SolveVelocitiesAndForces_ConeVelocityEnabled)
 		SHADER_PARAMETER(FQuat4f,		SolveVelocitiesAndForces_ConeQuat)
