Niagara - Fix issues encountered while debugging static switch inconsistencies.
+ Fix an issue with the emitter binding data interface where the self binding would fail when the the data interface was used in the emitter, but the system owned the data interface at compile time.
+ Refactor the linked value handling in the stack to use a full niagara variable with name and type rather than a parameter handle which just includes the name.  This fixes an issue where binding to a static variable from a different script could cause a non-static version to be created.
+ When caching static variables for the graph traversal, don't cache variables that are not exportable since they shouldn't affect the next traversal.

diff --git a/Engine/Plugins/FX/Niagara/Source/Niagara/Private/NiagaraDataInterfaceEmitterBinding.cpp b/Engine/Plugins/FX/Niagara/Source/Niagara/Private/NiagaraDataInterfaceEmitterBinding.cpp
--- a/Engine/Plugins/FX/Niagara/Source/Niagara/Private/NiagaraDataInterfaceEmitterBinding.cpp
+++ b/Engine/Plugins/FX/Niagara/Source/Niagara/Private/NiagaraDataInterfaceEmitterBinding.cpp
@@ -106,6 +106,7 @@ const FNiagaraEmitterHandle* FNiagaraDataInterfaceEmitterBinding::ResolveHandle(
 
 	if (BindingMode == ENiagaraDataInterfaceEmitterBindingMode::Self)
 	{
+		// First check if this data interface is owned directly by an emitter.
 		if (UNiagaraEmitter* OwnerEmitter = DataInterface->GetTypedOuter<UNiagaraEmitter>())
 		{
 			for (const FNiagaraEmitterHandle& EmitterHandle : OwnerSystem->GetEmitterHandles())
@@ -118,6 +119,7 @@ const FNiagaraEmitterHandle* FNiagaraDataInterfaceEmitterBinding::ResolveHandle(
 		}
 		else
 		{
+			// Try to find this data interface in one of the emitter's particle scripts.
 			for (const FNiagaraEmitterHandle& EmitterHandle : OwnerSystem->GetEmitterHandles())
 			{
 				FVersionedNiagaraEmitterData* EmitterData = EmitterHandle.GetEmitterData();
@@ -146,6 +148,22 @@ const FNiagaraEmitterHandle* FNiagaraDataInterfaceEmitterBinding::ResolveHandle(
 					}
 				}
 			}
+			// Check the system scripts since they contain the compiled emitter data.
+			TArray<UNiagaraScript*> SystemScripts = { OwnerSystem->GetSystemSpawnScript(), OwnerSystem->GetSystemUpdateScript() };
+			for (UNiagaraScript* SystemScript : SystemScripts)
+			{
+				for (const FNiagaraScriptResolvedDataInterfaceInfo& ResolvedDI : SystemScript->GetResolvedDataInterfaces())
+				{
+					if (ResolvedDI.ResolvedDataInterface == DataInterface)
+					{
+						if (ResolvedDI.ResolvedSourceEmitterName.IsEmpty() == false)
+						{
+							return OwnerSystem->GetEmitterHandles().FindByPredicate([&ResolvedDI](const FNiagaraEmitterHandle& EmitterHandle)
+								{ return EmitterHandle.GetUniqueInstanceName() == ResolvedDI.ResolvedSourceEmitterName; });
+						}
+					}
+				}
+			}
 		}
 	}
 	else if (BindingMode == ENiagaraDataInterfaceEmitterBindingMode::Other)
diff --git a/Engine/Plugins/FX/Niagara/Source/NiagaraEditor/Private/NiagaraClipboard.cpp b/Engine/Plugins/FX/Niagara/Source/NiagaraEditor/Private/NiagaraClipboard.cpp
--- a/Engine/Plugins/FX/Niagara/Source/NiagaraEditor/Private/NiagaraClipboard.cpp
+++ b/Engine/Plugins/FX/Niagara/Source/NiagaraEditor/Private/NiagaraClipboard.cpp
@@ -66,7 +66,7 @@ const UNiagaraClipboardFunctionInput* UNiagaraClipboardFunctionInput::CreateLoca
 	return NewInput;
 }
 
-const UNiagaraClipboardFunctionInput* UNiagaraClipboardFunctionInput::CreateLinkedValue(UObject* InOuter, FName InInputName, FNiagaraTypeDefinition InInputType, TOptional<bool> bInEditConditionValue, FName InLinkedValue)
+const UNiagaraClipboardFunctionInput* UNiagaraClipboardFunctionInput::CreateLinkedValue(UObject* InOuter, FName InInputName, FNiagaraTypeDefinition InInputType, TOptional<bool> bInEditConditionValue, const FNiagaraVariableBase& InLinkedValue)
 {
 	UNiagaraClipboardFunctionInput* NewInput = MakeNewInput(InOuter, InInputName, InInputType, bInEditConditionValue, ENiagaraClipboardFunctionInputValueMode::Linked);
 	NewInput->Linked = InLinkedValue;
@@ -550,12 +550,13 @@ UNiagaraClipboardFunctionInput* UNiagaraClipboardEditorScriptingUtilities::Creat
 	FNiagaraTypeDefinition InputType = GetRegisteredTypeDefinitionByName(InInputTypeName);
 	if (InputType.IsValid())
 	{
+		FNiagaraVariableBase LinkedParameter = FNiagaraVariableBase(InputType, InLinkedValue);
 		return const_cast<UNiagaraClipboardFunctionInput*>(UNiagaraClipboardFunctionInput::CreateLinkedValue(
 			InOuter != nullptr ? InOuter : GetTransientPackage(),
 			InInputName,
 			InputType,
 			bInHasEditCondition ? TOptional<bool>(bInEditConditionValue) : TOptional<bool>(),
-			InLinkedValue));
+			LinkedParameter));
 	}
 	return nullptr;
 }
diff --git a/Engine/Plugins/FX/Niagara/Source/NiagaraEditor/Private/NiagaraCompiler.cpp b/Engine/Plugins/FX/Niagara/Source/NiagaraEditor/Private/NiagaraCompiler.cpp
--- a/Engine/Plugins/FX/Niagara/Source/NiagaraEditor/Private/NiagaraCompiler.cpp
+++ b/Engine/Plugins/FX/Niagara/Source/NiagaraEditor/Private/NiagaraCompiler.cpp
@@ -1606,9 +1606,12 @@ TSharedPtr<FNiagaraGraphCachedDataBase, ESPMode::ThreadSafe> FNiagaraEditorModul
 			Builder.BuildParameterMaps(FoundOutputNode, true);
 			Builder.EndUsage();
 
-			for (const FNiagaraVariable& Var : Builder.StaticVariables)
+			for (int32 BuilderVarIndex = 0; BuilderVarIndex < Builder.StaticVariables.Num() && BuilderVarIndex < Builder.StaticVariableExportable.Num(); ++BuilderVarIndex)
 			{
-				StaticVariables.AddUnique(Var);
+				if (Builder.StaticVariableExportable[BuilderVarIndex])
+				{
+					StaticVariables.AddUnique(Builder.StaticVariables[BuilderVarIndex]);
+				}
 			}
 		}
 	}
diff --git a/Engine/Plugins/FX/Niagara/Source/NiagaraEditor/Private/NiagaraScriptMergeManager.cpp b/Engine/Plugins/FX/Niagara/Source/NiagaraEditor/Private/NiagaraScriptMergeManager.cpp
--- a/Engine/Plugins/FX/Niagara/Source/NiagaraEditor/Private/NiagaraScriptMergeManager.cpp
+++ b/Engine/Plugins/FX/Niagara/Source/NiagaraEditor/Private/NiagaraScriptMergeManager.cpp
@@ -82,15 +82,16 @@ FNiagaraStackFunctionInputOverrideMergeAdapter::FNiagaraStackFunctionInputOverri
 		if (OverridePin->LinkedTo[0]->GetOwningNode()->IsA<UNiagaraNodeParameterMapGet>())
 		{
 			LinkedValueData = FNiagaraStackLinkedValueData();
-			LinkedValueData->LinkedValueHandle = FNiagaraParameterHandle(OverridePin->LinkedTo[0]->PinName);
-			if (LinkedValueData->LinkedValueHandle.IsOutputHandle() ||
-				LinkedValueData->LinkedValueHandle.IsStackContextHandle() ||
-				LinkedValueData->LinkedValueHandle.IsEmitterHandle() ||
-				LinkedValueData->LinkedValueHandle.IsParticleAttributeHandle())
+			LinkedValueData->LinkedParameterValue = NiagaraSchema->PinToNiagaraVariable(OverridePin->LinkedTo[0]);
+			FNiagaraParameterHandle LinkedValueHandle = FNiagaraParameterHandle(LinkedValueData->LinkedParameterValue.GetName());
+			if (LinkedValueHandle.IsOutputHandle() ||
+				LinkedValueHandle.IsStackContextHandle() ||
+				LinkedValueHandle.IsEmitterHandle() ||
+				LinkedValueHandle.IsParticleAttributeHandle())
 			{
 				// If the linked handle is a module output or module data set attribute, record the node id so that we can check for renamed nodes
 				// when applying the diff.
-				TArray<FName> HandleParts = LinkedValueData->LinkedValueHandle.GetHandleParts();
+				TArray<FName> HandleParts = LinkedValueHandle.GetHandleParts();
 				if (HandleParts.Num() > 2)
 				{
 					FString FunctionName = HandleParts[1].ToString();
@@ -3042,10 +3043,10 @@ FNiagaraScriptMergeManager::FApplyDiffResults FNiagaraScriptMergeManager::AddInp
 			else if (OverrideToAdd->GetLinkedValueData().IsSet())
 			{
 				check(OverrideToAdd->GetOverrideNode() && OverrideToAdd->GetOverrideNode()->GetNiagaraGraph());
-				FNiagaraParameterHandle OldLinkedValueHandle = OverrideToAdd->GetLinkedValueData()->LinkedValueHandle;
+				FNiagaraVariableBase OldLinkedParameter = OverrideToAdd->GetLinkedValueData()->LinkedParameterValue;
+				FNiagaraParameterHandle OldLinkedParameterHandle = FNiagaraParameterHandle(OldLinkedParameter.GetName());
 				FGuid LinkedFunctionCallNodeId = OverrideToAdd->GetLinkedValueData()->LinkedFunctionNodeId;
-				FNiagaraVariable OldLinkedValueVariable = FNiagaraVariable(InputType, OldLinkedValueHandle.GetParameterHandleString());
-				UNiagaraScriptVariable* ScriptVar = OverrideToAdd->GetOverrideNode()->GetNiagaraGraph()->GetScriptVariable(OldLinkedValueVariable);
+				UNiagaraScriptVariable* ScriptVar = OverrideToAdd->GetOverrideNode()->GetNiagaraGraph()->GetScriptVariable(OldLinkedParameter);
 				ENiagaraDefaultMode DesiredMode = ENiagaraDefaultMode::Value;
 				if (ScriptVar)
 					DesiredMode = ScriptVar->DefaultMode;
@@ -3054,22 +3055,24 @@ FNiagaraScriptMergeManager::FApplyDiffResults FNiagaraScriptMergeManager::AddInp
 
 				// If the linked value handle has a valid function call node id, then we need to check to see if that function call node has been
 				// renamed due to the merge, and if so we need to update the handle.
-				FNiagaraParameterHandle NewLinkedValueHandle;
+				FNiagaraVariableBase NewLinkedParameter;
 				if (LinkedFunctionCallNodeId.IsValid() &&
-					(OldLinkedValueHandle.IsOutputHandle() || OldLinkedValueHandle.IsStackContextHandle() ||
-						OldLinkedValueHandle.IsEmitterHandle() || OldLinkedValueHandle.IsParticleAttributeHandle()) &&
-					OldLinkedValueHandle.GetHandleParts().Num() > 2)
+					OldLinkedParameterHandle.GetHandleParts().Num() >= 3 && (
+						OldLinkedParameterHandle.IsOutputHandle() ||
+						OldLinkedParameterHandle.IsStackContextHandle() ||
+						OldLinkedParameterHandle.IsEmitterHandle() ||
+						OldLinkedParameterHandle.IsParticleAttributeHandle()))
 				{
 					TObjectPtr<UEdGraphNode>* ReferencedFunctionCallNodePtr = TargetFunctionCall.GetGraph()->Nodes.FindByPredicate(
 						[&LinkedFunctionCallNodeId](UEdGraphNode* Node) { return Node->IsA<UNiagaraNodeFunctionCall>() && Node->NodeGuid == LinkedFunctionCallNodeId; });
 					if (ReferencedFunctionCallNodePtr != nullptr)
 					{
 						UNiagaraNodeFunctionCall* ReferencedFunctionCallNode = CastChecked<UNiagaraNodeFunctionCall>(*ReferencedFunctionCallNodePtr);
-						FString FunctionName = OldLinkedValueHandle.GetHandleParts()[1].ToString();
+						FString FunctionName = OldLinkedParameterHandle.GetHandleParts()[1].ToString();
 						if (ReferencedFunctionCallNode->GetFunctionName() != FunctionName)
 						{
 							// The function node has been renamed so we need to update the handle.  The handle is in the format [Namespace].[Function Call Name].[NameParts]
-							TArray<FName> OldHandleParts = OldLinkedValueHandle.GetHandleParts();
+							TArray<FName> OldHandleParts = OldLinkedParameterHandle.GetHandleParts();
 							TArray<FString> NewHandleNameParts;
 							NewHandleNameParts.Add(OldHandleParts[0].ToString());
 							NewHandleNameParts.Add(ReferencedFunctionCallNode->GetFunctionName());
@@ -3077,23 +3080,21 @@ FNiagaraScriptMergeManager::FApplyDiffResults FNiagaraScriptMergeManager::AddInp
 							{
 								NewHandleNameParts.Add(OldHandleParts[i].ToString());
 							}
-							NewLinkedValueHandle = FNiagaraParameterHandle(*FString::Join(NewHandleNameParts, TEXT(".")));
+							NewLinkedParameter = FNiagaraVariableBase(OldLinkedParameter.GetType(), *FString::Join(NewHandleNameParts, TEXT(".")));
 						}
 					}
 				}
-				if (NewLinkedValueHandle.IsValid() == false)
+				if (NewLinkedParameter.IsValid() == false)
 				{
-					NewLinkedValueHandle = OldLinkedValueHandle;
+					NewLinkedParameter = OldLinkedParameter;
 				}
 
-				FNiagaraVariable NewLinkedValueVariable = FNiagaraVariable(OverrideToAdd->GetType(), NewLinkedValueHandle.GetParameterHandleString());
-				TSet KnownParameters = { NewLinkedValueVariable };
-				FNiagaraStackGraphUtilities::SetLinkedValueHandleForFunctionInput(InputOverridePin, NewLinkedValueHandle, KnownParameters, DesiredMode, OverrideToAdd->GetOverrideNodeId());
-				FGuid LinkedOutputId = FNiagaraStackGraphUtilities::GetScriptVariableIdForLinkedModuleParameterHandle(
-					OldLinkedValueHandle, OverrideToAdd->GetType(), *TargetFunctionCall.GetNiagaraGraph());
+				TSet KnownParameters = { NewLinkedParameter };
+				FNiagaraStackGraphUtilities::SetLinkedParameterValueForFunctionInput(InputOverridePin, NewLinkedParameter, KnownParameters, DesiredMode, OverrideToAdd->GetOverrideNodeId());
+				FGuid LinkedOutputId = FNiagaraStackGraphUtilities::GetScriptVariableIdForLinkedModuleParameter(NewLinkedParameter, *TargetFunctionCall.GetNiagaraGraph());
 				if (LinkedOutputId.IsValid())
 				{
-					TargetFunctionCall.UpdateInputNameBinding(LinkedOutputId, NewLinkedValueHandle.GetParameterHandleString());
+					TargetFunctionCall.UpdateInputNameBinding(LinkedOutputId, NewLinkedParameter.GetName());
 				}
 				Results.bSucceeded = true;
 			}
diff --git a/Engine/Plugins/FX/Niagara/Source/NiagaraEditor/Private/NiagaraScriptMergeManager.h b/Engine/Plugins/FX/Niagara/Source/NiagaraEditor/Private/NiagaraScriptMergeManager.h
--- a/Engine/Plugins/FX/Niagara/Source/NiagaraEditor/Private/NiagaraScriptMergeManager.h
+++ b/Engine/Plugins/FX/Niagara/Source/NiagaraEditor/Private/NiagaraScriptMergeManager.h
@@ -34,12 +34,12 @@ class FNiagaraStackFunctionMergeAdapter;
 
 struct FNiagaraStackLinkedValueData
 {
-	FNiagaraParameterHandle LinkedValueHandle;
+	FNiagaraVariableBase LinkedParameterValue;
 	FGuid LinkedFunctionNodeId;
 
 	bool operator==(const FNiagaraStackLinkedValueData& Other) const
 	{
-		return LinkedValueHandle == Other.LinkedValueHandle && LinkedFunctionNodeId == Other.LinkedFunctionNodeId;
+		return LinkedParameterValue == Other.LinkedParameterValue && LinkedFunctionNodeId == Other.LinkedFunctionNodeId;
 	}
 };
 
diff --git a/Engine/Plugins/FX/Niagara/Source/NiagaraEditor/Private/NiagaraValidationRules.cpp b/Engine/Plugins/FX/Niagara/Source/NiagaraEditor/Private/NiagaraValidationRules.cpp
--- a/Engine/Plugins/FX/Niagara/Source/NiagaraEditor/Private/NiagaraValidationRules.cpp
+++ b/Engine/Plugins/FX/Niagara/Source/NiagaraEditor/Private/NiagaraValidationRules.cpp
@@ -1114,7 +1114,7 @@ void UNiagaraValidationRule_LWC::CheckValidity(const FNiagaraValidationContext&
 				// check if the linked input variable is a vector
 				if (Input->GetValueMode() == UNiagaraStackFunctionInput::EValueMode::Linked && Settings->bEnforceStrictStackTypes)
 				{
-					FNiagaraVariable VectorVar(FNiagaraTypeDefinition::GetVec3Def(), Input->GetLinkedValueHandle().GetParameterHandleString());
+					FNiagaraVariable VectorVar(FNiagaraTypeDefinition::GetVec3Def(), Input->GetLinkedParameterValue().GetName());
 					const UNiagaraGraph* NiagaraGraph = Input->GetInputFunctionCallNode().GetNiagaraGraph();
 
 					// we check if metadata for a vector attribute with the linked name exists in the emitter/system script graph. Not 100% correct, but it needs to be fast and a few false negatives are acceptable.
diff --git a/Engine/Plugins/FX/Niagara/Source/NiagaraEditor/Private/UpgradeNiagaraScriptResults.cpp b/Engine/Plugins/FX/Niagara/Source/NiagaraEditor/Private/UpgradeNiagaraScriptResults.cpp
--- a/Engine/Plugins/FX/Niagara/Source/NiagaraEditor/Private/UpgradeNiagaraScriptResults.cpp
+++ b/Engine/Plugins/FX/Niagara/Source/NiagaraEditor/Private/UpgradeNiagaraScriptResults.cpp
@@ -177,7 +177,7 @@ FString UNiagaraPythonScriptModuleInput::AsLinkedValue() const
 {
 	if (IsLinkedValue())
 	{
-		return Input->Linked.ToString();
+		return Input->Linked.GetName().ToString();
 	}
 	return FString();
 }
@@ -313,7 +313,8 @@ void UUpgradeNiagaraScriptResults::SetLinkedInput(const FString& InputName, FStr
 	if (UNiagaraPythonScriptModuleInput* ModuleInput = GetNewInput(FName(InputName)))
 	{
 		const UNiagaraClipboardFunctionInput* CurrentInput = ModuleInput->Input;
-		ModuleInput->Input = UNiagaraClipboardFunctionInput::CreateLinkedValue(ModuleInput, CurrentInput->InputName, CurrentInput->InputType, CurrentInput->bEditConditionValue, FName(Value));
+		FNiagaraVariableBase LinkedParameter = FNiagaraVariableBase(CurrentInput->InputType, *Value);
+		ModuleInput->Input = UNiagaraClipboardFunctionInput::CreateLinkedValue(ModuleInput, CurrentInput->InputName, CurrentInput->InputType, CurrentInput->bEditConditionValue, LinkedParameter);
 	}
 }
 
diff --git a/Engine/Plugins/FX/Niagara/Source/NiagaraEditor/Private/ViewModels/NiagaraParameterPanelViewModel.cpp b/Engine/Plugins/FX/Niagara/Source/NiagaraEditor/Private/ViewModels/NiagaraParameterPanelViewModel.cpp
--- a/Engine/Plugins/FX/Niagara/Source/NiagaraEditor/Private/ViewModels/NiagaraParameterPanelViewModel.cpp
+++ b/Engine/Plugins/FX/Niagara/Source/NiagaraEditor/Private/ViewModels/NiagaraParameterPanelViewModel.cpp
@@ -1086,7 +1086,7 @@ void FNiagaraSystemToolkitParameterPanelViewModel::OnParameterItemSelected(const
 	InvalidateCachedDependencies();
 }
 
-bool FNiagaraSystemToolkitParameterPanelViewModel::IsVariableSelected(FNiagaraVariableBase& InVar) const
+bool FNiagaraSystemToolkitParameterPanelViewModel::IsVariableSelected(const FNiagaraVariableBase& InVar) const
 {
 	if (VariableObjectSelection.IsValid())
 	{
diff --git a/Engine/Plugins/FX/Niagara/Source/NiagaraEditor/Private/ViewModels/Stack/NiagaraStackFunctionInput.cpp b/Engine/Plugins/FX/Niagara/Source/NiagaraEditor/Private/ViewModels/Stack/NiagaraStackFunctionInput.cpp
--- a/Engine/Plugins/FX/Niagara/Source/NiagaraEditor/Private/ViewModels/Stack/NiagaraStackFunctionInput.cpp
+++ b/Engine/Plugins/FX/Niagara/Source/NiagaraEditor/Private/ViewModels/Stack/NiagaraStackFunctionInput.cpp
@@ -649,7 +649,7 @@ FText UNiagaraStackFunctionInput::GetCollapsedStateText() const
 			}
 			break;
 		case EValueMode::Linked:
-			CollapsedTextCache = FText::FromString(InputValues.LinkedHandle.GetParameterHandleString().ToString());
+			CollapsedTextCache = FText::FromString(InputValues.LinkedParameter.GetName().ToString());
 			break;
 		case EValueMode::Expression:
 			CollapsedTextCache = FText::Format(FText::FromString("({0})"), FText::FromString(InputValues.ExpressionNode->GetCustomHlsl()));
@@ -733,7 +733,7 @@ FText UNiagaraStackFunctionInput::GetValueToolTip() const
 			break;
 		case EValueMode::Linked:
 		{
-			FNiagaraVariable InputVariable(InputType, InputValues.LinkedHandle.GetParameterHandleString());
+			FNiagaraVariable InputVariable = InputValues.LinkedParameter;
 			if (FNiagaraConstants::IsNiagaraConstant(InputVariable))
 			{
 				const FNiagaraVariableMetaData* FoundMetaData = FNiagaraConstants::GetConstantMetaData(InputVariable);
@@ -1547,14 +1547,14 @@ const FNiagaraParameterHandle& UNiagaraStackFunctionInput::GetInputParameterHand
 	return InputParameterHandle;
 }
 
-const FNiagaraParameterHandle& UNiagaraStackFunctionInput::GetLinkedValueHandle() const
+const FNiagaraVariableBase& UNiagaraStackFunctionInput::GetLinkedParameterValue() const
 {
-	return InputValues.LinkedHandle;
+	return InputValues.LinkedParameter;
 }
 
-void UNiagaraStackFunctionInput::SetLinkedValueHandle(const FNiagaraParameterHandle& InParameterHandle)
+void UNiagaraStackFunctionInput::SetLinkedParameterValue(const FNiagaraVariableBase& InParameter)
 {
-	if (InParameterHandle == InputValues.LinkedHandle)
+	if (InParameter == InputValues.LinkedParameter)
 	{
 		return;
 	}
@@ -1568,13 +1568,13 @@ void UNiagaraStackFunctionInput::SetLinkedValueHandle(const FNiagaraParameterHan
 		RemoveRapidIterationParametersForAffectedScripts();
 	}
 
-	if (InParameterHandle != DefaultInputValues.LinkedHandle)
+	if (InParameter != DefaultInputValues.LinkedParameter)
 	{
-		if (InParameterHandle.IsUserHandle())
+		FNiagaraParameterHandle ParameterHandle = FNiagaraParameterHandle(InParameter.GetName());
+		if (ParameterHandle.IsUserHandle())
 		{
 			// If the handle is a user parameter, make sure the system has it exposed.  If it's not exposed add it directly here
 			// rather than waiting on the compile results so that it's immediately available.
-			FNiagaraVariable LinkedParameterVariable = FNiagaraVariable(InputType, InParameterHandle.GetParameterHandleString()); 
 			FNiagaraParameterStore& UserParameters = GetSystemViewModel()->GetSystem().GetExposedParameters();
 
 			// we only know the input type, and the parameter handle. We don't know the assigned parameter's real type (a vector parameter can be linked to a position input).
@@ -1583,7 +1583,7 @@ void UNiagaraStackFunctionInput::SetLinkedValueHandle(const FNiagaraParameterHan
 			bool bSkipCreation = false;
 			if(InputType == FNiagaraTypeDefinition::GetPositionDef())
 			{
-				FNiagaraVariable ExistingVectorParameter(FNiagaraTypeDefinition::GetVec3Def(), InParameterHandle.GetParameterHandleString());
+				FNiagaraVariable ExistingVectorParameter(FNiagaraTypeDefinition::GetVec3Def(), InParameter.GetName());
 				if(UserParameters.IndexOf(ExistingVectorParameter) != INDEX_NONE)
 				{
 					bSkipCreation = true;
@@ -1591,20 +1591,20 @@ void UNiagaraStackFunctionInput::SetLinkedValueHandle(const FNiagaraParameterHan
 			}
 			else if(InputType == FNiagaraTypeDefinition::GetVec3Def())
 			{
-				FNiagaraVariable ExistingPositionParameter(FNiagaraTypeDefinition::GetPositionDef(), InParameterHandle.GetParameterHandleString());
+				FNiagaraVariable ExistingPositionParameter(FNiagaraTypeDefinition::GetPositionDef(), InParameter.GetName());
 				if(UserParameters.IndexOf(ExistingPositionParameter) != INDEX_NONE)
 				{
 					bSkipCreation = true;
 				}
 			}
-			if (UserParameters.IndexOf(LinkedParameterVariable) == INDEX_NONE && bSkipCreation == false)
+			if (UserParameters.IndexOf(InParameter) == INDEX_NONE && bSkipCreation == false)
 			{
 				if (InputType.IsDataInterface())
 				{
 					int32 DataInterfaceOffset;
 					bool bInitialize = true;
 					bool bTriggerRebind = true;
-					UserParameters.AddParameter(LinkedParameterVariable, bInitialize, bTriggerRebind, &DataInterfaceOffset);
+					UserParameters.AddParameter(InParameter, bInitialize, bTriggerRebind, &DataInterfaceOffset);
 					if (InputValues.Mode == EValueMode::Data && InputValues.DataObject.IsValid())
 					{
 						InputValues.DataObject->CopyTo(UserParameters.GetDataInterface(DataInterfaceOffset));
@@ -1616,7 +1616,7 @@ void UNiagaraStackFunctionInput::SetLinkedValueHandle(const FNiagaraParameterHan
 					const bool bInitialize = true;
 					const bool bTriggerRebind = true;
 					int32 ParameterOffset = INDEX_NONE;
-					UserParameters.AddParameter(LinkedParameterVariable, bInitialize, bTriggerRebind, &ParameterOffset);
+					UserParameters.AddParameter(InParameter, bInitialize, bTriggerRebind, &ParameterOffset);
 					if (InputValues.Mode == EValueMode::ObjectAsset)
 					{
 						UNiagaraNodeInput* InputNode = InputValues.ObjectAssetInputNode.Get();
@@ -1626,28 +1626,30 @@ void UNiagaraStackFunctionInput::SetLinkedValueHandle(const FNiagaraParameterHan
 				}
 				else
 				{
+					FNiagaraVariable ParameterWithData = InParameter;
 					if (InputValues.Mode == EValueMode::Local && InputValues.LocalStruct.IsValid())
 					{
 						// If the current value is local, and valid transfer that value to the user parameter.
-						LinkedParameterVariable.SetData(InputValues.LocalStruct->GetStructMemory());
+						ParameterWithData.SetData(InputValues.LocalStruct->GetStructMemory());
 					}
 					else
 					{
-						FNiagaraEditorUtilities::ResetVariableToDefaultValue(LinkedParameterVariable);
+						FNiagaraEditorUtilities::ResetVariableToDefaultValue(ParameterWithData);
 					}
-					UserParameters.SetParameterData(LinkedParameterVariable.GetData(), LinkedParameterVariable, true);
+					UserParameters.SetParameterData(ParameterWithData.GetData(), ParameterWithData, true);
 				}
 			}
 		}
 
 		// Only set the linked value if it's actually different from the default.
 		UEdGraphPin& OverridePin = GetOrCreateOverridePin();
-		TSet<FNiagaraVariable> KnownParameters = FNiagaraStackGraphUtilities::GetParametersForContext(OverridePin.GetOwningNode()->GetGraph(), GetSystemViewModel()->GetSystem());
-		FNiagaraStackGraphUtilities::SetLinkedValueHandleForFunctionInput(OverridePin, InParameterHandle, KnownParameters);
-		FGuid LinkedOutputId = FNiagaraStackGraphUtilities::GetScriptVariableIdForLinkedModuleParameterHandle(InParameterHandle, InputType, *OwningFunctionCallNode->GetNiagaraGraph());
+		TSet<FNiagaraVariableBase> KnownParameters;
+		FNiagaraStackGraphUtilities::GetParametersForContext(OverridePin.GetOwningNode()->GetGraph(), GetSystemViewModel()->GetSystem(), KnownParameters);
+		FNiagaraStackGraphUtilities::SetLinkedParameterValueForFunctionInput(OverridePin, InParameter, KnownParameters);
+		FGuid LinkedOutputId = FNiagaraStackGraphUtilities::GetScriptVariableIdForLinkedModuleParameter(InParameter, *OwningFunctionCallNode->GetNiagaraGraph());
 		if (LinkedOutputId.IsValid())
 		{
-			OwningFunctionCallNode->UpdateInputNameBinding(LinkedOutputId, InParameterHandle.GetParameterHandleString());
+			OwningFunctionCallNode->UpdateInputNameBinding(LinkedOutputId, InParameter.GetName());
 		}
 	}
 
@@ -1655,6 +1657,20 @@ void UNiagaraStackFunctionInput::SetLinkedValueHandle(const FNiagaraParameterHan
 	RefreshValues();
 }
 
+void UNiagaraStackFunctionInput::SetLinkedParameterValueViaConversionScript(const FNiagaraVariableBase& InParameter, UNiagaraScript& ConversionScript)
+{
+	FScopedTransaction ScopedTransaction(LOCTEXT("SetConversionInput", "Make auto-convert dynamic input"));
+	SetDynamicInput(&ConversionScript);
+	for (UNiagaraStackFunctionInput* ChildInput : GetChildInputs())
+	{
+		if (InParameter.GetType() == ChildInput->GetInputType())
+		{
+			ChildInput->SetLinkedParameterValue(InParameter);
+			break;
+		}
+	}
+}
+
 bool UsageRunsBefore(ENiagaraScriptUsage UsageA, ENiagaraScriptUsage UsageB, bool bCheckInterpSpawn, FVersionedNiagaraEmitter InEmitter, FGuid UsageAID, FGuid UsageBID)
 {
 	static TArray<ENiagaraScriptUsage> UsagesOrderedByExecution
@@ -1731,14 +1747,13 @@ FName GetNamespaceForUsage(ENiagaraScriptUsage Usage)
 	}
 }
 
-bool IsNamespaceAllowedInUsage(FName Namespace, ENiagaraScriptUsage Usage)
+bool IsSourceNamespaceAllowedInTargetNamespace(FName SourceNamespace, FName TargetNamespace)
 {
-	FName UsageName = GetNamespaceForUsage(Usage);
-	if (UsageName == FNiagaraConstants::SystemNamespace && (Namespace == FNiagaraConstants::EmitterNamespace || Namespace == FNiagaraConstants::ParticleAttributeNamespace))
+	if (TargetNamespace == FNiagaraConstants::SystemNamespace && (SourceNamespace == FNiagaraConstants::EmitterNamespace || SourceNamespace == FNiagaraConstants::ParticleAttributeNamespace))
 	{
 		return false;
 	}
-	if (UsageName == FNiagaraConstants::EmitterNamespace && Namespace == FNiagaraConstants::ParticleAttributeNamespace)
+	if (TargetNamespace == FNiagaraConstants::EmitterNamespace && SourceNamespace == FNiagaraConstants::ParticleAttributeNamespace)
 	{
 		return false;
 	}
@@ -1760,7 +1775,7 @@ UNiagaraScript* UNiagaraStackFunctionInput::FindConversionScript(const FNiagaraT
 	return ConversionScriptCache[FromType];
 }
 
-void UNiagaraStackFunctionInput::GetAvailableParameterHandles(TArray<FNiagaraParameterHandle>& AvailableParameterHandles, TMap<FNiagaraVariable, UNiagaraScript*>& AvailableConversionHandles, bool bIncludeConversionScripts) const
+void UNiagaraStackFunctionInput::GetAvailableParameters(TArray<FNiagaraVariableBase>& OutAvailableParameters, TMap<FNiagaraVariableBase, UNiagaraScript*>& OutAvailableParametersWithConversions, bool bIncludeConversionScripts) const
 {
 	TMap<FNiagaraTypeDefinition, UNiagaraScript*> ConversionScriptCache;
 	
@@ -1769,11 +1784,11 @@ void UNiagaraStackFunctionInput::GetAvailableParameterHandles(TArray<FNiagaraPar
 	{
 		if (FNiagaraEditorUtilities::AreTypesAssignable(SystemVariable.GetType(), InputType))
 		{
-			AvailableParameterHandles.Add(FNiagaraParameterHandle::CreateEngineParameterHandle(SystemVariable));
+			OutAvailableParameters.Add(SystemVariable);
 		}
 		else if (UNiagaraScript* ConversionScript = FindConversionScript(SystemVariable.GetType(), ConversionScriptCache, bIncludeConversionScripts))
 		{
-			AvailableConversionHandles.Add(SystemVariable, ConversionScript);
+			OutAvailableParametersWithConversions.Add(SystemVariable, ConversionScript);
 		}
 	}
 
@@ -1784,11 +1799,11 @@ void UNiagaraStackFunctionInput::GetAvailableParameterHandles(TArray<FNiagaraPar
 	{
 		if (FNiagaraEditorUtilities::AreTypesAssignable(ExposedVar.GetType(), InputType))
 		{
-			AvailableParameterHandles.Add(FNiagaraParameterHandle::CreateEngineParameterHandle(ExposedVar));
+			OutAvailableParameters.Add(ExposedVar);
 		}
 		else if (UNiagaraScript* ConversionScript = FindConversionScript(ExposedVar.GetType(), ConversionScriptCache, bIncludeConversionScripts))
 		{
-			AvailableConversionHandles.Add(ExposedVar, ConversionScript);
+			OutAvailableParametersWithConversions.Add(ExposedVar, ConversionScript);
 		}
 	}
 
@@ -1805,9 +1820,16 @@ void UNiagaraStackFunctionInput::GetAvailableParameterHandles(TArray<FNiagaraPar
 		GetSystemViewModel()->GetSystemScriptViewModel()->GetGraphViewModel()->GetGraph()->GetNodesOfClass<UNiagaraNodeOutput>(AllOutputNodes);
 	}
 	UNiagaraNodeOutput* CurrentOutputNode = FNiagaraStackGraphUtilities::GetEmitterOutputNodeForStackNode(*OwningModuleNode);
+	ENiagaraScriptUsage CurrentUsage = CurrentOutputNode ? CurrentOutputNode->GetUsage() : ENiagaraScriptUsage::Module;
+	FName CurrentNamespace = GetNamespaceForUsage(CurrentUsage);
 	TArray<FName> StackContextRoots = FNiagaraStackGraphUtilities::StackContextResolution(Emitter, CurrentOutputNode);
 	for (UNiagaraNodeOutput* OutputNode : AllOutputNodes)
 	{
+		if (OutputNode == nullptr)
+		{
+			continue;
+		}
+
 		// Check if this is in a spawn event handler and the emitter is not using interpolated spawn so we
 		// we can hide particle update parameters
 		bool bSpawnScript = false;
@@ -1826,7 +1848,9 @@ void UNiagaraStackFunctionInput::GetAvailableParameterHandles(TArray<FNiagaraPar
 		bool bCheckInterpSpawn = bInterpolatedSpawn || !bSpawnScript;
 		if (OutputNode == CurrentOutputNode || (CurrentOutputNode != nullptr && UsageRunsBefore(OutputNode->GetUsage(), CurrentOutputNode->GetUsage(), bCheckInterpSpawn, Emitter, OutputNode->GetUsageId(), CurrentOutputNode->GetUsageId())) || (CurrentOutputNode != nullptr && IsSpawnUsage(CurrentOutputNode->GetUsage())))
 		{
-			TArray<FNiagaraParameterHandle> AvailableParameterHandlesForThisOutput;
+			FName OutputNamespace = GetNamespaceForUsage(OutputNode->GetUsage());
+			TArray<TPair<FNiagaraVariableBase, FNiagaraParameterHandle>> AvailableParametersForThisOutput;
+
 			TArray<FNiagaraStackGraphUtilities::FStackNodeGroup> StackGroups;
 			FNiagaraStackGraphUtilities::GetStackNodeGroups(*OutputNode, StackGroups);
 
@@ -1848,6 +1872,11 @@ void UNiagaraStackFunctionInput::GetAvailableParameterHandles(TArray<FNiagaraPar
 						FNiagaraVariable& HistoryVariable = Histories[0].Variables[j];
 						FNiagaraParameterHandle AvailableHandle = FNiagaraParameterHandle(HistoryVariable.GetName());
 
+						if (IsSourceNamespaceAllowedInTargetNamespace(AvailableHandle.GetNamespace(), CurrentNamespace) == false)
+						{
+							continue;
+						}
+
 						// check if the variable was written to
 						bool bWritten = false;
 						for (const FNiagaraParameterMapHistory::FModuleScopedPin& WritePin : Histories[0].PerVariableWriteHistory[j])
@@ -1859,53 +1888,56 @@ void UNiagaraStackFunctionInput::GetAvailableParameterHandles(TArray<FNiagaraPar
 							}
 						}
 
+						if (bWritten == false)
+						{
+							continue;
+						}
+
 						// as transient variables & output variables are only valid within their current stage, we may want to skip the current parameter handle
 						if(OutputNode != CurrentOutputNode && (AvailableHandle.IsOutputHandle() || AvailableHandle.IsTransientHandle()))
 						{
 							continue;
 						}
 
-						if (bWritten)
+						if (FNiagaraEditorUtilities::AreTypesAssignable(HistoryVariable.GetType(), InputType))
 						{
-							if (FNiagaraEditorUtilities::AreTypesAssignable(HistoryVariable.GetType(), InputType))
+							OutAvailableParameters.AddUnique(HistoryVariable);
+							if (AvailableHandle.GetNamespace() == OutputNamespace)
 							{
-								AvailableParameterHandles.AddUnique(AvailableHandle);
-								AvailableParameterHandlesForThisOutput.AddUnique(AvailableHandle);
+								AvailableParametersForThisOutput.Add(TPair<FNiagaraVariableBase, FNiagaraParameterHandle>(HistoryVariable, AvailableHandle));
+							}
 
-								// Check to see if any variables can be converted to StackContext. This may be more portable for people to setup.
-								for (int32 StackRootIdx = 0; StackRootIdx < StackContextRoots.Num(); StackRootIdx++)
+							// Check to see if any variables can be converted to StackContext. This may be more portable for people to setup.
+							for (int32 StackRootIdx = 0; StackRootIdx < StackContextRoots.Num(); StackRootIdx++)
+							{
+								if (HistoryVariable.IsInNameSpace(StackContextRoots[StackRootIdx]))
 								{
-									if (HistoryVariable.IsInNameSpace(StackContextRoots[StackRootIdx]))
-									{
-										// We do a replace here so that we can leave modifiers and other parts intact that might also be aliased.
-										FString NewName = HistoryVariable.GetName().ToString().Replace(*StackContextRoots[StackRootIdx].ToString(), *FNiagaraConstants::StackContextNamespace.ToString());
+									// We do a replace here so that we can leave modifiers and other parts intact that might also be aliased.
+									FString NewName = HistoryVariable.GetName().ToString().Replace(*StackContextRoots[StackRootIdx].ToString(), *FNiagaraConstants::StackContextNamespace.ToString());
 
-										FNiagaraParameterHandle AvailableAliasedHandle = FNiagaraParameterHandle(*NewName);
-										AvailableParameterHandles.AddUnique(AvailableAliasedHandle);
-										AvailableParameterHandlesForThisOutput.AddUnique(AvailableAliasedHandle);
-									}
+									FNiagaraVariableBase AvailableAliasedParameter = FNiagaraVariableBase(HistoryVariable.GetType(), *NewName);
+									OutAvailableParameters.AddUnique(AvailableAliasedParameter);
 								}
 							}
-							else if (UNiagaraScript* ConversionScript = FindConversionScript(HistoryVariable.GetType(), ConversionScriptCache, bIncludeConversionScripts))
-							{
-								AvailableConversionHandles.Add(HistoryVariable, ConversionScript);
-							}
+						}
+						else if (UNiagaraScript* ConversionScript = FindConversionScript(HistoryVariable.GetType(), ConversionScriptCache, bIncludeConversionScripts))
+						{
+							OutAvailableParametersWithConversions.Add(HistoryVariable, ConversionScript);
 						}
 					}
 				}
 			}
 
-			if (OutputNode != CurrentOutputNode && IsSpawnUsage(OutputNode->GetUsage()))
+			if (OutputNode != CurrentOutputNode && IsSpawnUsage(OutputNode->GetUsage()) && OutputNamespace != NAME_None)
 			{
-				FName OutputNodeNamespace = GetNamespaceForUsage(OutputNode->GetUsage());
-				if (OutputNodeNamespace.IsNone() == false)
+				for (const TPair<FNiagaraVariableBase, FNiagaraParameterHandle>& AvailableParameterForThisOutput : AvailableParametersForThisOutput)
 				{
-					for (FNiagaraParameterHandle& AvailableParameterHandleForThisOutput : AvailableParameterHandlesForThisOutput)
+					const FNiagaraVariableBase& AvailableParameter = AvailableParameterForThisOutput.Key;
+					const FNiagaraParameterHandle& AvailableHandle = AvailableParameterForThisOutput.Value;
+					if (AvailableParameter.GetType().IsStatic() == false)
 					{
-						if (AvailableParameterHandleForThisOutput.GetNamespace() == OutputNodeNamespace)
-						{
-							AvailableParameterHandles.AddUnique(FNiagaraParameterHandle::CreateInitialParameterHandle(AvailableParameterHandleForThisOutput));
-						}
+						FNiagaraParameterHandle InitialParameterHandle = FNiagaraParameterHandle::CreateInitialParameterHandle(AvailableHandle);
+						OutAvailableParameters.AddUnique(FNiagaraVariableBase(AvailableParameter.GetType(), InitialParameterHandle.GetParameterHandleString()));
 					}
 				}
 			}
@@ -1921,32 +1953,14 @@ void UNiagaraStackFunctionInput::GetAvailableParameterHandles(TArray<FNiagaraPar
 		{
 			if (FNiagaraEditorUtilities::AreTypesAssignable(CollectionParam.GetType(), InputType))
 			{
-				AvailableParameterHandles.AddUnique(FNiagaraParameterHandle(CollectionParam.GetName()));
+				OutAvailableParameters.AddUnique(CollectionParam);
 			}
 			else if (UNiagaraScript* ConversionScript = FindConversionScript(CollectionParam.GetType(), ConversionScriptCache, bIncludeConversionScripts))
 			{
-				AvailableConversionHandles.Add(CollectionParam, ConversionScript);
+				OutAvailableParametersWithConversions.Add(CollectionParam, ConversionScript);
 			}
 		}
 	}
-
-	// Filter gathered attributes by allowed usage
-	ENiagaraScriptUsage CurrentUsage = CurrentOutputNode ? CurrentOutputNode->GetUsage() : ENiagaraScriptUsage::Module;
-	AvailableParameterHandles.SetNum(Algo::RemoveIf(AvailableParameterHandles, [CurrentUsage](FNiagaraParameterHandle& Handle)
-	{
-		return !IsNamespaceAllowedInUsage(Handle.GetNamespace(), CurrentUsage);
-	}));
-	TArray<FNiagaraVariable> ParameterNames;
-	AvailableConversionHandles.GetKeys(ParameterNames);
-	for (int i = 0; i < ParameterNames.Num(); i++)
-	{
-		FNiagaraVariable Parameter = ParameterNames[i];
-		FNiagaraParameterHandle Handle(Parameter.GetName());
-		if (!IsNamespaceAllowedInUsage(Handle.GetNamespace(), CurrentUsage))
-		{
-			AvailableConversionHandles.Remove(Parameter);
-		}
-	}
 }
 
 UNiagaraNodeFunctionCall* UNiagaraStackFunctionInput::GetDefaultFunctionNode() const
@@ -2285,7 +2299,7 @@ bool UNiagaraStackFunctionInput::CanReset() const
 				bNewCanReset = false;
 				break;
 			case EValueMode::Linked:
-				bNewCanReset = InputValues.LinkedHandle != DefaultInputValues.LinkedHandle;
+				bNewCanReset = InputValues.LinkedParameter != DefaultInputValues.LinkedParameter;
 				break;
 			case EValueMode::Local:
 				bNewCanReset =
@@ -2361,7 +2375,7 @@ void UNiagaraStackFunctionInput::Reset()
 		else if (DefaultInputValues.Mode == EValueMode::Linked)
 		{
 			FScopedTransaction ScopedTransaction(LOCTEXT("ResetInputLinkedValueTransaction", "Reset the input to its default linked value."));
-			SetLinkedValueHandle(DefaultInputValues.LinkedHandle);
+			SetLinkedParameterValue(DefaultInputValues.LinkedParameter);
 		}
 		else if (DefaultInputValues.Mode == EValueMode::Local)
 		{
@@ -2930,53 +2944,6 @@ TArray<UNiagaraScript*> UNiagaraStackFunctionInput::GetPossibleConversionScripts
 	return MatchingDynamicInputs;
 }
 
-void UNiagaraStackFunctionInput::SetLinkedInputViaConversionScript(const FName& LinkedInputName, const FNiagaraTypeDefinition& FromType)
-{
-	TArray<UNiagaraScript*> NiagaraScripts = GetPossibleConversionScripts(FromType);
-	int32 ScriptCount = NiagaraScripts.Num();
-	if (ScriptCount == 0)
-	{
-		return;
-	}
-	if (ScriptCount > 1)
-	{
-		FString ScriptNames;
-		for (UNiagaraScript* NiagaraScript : NiagaraScripts)
-		{
-			ScriptNames += NiagaraScript->GetPathName() + "\n";
-		}
-		FNiagaraEditorUtilities::WarnWithToastAndLog(FText::Format(TooManyConversionScripts, FText::FromString(ScriptNames)));
-	}
-	FScopedTransaction ScopedTransaction(LOCTEXT("SetConversionInput", "Make auto-convert dynamic input"));
-	SetDynamicInput(NiagaraScripts[0]);
-	for (UNiagaraStackFunctionInput* ChildInput : GetChildInputs())
-	{
-		if (FromType == ChildInput->GetInputType())
-		{
-			ChildInput->SetLinkedValueHandle(FNiagaraParameterHandle(LinkedInputName));
-			break;
-		}
-	}
-}
-
-void UNiagaraStackFunctionInput::SetLinkedInputViaConversionScript(const FNiagaraVariable& LinkedInput, UNiagaraScript* ConversionScript)
-{
-	if (ConversionScript == nullptr)
-	{
-		return;
-	}
-	FScopedTransaction ScopedTransaction(LOCTEXT("SetConversionInput", "Make auto-convert dynamic input"));
-	SetDynamicInput(ConversionScript);
-	for (UNiagaraStackFunctionInput* ChildInput : GetChildInputs())
-	{
-		if (LinkedInput.GetType() == ChildInput->GetInputType())
-		{
-			ChildInput->SetLinkedValueHandle(FNiagaraParameterHandle(LinkedInput.GetName()));
-			break;
-		}
-	}
-}
-
 void UNiagaraStackFunctionInput::SetClipboardContentViaConversionScript(const UNiagaraClipboardFunctionInput& ClipboardFunctionInput)
 {
 	TArray<UNiagaraScript*> NiagaraScripts = GetPossibleConversionScripts(ClipboardFunctionInput.InputType);
@@ -3041,7 +3008,7 @@ const UNiagaraClipboardFunctionInput* UNiagaraStackFunctionInput::ToClipboardFun
 		break;
 	}
 	case EValueMode::Linked:
-		ClipboardInput = UNiagaraClipboardFunctionInput::CreateLinkedValue(InOuter, InputName, InputType, bEditConditionValue, InputValues.LinkedHandle.GetParameterHandleString());
+		ClipboardInput = UNiagaraClipboardFunctionInput::CreateLinkedValue(InOuter, InputName, InputType, bEditConditionValue, InputValues.LinkedParameter);
 		break;
 	case EValueMode::Data:
 		ClipboardInput = UNiagaraClipboardFunctionInput::CreateDataValue(InOuter, InputName, InputType, bEditConditionValue, InputValues.DataObject.Get());
@@ -3110,7 +3077,7 @@ void UNiagaraStackFunctionInput::SetValueFromClipboardFunctionInput(const UNiaga
 			break;
 		}
 		case ENiagaraClipboardFunctionInputValueMode::Linked:
-			SetLinkedValueHandle(FNiagaraParameterHandle(ClipboardFunctionInput.Linked));
+			SetLinkedParameterValue(ClipboardFunctionInput.Linked);
 			break;
 		case ENiagaraClipboardFunctionInputValueMode::Data:
 		{
@@ -3260,7 +3227,7 @@ void UNiagaraStackFunctionInput::GetSearchItems(TArray<FStackSearchItem>& Search
 		}
 		else if (InputValues.Mode == EValueMode::Linked)
 		{
-			SearchItems.Add({ FName("LinkedParamName"), FText::FromName(InputValues.LinkedHandle.GetParameterHandleString()) });
+			SearchItems.Add({ FName("LinkedParamName"), FText::FromName(InputValues.LinkedParameter.GetName()) });
 		}
 		else if (InputValues.Mode == EValueMode::Dynamic && InputValues.DynamicNode.Get() != nullptr)
 		{
@@ -3294,8 +3261,7 @@ const UNiagaraStackFunctionInput::FCollectedUsageData& UNiagaraStackFunctionInpu
 			bool bFoundOverride = false;
 			if (InputValues.Mode == EValueMode::Linked)
 			{
-				FNiagaraVariableBase Var(InputType, InputValues.LinkedHandle.GetParameterHandleString());
-				bFoundOverride = ParamVM->IsVariableSelected(Var);
+				bFoundOverride = ParamVM->IsVariableSelected(InputValues.LinkedParameter);
 			}
 			else if (InputValues.Mode == EValueMode::Dynamic)
 			{
@@ -3430,7 +3396,7 @@ void UNiagaraStackFunctionInput::GetDefaultDataInterfaceValueFromDefaultPin(UEdG
 	else
 	{
 		// If there was no input node, try to get a linked data interface default.
-		GetDefaultLinkedHandleOrLinkedFunctionFromDefaultPin(DefaultPin, InInputValues);
+		GetDefaultLinkedParameterOrLinkedFunctionFromDefaultPin(DefaultPin, InInputValues);
 		if (InInputValues.Mode == EValueMode::None)
 		{
 			// If there is no specified default input and no linked value, use the CDO as the default data value.
@@ -3483,18 +3449,18 @@ void UNiagaraStackFunctionInput::GetDefaultLocalValueFromDefaultPin(UEdGraphPin*
 	FMemory::Memcpy(InInputValues.LocalStruct->GetStructMemory(), LocalValueVariable.GetData(), InputType.GetSize());
 }
 
-struct FLinkedHandleOrFunctionNode
+struct FLinkedParameterOrFunctionNode
 {
-	TOptional<FNiagaraParameterHandle> LinkedHandle;
+	TOptional<FNiagaraVariableBase> LinkedParameter;
 	TWeakObjectPtr<UNiagaraNodeFunctionCall> LinkedFunctionCallNode;
 };
 
-void UNiagaraStackFunctionInput::GetDefaultLinkedHandleOrLinkedFunctionFromDefaultPin(UEdGraphPin* DefaultPin, UNiagaraStackFunctionInput::FInputValues& InInputValues) const
+void UNiagaraStackFunctionInput::GetDefaultLinkedParameterOrLinkedFunctionFromDefaultPin(UEdGraphPin* DefaultPin, UNiagaraStackFunctionInput::FInputValues& InInputValues) const
 {
 	// A default pin linked to a parameter map set is a default linked value.  Linked inputs can be setup in a chain and the first
 	// available one will be used so that case must be handled too.  So first collect up the potential linked values and validate 
 	// the linked node structure.
-	TArray<FLinkedHandleOrFunctionNode> LinkedValues;
+	TArray<FLinkedParameterOrFunctionNode> LinkedValues;
 	UEdGraphPin* CurrentDefaultPin = DefaultPin;
 	while (CurrentDefaultPin != nullptr && CurrentDefaultPin->LinkedTo.Num() == 1)
 	{
@@ -3503,15 +3469,15 @@ void UNiagaraStackFunctionInput::GetDefaultLinkedHandleOrLinkedFunctionFromDefau
 		if (LinkedPin->GetOwningNode()->IsA<UNiagaraNodeParameterMapGet>())
 		{
 			UNiagaraNodeParameterMapGet* GetNode = CastChecked<UNiagaraNodeParameterMapGet>(LinkedPin->GetOwningNode());
-			FLinkedHandleOrFunctionNode LinkedValue;
-			LinkedValue.LinkedHandle = FNiagaraParameterHandle(*LinkedPin->GetName());
+			FLinkedParameterOrFunctionNode LinkedValue;
+			LinkedValue.LinkedParameter = UEdGraphSchema_Niagara::PinToNiagaraVariable(LinkedPin);
 			LinkedValues.Add(LinkedValue);
 			CurrentDefaultPin = GetNode->GetDefaultPin(LinkedPin);
 		}
 		else if (LinkedPin->GetOwningNode()->IsA<UNiagaraNodeFunctionCall>())
 		{
 			UNiagaraNodeFunctionCall* FunctionCallNode = CastChecked<UNiagaraNodeFunctionCall>(LinkedPin->GetOwningNode());
-			FLinkedHandleOrFunctionNode LinkedValue;
+			FLinkedParameterOrFunctionNode LinkedValue;
 			LinkedValue.LinkedFunctionCallNode = FunctionCallNode;
 			LinkedValues.Add(LinkedValue);
 			CurrentDefaultPin = nullptr;
@@ -3525,7 +3491,7 @@ void UNiagaraStackFunctionInput::GetDefaultLinkedHandleOrLinkedFunctionFromDefau
 		}
 	}
 
-	FLinkedHandleOrFunctionNode* ValueToUse = nullptr;
+	FLinkedParameterOrFunctionNode* ValueToUse = nullptr;
 	if (LinkedValues.Num() == 1)
 	{
 		ValueToUse = &LinkedValues[0];
@@ -3533,13 +3499,13 @@ void UNiagaraStackFunctionInput::GetDefaultLinkedHandleOrLinkedFunctionFromDefau
 	else if(LinkedValues.Num() > 1)
 	{
 		// If there are a chain of linked values use the first one that's available, otherwise just use the last one.
-		TArray<FNiagaraParameterHandle> AvailableHandles;
-		TMap<FNiagaraVariable, UNiagaraScript*> ConversionHandles;
-		GetAvailableParameterHandles(AvailableHandles, ConversionHandles, false);
-		for (FLinkedHandleOrFunctionNode& LinkedValue : LinkedValues)
+		TArray<FNiagaraVariableBase> AvailableParameters;
+		TMap<FNiagaraVariableBase, UNiagaraScript*> ConversionHandles;
+		GetAvailableParameters(AvailableParameters, ConversionHandles, false);
+		for (FLinkedParameterOrFunctionNode& LinkedValue : LinkedValues)
 		{
 			if (LinkedValue.LinkedFunctionCallNode.IsValid() ||
-				AvailableHandles.Contains(LinkedValue.LinkedHandle.GetValue()))
+				AvailableParameters.Contains(LinkedValue.LinkedParameter.GetValue()))
 			{
 				ValueToUse = &LinkedValue;
 				break;
@@ -3549,10 +3515,10 @@ void UNiagaraStackFunctionInput::GetDefaultLinkedHandleOrLinkedFunctionFromDefau
 
 	if (ValueToUse != nullptr)
 	{
-		if (ValueToUse->LinkedHandle.IsSet())
+		if (ValueToUse->LinkedParameter.IsSet())
 		{
 			InInputValues.Mode = EValueMode::Linked;
-			InInputValues.LinkedHandle = ValueToUse->LinkedHandle.GetValue();
+			InInputValues.LinkedParameter = ValueToUse->LinkedParameter.GetValue();
 		}
 		else
 		{
@@ -3599,7 +3565,7 @@ void UNiagaraStackFunctionInput::UpdateValuesFromScriptDefaults(FInputValues& In
 		{
 			// The next highest precedence value is a linked value from a variable binding so check that.
 			InInputValues.Mode = EValueMode::Linked;
-			InInputValues.LinkedHandle = DefaultBinding.GetName();
+			InInputValues.LinkedParameter = FNiagaraVariable(InputType, DefaultBinding.GetName());
 		}
 		else if (SourceScript.IsValid())
 		{
@@ -3638,7 +3604,7 @@ void UNiagaraStackFunctionInput::UpdateValuesFromScriptDefaults(FInputValues& In
 					else if(DefaultPin->LinkedTo.Num() == 1 && DefaultPin->LinkedTo[0]->GetOwningNode() != nullptr)
 					{
 						// If a default pin is linked to a parameter map it can be a linked value.
-						GetDefaultLinkedHandleOrLinkedFunctionFromDefaultPin(DefaultPin, InInputValues);
+						GetDefaultLinkedParameterOrLinkedFunctionFromDefaultPin(DefaultPin, InInputValues);
 					}
 
 					if(InInputValues.Mode == EValueMode::None)
@@ -3702,7 +3668,7 @@ void UNiagaraStackFunctionInput::UpdateValuesFromOverridePin(const FInputValues&
 		{
 			// Parameter map get nodes handle linked values.
 			NewInputValues.Mode = EValueMode::Linked;
-			NewInputValues.LinkedHandle = FNiagaraParameterHandle(InOverridePin.LinkedTo[0]->PinName);
+			NewInputValues.LinkedParameter = UEdGraphSchema_Niagara::PinToNiagaraVariable(InOverridePin.LinkedTo[0]);
 		}
 		else if (LinkedNode->IsA<UNiagaraNodeCustomHlsl>())
 		{
diff --git a/Engine/Plugins/FX/Niagara/Source/NiagaraEditor/Private/ViewModels/Stack/NiagaraStackGraphUtilities.cpp b/Engine/Plugins/FX/Niagara/Source/NiagaraEditor/Private/ViewModels/Stack/NiagaraStackGraphUtilities.cpp
--- a/Engine/Plugins/FX/Niagara/Source/NiagaraEditor/Private/ViewModels/Stack/NiagaraStackGraphUtilities.cpp
+++ b/Engine/Plugins/FX/Niagara/Source/NiagaraEditor/Private/ViewModels/Stack/NiagaraStackGraphUtilities.cpp
@@ -2084,25 +2084,26 @@ UEdGraphPin* FNiagaraStackGraphUtilities::GetLinkedValueHandleForFunctionInput(c
 	return GetOutputPin;
 }
 
-TSet<FNiagaraVariable> FNiagaraStackGraphUtilities::GetParametersForContext(UEdGraph* InGraph, UNiagaraSystem& System)
+void FNiagaraStackGraphUtilities::GetParametersForContext(UEdGraph* InGraph, UNiagaraSystem& System, TSet<FNiagaraVariableBase>& OutParameters)
 {
-	TSet<FNiagaraVariable> Result;
 	if (UNiagaraGraph* NiagaraGraph = Cast<UNiagaraGraph>(InGraph))
 	{
-		const TMap<FNiagaraVariable, FNiagaraGraphParameterReferenceCollection>& ReferenceMap = NiagaraGraph->GetParameterReferenceMap();
-		ReferenceMap.GetKeys(Result);
+		for (const TPair<FNiagaraVariable, FNiagaraGraphParameterReferenceCollection>& ReferencePair : NiagaraGraph->GetParameterReferenceMap())
+		{
+			OutParameters.Add(ReferencePair.Key);
+		}
 	}
+
 	TArray<FNiagaraVariable> UserParams;
 	System.GetExposedParameters().GetUserParameters(UserParams);
 	for (FNiagaraVariable& Var : UserParams)
 	{
 		FNiagaraUserRedirectionParameterStore::MakeUserVariable(Var);
+		OutParameters.Add(Var);
 	}
-	Result.Append(UserParams);
-	return Result;
 }
 
-void FNiagaraStackGraphUtilities::SetLinkedValueHandleForFunctionInput(UEdGraphPin& OverridePin, FNiagaraParameterHandle LinkedParameterHandle, const TSet<FNiagaraVariable>& KnownParameters, ENiagaraDefaultMode DesiredDefaultMode, const FGuid& NewNodePersistentId)
+void FNiagaraStackGraphUtilities::SetLinkedParameterValueForFunctionInput(UEdGraphPin& OverridePin, const FNiagaraVariableBase& LinkedParameter, const TSet<FNiagaraVariableBase>& KnownParameters, ENiagaraDefaultMode DesiredDefaultMode, const FGuid& NewNodePersistentId)
 {
 	checkf(OverridePin.LinkedTo.Num() == 0, TEXT("Can't set a linked value handle when the override pin already has a value."));
 	const UNiagaraSettings* Settings = GetDefault<UNiagaraSettings>();
@@ -2125,12 +2126,12 @@ void FNiagaraStackGraphUtilities::SetLinkedValueHandleForFunctionInput(UEdGraphP
 	const UEdGraphSchema_Niagara* NiagaraSchema = GetDefault<UEdGraphSchema_Niagara>();
 	FNiagaraTypeDefinition InputType = NiagaraSchema->PinToTypeDefinition(&OverridePin);
 
-	FNiagaraVariable ParameterToRead(InputType, LinkedParameterHandle.GetParameterHandleString());
+	FNiagaraVariableBase ParameterToRead = LinkedParameter;
 	if (Settings->bEnforceStrictStackTypes == false && (InputType == FNiagaraTypeDefinition::GetVec3Def() || InputType == FNiagaraTypeDefinition::GetPositionDef()))
 	{
 		// see if the requested input exists or if we can use one of the equivalent loose types
 		FNiagaraTypeDefinition AlternateType = InputType == FNiagaraTypeDefinition::GetVec3Def() ? FNiagaraTypeDefinition::GetPositionDef() : FNiagaraTypeDefinition::GetVec3Def();
-		FNiagaraVariable ParameterAlternative(AlternateType, LinkedParameterHandle.GetParameterHandleString());
+		FNiagaraVariable ParameterAlternative(AlternateType, LinkedParameter.GetName());
 		if (KnownParameters.Contains(ParameterToRead) == false && KnownParameters.Contains(ParameterAlternative))
 		{
 			ParameterToRead = ParameterAlternative;
@@ -2140,14 +2141,14 @@ void FNiagaraStackGraphUtilities::SetLinkedValueHandleForFunctionInput(UEdGraphP
 	{
 		// check if a static var of the requested input exists and use that instead
 		FNiagaraTypeDefinition StaticInputType = InputType.ToStaticDef();
-		FNiagaraVariable StaticParameter(StaticInputType, LinkedParameterHandle.GetParameterHandleString());
+		FNiagaraVariable StaticParameter(StaticInputType, LinkedParameter.GetName());
 		if (KnownParameters.Contains(StaticParameter))
 		{
 			ParameterToRead = StaticParameter;
 		}
 	}
 	
-	UEdGraphPin* GetOutputPin = GetNode->RequestNewTypedPin(EGPD_Output, ParameterToRead.GetType(), LinkedParameterHandle.GetParameterHandleString());
+	UEdGraphPin* GetOutputPin = GetNode->RequestNewTypedPin(EGPD_Output, ParameterToRead.GetType(), ParameterToRead.GetName());
 	MakeLinkTo(GetInputPin, PreviousStackNodeOutputPin);
 	MakeLinkTo(GetOutputPin, &OverridePin);
 
@@ -4258,10 +4259,15 @@ void FNiagaraStackGraphUtilities::SynchronizeReferencingMapPinsWithFunctionCall(
 	}
 }
 
-FGuid FNiagaraStackGraphUtilities::GetScriptVariableIdForLinkedModuleParameterHandle(const FNiagaraParameterHandle& LinkedHandle, FNiagaraTypeDefinition LinkedType, UNiagaraGraph& TargetGraph)
+FGuid FNiagaraStackGraphUtilities::GetScriptVariableIdForLinkedModuleParameter(const FNiagaraVariableBase& LinkedParameter, UNiagaraGraph& TargetGraph)
 {
-	if (LinkedHandle.GetHandleParts().Num() >= 3 &&
-		(LinkedHandle.IsOutputHandle() || LinkedHandle.IsStackContextHandle() || LinkedHandle.IsSystemHandle() || LinkedHandle.IsEmitterHandle() || LinkedHandle.IsParticleAttributeHandle()))
+	FNiagaraParameterHandle LinkedHandle = FNiagaraParameterHandle(LinkedParameter.GetName());
+	if (LinkedHandle.GetHandleParts().Num() >= 3 && (
+			LinkedHandle.IsOutputHandle() ||
+			LinkedHandle.IsStackContextHandle() ||
+			LinkedHandle.IsSystemHandle() ||
+			LinkedHandle.IsEmitterHandle() ||
+			LinkedHandle.IsParticleAttributeHandle()))
 	{
 		// First find the source function call node.
 		FString LinkedFunctionName = LinkedHandle.GetHandleParts()[1].ToString();
@@ -4274,8 +4280,7 @@ FGuid FNiagaraStackGraphUtilities::GetScriptVariableIdForLinkedModuleParameterHa
 			UNiagaraGraph* CalledGraph = (*SourceFunctionCallNodePtr)->GetCalledGraph();
 			if (CalledGraph != nullptr)
 			{
-				FNiagaraVariable LinkedVariable(LinkedType, LinkedHandle.GetParameterHandleString());
-				FNiagaraVariable LinkedModuleVariable = FNiagaraUtilities::ResolveAliases(LinkedVariable, FNiagaraAliasContext()
+				FNiagaraVariable LinkedModuleVariable = FNiagaraUtilities::ResolveAliases(LinkedParameter, FNiagaraAliasContext()
 					.ChangeModuleNameToModule(LinkedFunctionName));
 				for (const TPair<FNiagaraVariable, TObjectPtr<UNiagaraScriptVariable>>& VariableMetadataPair : CalledGraph->GetAllMetaData())
 				{
diff --git a/Engine/Plugins/FX/Niagara/Source/NiagaraEditor/Private/ViewModels/Stack/NiagaraStackScriptItemGroup.cpp b/Engine/Plugins/FX/Niagara/Source/NiagaraEditor/Private/ViewModels/Stack/NiagaraStackScriptItemGroup.cpp
--- a/Engine/Plugins/FX/Niagara/Source/NiagaraEditor/Private/ViewModels/Stack/NiagaraStackScriptItemGroup.cpp
+++ b/Engine/Plugins/FX/Niagara/Source/NiagaraEditor/Private/ViewModels/Stack/NiagaraStackScriptItemGroup.cpp
@@ -1410,10 +1410,11 @@ void RenameInputsFromClipboard(TMap<FName, FName> OldModuleOutputNameToNewModule
 		TOptional<bool> bEditConditionValue = SourceInput->bHasEditCondition 
 			? TOptional<bool>(SourceInput->bEditConditionValue)
 			: TOptional<bool>();
-		if (SourceInput->ValueMode == ENiagaraClipboardFunctionInputValueMode::Linked && OldModuleOutputNameToNewModuleOutputNameMap.Contains(SourceInput->Linked))
+		if (SourceInput->ValueMode == ENiagaraClipboardFunctionInputValueMode::Linked && OldModuleOutputNameToNewModuleOutputNameMap.Contains(SourceInput->Linked.GetName()))
 		{
-			FName RenamedLinkedValue = OldModuleOutputNameToNewModuleOutputNameMap[SourceInput->Linked];
-			OutRenamedInputs.Add(UNiagaraClipboardFunctionInput::CreateLinkedValue(InOuter, SourceInput->InputName, SourceInput->InputType, bEditConditionValue, RenamedLinkedValue));
+			FName RenamedLinkedValue = OldModuleOutputNameToNewModuleOutputNameMap[SourceInput->Linked.GetName()];
+			FNiagaraVariableBase RenamedLinkedParameter = FNiagaraVariableBase(SourceInput->Linked.GetType(), RenamedLinkedValue);
+			OutRenamedInputs.Add(UNiagaraClipboardFunctionInput::CreateLinkedValue(InOuter, SourceInput->InputName, SourceInput->InputType, bEditConditionValue, RenamedLinkedParameter));
 		}
 		else if (SourceInput->ValueMode == ENiagaraClipboardFunctionInputValueMode::Dynamic)
 		{
diff --git a/Engine/Plugins/FX/Niagara/Source/NiagaraEditor/Public/NiagaraClipboard.h b/Engine/Plugins/FX/Niagara/Source/NiagaraEditor/Public/NiagaraClipboard.h
--- a/Engine/Plugins/FX/Niagara/Source/NiagaraEditor/Public/NiagaraClipboard.h
+++ b/Engine/Plugins/FX/Niagara/Source/NiagaraEditor/Public/NiagaraClipboard.h
@@ -41,7 +41,7 @@ class UNiagaraClipboardFunctionInput : public UObject
 public:
 	static NIAGARAEDITOR_API const UNiagaraClipboardFunctionInput* CreateLocalValue(UObject* InOuter, FName InInputName, FNiagaraTypeDefinition InInputType, TOptional<bool> bInEditConditionValue, TArray<uint8>& InLocalValueData);
 
-	static NIAGARAEDITOR_API const UNiagaraClipboardFunctionInput* CreateLinkedValue(UObject* InOuter, FName InInputName, FNiagaraTypeDefinition InInputType, TOptional<bool> bInEditConditionValue, FName InLinkedValue);
+	static NIAGARAEDITOR_API const UNiagaraClipboardFunctionInput* CreateLinkedValue(UObject* InOuter, FName InInputName, FNiagaraTypeDefinition InInputType, TOptional<bool> bInEditConditionValue, const FNiagaraVariableBase& InLinkedValue);
 
 	static NIAGARAEDITOR_API const UNiagaraClipboardFunctionInput* CreateDataValue(UObject* InOuter, FName InInputName, FNiagaraTypeDefinition InInputType, TOptional<bool> bInEditConditionValue, UNiagaraDataInterface* InDataValue);
 
@@ -72,7 +72,7 @@ public:
 	TArray<uint8> Local;
 
 	UPROPERTY()
-	FName Linked;
+	FNiagaraVariableBase Linked;
 
 	UPROPERTY()
 	TObjectPtr<UNiagaraDataInterface> Data;
diff --git a/Engine/Plugins/FX/Niagara/Source/NiagaraEditor/Public/ViewModels/NiagaraParameterPanelViewModel.h b/Engine/Plugins/FX/Niagara/Source/NiagaraEditor/Public/ViewModels/NiagaraParameterPanelViewModel.h
--- a/Engine/Plugins/FX/Niagara/Source/NiagaraEditor/Public/ViewModels/NiagaraParameterPanelViewModel.h
+++ b/Engine/Plugins/FX/Niagara/Source/NiagaraEditor/Public/ViewModels/NiagaraParameterPanelViewModel.h
@@ -285,7 +285,7 @@ public:
 	virtual void SetMainParameterPanelViewModel(TWeakPtr<INiagaraParameterPanelViewModel> InMainParameterPanelViewModel) { MainParameterPanelViewModel = InMainParameterPanelViewModel; }
 
 
-	virtual bool IsVariableSelected(FNiagaraVariableBase& InVar) const { return false; };
+	virtual bool IsVariableSelected(const FNiagaraVariableBase& InVar) const { return false; };
 
 	virtual TSharedPtr<FNiagaraObjectSelection> GetVariableObjectSelection() { return TSharedPtr<FNiagaraObjectSelection>(); };
 	
@@ -382,7 +382,7 @@ public:
 	void DebugParameters(const TArray<FNiagaraParameterPanelItem> ItemToDuplicate) const;
 
 	virtual void OnParameterItemSelected(const FNiagaraParameterPanelItem& SelectedItem, ESelectInfo::Type SelectInfo) const override;
-	virtual bool IsVariableSelected(FNiagaraVariableBase& InVar) const override;
+	virtual bool IsVariableSelected(const FNiagaraVariableBase& InVar) const override;
 	virtual void RefreshDueToActiveDocumentChanged() override;
 
 	virtual bool GetSectionEnabled(FText Section) const override;
diff --git a/Engine/Plugins/FX/Niagara/Source/NiagaraEditor/Public/ViewModels/Stack/NiagaraStackFunctionInput.h b/Engine/Plugins/FX/Niagara/Source/NiagaraEditor/Public/ViewModels/Stack/NiagaraStackFunctionInput.h
--- a/Engine/Plugins/FX/Niagara/Source/NiagaraEditor/Public/ViewModels/Stack/NiagaraStackFunctionInput.h
+++ b/Engine/Plugins/FX/Niagara/Source/NiagaraEditor/Public/ViewModels/Stack/NiagaraStackFunctionInput.h
@@ -144,13 +144,16 @@ public:
 	NIAGARAEDITOR_API const FNiagaraParameterHandle& GetInputParameterHandle() const;
 
 	/** Gets the handle to the linked value for this input if there is one. */
-	NIAGARAEDITOR_API const FNiagaraParameterHandle& GetLinkedValueHandle() const;
+	NIAGARAEDITOR_API const FNiagaraVariableBase& GetLinkedParameterValue() const;
 
 	/** Sets the value of this input to a linked parameter handle. */
-	NIAGARAEDITOR_API void SetLinkedValueHandle(const FNiagaraParameterHandle& InParameterHandle);
+	NIAGARAEDITOR_API void SetLinkedParameterValue(const FNiagaraVariableBase& InParameter);
+
+	/** Sets the value of this input to a linked parameter handle using a conversion script. */
+	NIAGARAEDITOR_API void SetLinkedParameterValueViaConversionScript(const FNiagaraVariableBase& InParameter, UNiagaraScript& ConversionScript);
 
 	/** Gets the current set of available parameter handles which can be assigned to this input. Optionally returns possible conversion scripts. */
-	NIAGARAEDITOR_API void GetAvailableParameterHandles(TArray<FNiagaraParameterHandle>& AvailableParameterHandles, TMap<FNiagaraVariable, UNiagaraScript*>& AvailableConversionHandles, bool bIncludeConversionScripts = true) const;
+	NIAGARAEDITOR_API void GetAvailableParameters(TArray<FNiagaraVariableBase>& OutAvailableParameters, TMap<FNiagaraVariableBase, UNiagaraScript*>& OutAvailableParametersWithConversions, bool bIncludeConversionScripts = true) const;
 
 	/** Gets the function node form the script graph if the current value mode is DefaultFunction. */
 	NIAGARAEDITOR_API UNiagaraNodeFunctionCall* GetDefaultFunctionNode() const;
@@ -281,8 +284,6 @@ public:
 	NIAGARAEDITOR_API TArray<UNiagaraScript*> GetPossibleConversionScripts(const FNiagaraTypeDefinition& FromType) const;
 	static NIAGARAEDITOR_API TArray<UNiagaraScript*> GetPossibleConversionScripts(const FNiagaraTypeDefinition& FromType, const FNiagaraTypeDefinition& ToType);
 
-	NIAGARAEDITOR_API void SetLinkedInputViaConversionScript(const FName& LinkedInputName, const FNiagaraTypeDefinition& FromType);
-	NIAGARAEDITOR_API void SetLinkedInputViaConversionScript(const FNiagaraVariable& LinkedInput, UNiagaraScript* ConversionScript);
 	NIAGARAEDITOR_API void SetClipboardContentViaConversionScript(const UNiagaraClipboardFunctionInput& ClipboardFunctionInput);
 
 	NIAGARAEDITOR_API void ChangeScriptVersion(FGuid NewScriptVersion);
@@ -344,7 +345,7 @@ private:
 
 		EValueMode Mode;
 		TSharedPtr<FStructOnScope> LocalStruct;
-		FNiagaraParameterHandle LinkedHandle;
+		FNiagaraVariableBase LinkedParameter;
 		TWeakObjectPtr<UNiagaraNodeFunctionCall> DynamicNode;
 		TWeakObjectPtr<UNiagaraNodeCustomHlsl> ExpressionNode;
 		TWeakObjectPtr<UNiagaraDataInterface> DataObject;
@@ -388,7 +389,7 @@ private:
 
 	NIAGARAEDITOR_API void GetDefaultLocalValueFromDefaultPin(UEdGraphPin* DefaultPin, UNiagaraStackFunctionInput::FInputValues& InInputValues) const;
 
-	NIAGARAEDITOR_API void GetDefaultLinkedHandleOrLinkedFunctionFromDefaultPin(UEdGraphPin* DefaultPin, UNiagaraStackFunctionInput::FInputValues& InInputValues) const;
+	NIAGARAEDITOR_API void GetDefaultLinkedParameterOrLinkedFunctionFromDefaultPin(UEdGraphPin* DefaultPin, UNiagaraStackFunctionInput::FInputValues& InInputValues) const;
 
 	NIAGARAEDITOR_API void UpdateValuesFromScriptDefaults(FInputValues& InInputValues) const;
 
diff --git a/Engine/Plugins/FX/Niagara/Source/NiagaraEditor/Public/ViewModels/Stack/NiagaraStackGraphUtilities.h b/Engine/Plugins/FX/Niagara/Source/NiagaraEditor/Public/ViewModels/Stack/NiagaraStackGraphUtilities.h
--- a/Engine/Plugins/FX/Niagara/Source/NiagaraEditor/Public/ViewModels/Stack/NiagaraStackGraphUtilities.h
+++ b/Engine/Plugins/FX/Niagara/Source/NiagaraEditor/Public/ViewModels/Stack/NiagaraStackGraphUtilities.h
@@ -223,8 +223,8 @@ namespace FNiagaraStackGraphUtilities
 
 	UEdGraphPin* GetLinkedValueHandleForFunctionInput(const UEdGraphPin& OverridePin);
 
-	TSet<FNiagaraVariable> GetParametersForContext(UEdGraph* Graph, UNiagaraSystem& System);
-	NIAGARAEDITOR_API void SetLinkedValueHandleForFunctionInput(UEdGraphPin& OverridePin, FNiagaraParameterHandle LinkedParameterHandle, const TSet<FNiagaraVariable>& KnownParameters, ENiagaraDefaultMode DesiredDefaultMode = ENiagaraDefaultMode::FailIfPreviouslyNotSet, const FGuid& NewNodePersistentId = FGuid());
+	void GetParametersForContext(UEdGraph* Graph, UNiagaraSystem& System, TSet<FNiagaraVariableBase>& OutParameters);
+	NIAGARAEDITOR_API void SetLinkedParameterValueForFunctionInput(UEdGraphPin& OverridePin, const FNiagaraVariableBase& LinkedParameter, const TSet<FNiagaraVariableBase>& KnownParameters, ENiagaraDefaultMode DesiredDefaultMode = ENiagaraDefaultMode::FailIfPreviouslyNotSet, const FGuid& NewNodePersistentId = FGuid());
 
 	NIAGARAEDITOR_API void SetDataInterfaceValueForFunctionInput(UEdGraphPin& OverridePin, UClass* DataObjectType, FString InputNodeInputName, UNiagaraDataInterface*& OutDataObject, const FGuid& NewNodePersistentId = FGuid());
 
@@ -366,7 +366,7 @@ namespace FNiagaraStackGraphUtilities
 
 	NIAGARAEDITOR_API void SynchronizeReferencingMapPinsWithFunctionCall(UNiagaraNodeFunctionCall& InFunctionCallNode);
 
-	FGuid GetScriptVariableIdForLinkedModuleParameterHandle(const FNiagaraParameterHandle& LinkedOutputHandle, FNiagaraTypeDefinition LinkedType, UNiagaraGraph& TargetGraph);
+	FGuid GetScriptVariableIdForLinkedModuleParameter(const FNiagaraVariableBase& LinkedParameter, UNiagaraGraph& TargetGraph);
 
 	/** Fixes orphaned and disconnected output pins on dynamic input nodes resulting from the underlying script changing and the pins being reallocated. */
 	void FixDynamicInputNodeOutputPinsFromExternalChanges(UNiagaraNodeFunctionCall& InFunctionCallNode);
diff --git a/Engine/Plugins/FX/Niagara/Source/NiagaraEditorWidgets/Private/Stack/SNiagaraStackFunctionInputValue.cpp b/Engine/Plugins/FX/Niagara/Source/NiagaraEditorWidgets/Private/Stack/SNiagaraStackFunctionInputValue.cpp
--- a/Engine/Plugins/FX/Niagara/Source/NiagaraEditorWidgets/Private/Stack/SNiagaraStackFunctionInputValue.cpp
+++ b/Engine/Plugins/FX/Niagara/Source/NiagaraEditorWidgets/Private/Stack/SNiagaraStackFunctionInputValue.cpp
@@ -320,9 +320,9 @@ TSharedRef<SWidget> SNiagaraStackFunctionInputValue::ConstructValueWidgets()
 			.ParameterName(this, &SNiagaraStackFunctionInputValue::GetLinkedValueHandleName)
 			.OnDoubleClicked(this, &SNiagaraStackFunctionInputValue::OnLinkedInputDoubleClicked);
 
-		if(FunctionInput->GetLinkedValueHandle().IsUserHandle())
+		FNiagaraParameterHandle ParameterHandle = FNiagaraParameterHandle(FunctionInput->GetLinkedParameterValue().GetName());
+		if(ParameterHandle.IsUserHandle())
 		{
-			FNiagaraParameterHandle ParameterHandle = FunctionInput->GetLinkedValueHandle();
 			TArray<FNiagaraVariable> UserParameters;
 			FunctionInput->GetSystemViewModel()->GetSystem().GetExposedParameters().GetUserParameters(UserParameters);
 			FNiagaraVariable* MatchingVariable = UserParameters.FindByPredicate([ParameterHandle](const FNiagaraVariable& Variable)
@@ -689,7 +689,7 @@ void SNiagaraStackFunctionInputValue::ParameterPropertyValueChanged(const FPrope
 
 FName SNiagaraStackFunctionInputValue::GetLinkedValueHandleName() const
 {
-	return FunctionInput->GetLinkedValueHandle().GetParameterHandleString();
+	return FunctionInput->GetLinkedParameterValue().GetName();
 }
 
 FText SNiagaraStackFunctionInputValue::GetDataValueText() const
@@ -777,7 +777,9 @@ FReply SNiagaraStackFunctionInputValue::OnLinkedInputDoubleClicked(const FGeomet
 {
 	FString ParamCollection;
 	FString ParamName;
-	FunctionInput->GetLinkedValueHandle().GetName().ToString().Split(TEXT("."), &ParamCollection, &ParamName);
+
+	FNiagaraParameterHandle ParameterHandle = FNiagaraParameterHandle(FunctionInput->GetLinkedParameterValue().GetName());
+	ParameterHandle.GetName().ToString().Split(TEXT("."), &ParamCollection, &ParamName);
 
 	TArray<UNiagaraParameterCollection*> AvailableParameterCollections;
 	FNiagaraEditorUtilities::GetAvailableParameterCollections(AvailableParameterCollections);
@@ -936,14 +938,14 @@ void SNiagaraStackFunctionInputValue::CreateScratchSelected()
 	FunctionInput->SetScratch();
 }
 
-void SNiagaraStackFunctionInputValue::ParameterHandleSelected(FNiagaraParameterHandle Handle)
+void SNiagaraStackFunctionInputValue::ParameterSelected(FNiagaraVariableBase Parameter)
 {
-	FunctionInput->SetLinkedValueHandle(Handle);
+	FunctionInput->SetLinkedParameterValue(Parameter);
 }
 
-void SNiagaraStackFunctionInputValue::ConversionHandleSelected(FNiagaraVariable Handle, UNiagaraScript* ConversionScript)
+void SNiagaraStackFunctionInputValue::ParameterWithConversionSelected(FNiagaraVariableBase Parameter, UNiagaraScript* ConversionScript)
 {
-	FunctionInput->SetLinkedInputViaConversionScript(Handle, ConversionScript);
+	FunctionInput->SetLinkedParameterValueViaConversionScript(Parameter, *ConversionScript);
 }
 
 EVisibility SNiagaraStackFunctionInputValue::GetResetButtonVisibility() const
@@ -1035,14 +1037,19 @@ FReply SNiagaraStackFunctionInputValue::OnFunctionInputDrop(const FGeometry& InG
 			if (FNiagaraEditorUtilities::AreTypesAssignable(FromType, ToType))
 			{
 				// the types are the same, so we can just link the value directly
-				FunctionInput->SetLinkedValueHandle(FNiagaraParameterHandle(Action->GetParameter().GetName()));
+				FunctionInput->SetLinkedParameterValue(Action->GetParameter());
+				return FReply::Handled();
 			}
 			else
 			{
 				// the types don't match, so we use a dynamic input to convert from one to the other
-				FunctionInput->SetLinkedInputViaConversionScript(Action->GetParameter().GetName(), FromType);
+				TArray<UNiagaraScript*> ConversionScripts = FunctionInput->GetPossibleConversionScripts(FromType);
+				if (ConversionScripts.Num() > 0)
+				{
+					FunctionInput->SetLinkedParameterValueViaConversionScript(Action->GetParameter(), *ConversionScripts[0]);
+					return FReply::Handled();
+				}
 			}
-			return FReply::Handled();
 		}
 	}
 
@@ -1281,14 +1288,18 @@ TArray<TSharedPtr<FNiagaraMenuAction_Generic>> SNiagaraStackFunctionInputValue::
 	}
 
 	// Link existing attribute
-	TArray<FNiagaraParameterHandle> AvailableHandles;
-	TMap<FNiagaraVariable, UNiagaraScript*> AvailableConversionHandles;
-	FunctionInput->GetAvailableParameterHandles(AvailableHandles, AvailableConversionHandles);
+	const UNiagaraSettings* Settings = GetDefault<UNiagaraSettings>();
+	bool bAllowConversions = ensure(Settings != nullptr) && Settings->bShowConvertibleInputsInStack;
+
+	TArray<FNiagaraVariableBase> AvailableParameters;
+	TMap<FNiagaraVariableBase, UNiagaraScript*> AvailableParametersWithConversion;
+	FunctionInput->GetAvailableParameters(AvailableParameters, AvailableParametersWithConversion, bAllowConversions);
 
 	const FString RootCategoryName = FString("Link Inputs");
 	const FText MapInputFormat = LOCTEXT("LinkInputFormat", "Link this input to {0}");
-	for (const FNiagaraParameterHandle& AvailableHandle : AvailableHandles)
+	for (const FNiagaraVariableBase& AvailableParameter : AvailableParameters)
 	{
+		FNiagaraParameterHandle AvailableHandle = FNiagaraParameterHandle(AvailableParameter.GetName());
 		TArray<FName> HandleParts = AvailableHandle.GetHandleParts();
 		FNiagaraNamespaceMetadata NamespaceMetadata = GetDefault<UNiagaraEditorSettings>()->GetMetaDataForNamespaces(HandleParts);
 		if (NamespaceMetadata.IsValid())
@@ -1300,7 +1311,7 @@ TArray<TSharedPtr<FNiagaraMenuAction_Generic>> SNiagaraStackFunctionInputValue::
 			const FText Tooltip = FText::Format(MapInputFormat, FText::FromName(AvailableHandle.GetParameterHandleString()));
 			
 			TSharedPtr<FNiagaraMenuAction_Generic> LinkAction(new FNiagaraMenuAction_Generic(
-				FNiagaraMenuAction_Generic::FOnExecuteAction::CreateSP(this, &SNiagaraStackFunctionInputValue::ParameterHandleSelected, AvailableHandle),
+				FNiagaraMenuAction_Generic::FOnExecuteAction::CreateSP(this, &SNiagaraStackFunctionInputValue::ParameterSelected, AvailableParameter),
 				DisplayName, ENiagaraMenuSections::General, {RootCategoryName, Category.ToString()}, Tooltip, FText()));
 			
 			LinkAction->SetParameterVariable(FNiagaraVariable(FunctionInput->GetInputType(), AvailableHandle.GetParameterHandleString()));
@@ -1310,15 +1321,10 @@ TArray<TSharedPtr<FNiagaraMenuAction_Generic>> SNiagaraStackFunctionInputValue::
 		}
 	}
 
-	const UNiagaraSettings* Settings = GetDefault<UNiagaraSettings>();
-	if (ensure(Settings) && Settings->bShowConvertibleInputsInStack == false)
-	{
-		AvailableConversionHandles.Empty();
-	}
 	const FText ConvertInputFormat = LOCTEXT("ConvertInputFormat", "Link this input to {0} via a conversion script");
-	for (const auto& Entry : AvailableConversionHandles)
+	for (const auto& Entry : AvailableParametersWithConversion)
 	{
-		FNiagaraVariable ParameterVariable = Entry.Key;
+		FNiagaraVariableBase ParameterVariable = Entry.Key;
 		UNiagaraScript* ConversionScript = Entry.Value;
 		const FNiagaraParameterHandle& AvailableHandle(ParameterVariable.GetName());
 		TArray<FName> HandleParts = AvailableHandle.GetHandleParts();
@@ -1332,7 +1338,7 @@ TArray<TSharedPtr<FNiagaraMenuAction_Generic>> SNiagaraStackFunctionInputValue::
 			const FText Tooltip = FText::Format(ConvertInputFormat, FText::FromName(AvailableHandle.GetParameterHandleString()));
 
 			TSharedPtr<FNiagaraMenuAction_Generic> LinkAction(new FNiagaraMenuAction_Generic(
-				FNiagaraMenuAction_Generic::FOnExecuteAction::CreateSP(this, &SNiagaraStackFunctionInputValue::ConversionHandleSelected, ParameterVariable, ConversionScript),
+				FNiagaraMenuAction_Generic::FOnExecuteAction::CreateSP(this, &SNiagaraStackFunctionInputValue::ParameterWithConversionSelected, ParameterVariable, ConversionScript),
 				DisplayName, ENiagaraMenuSections::General, {RootCategoryName, Category.ToString()}, Tooltip, FText()));
 
 			LinkAction->SetParameterVariable(ParameterVariable);
@@ -1362,18 +1368,21 @@ TArray<TSharedPtr<FNiagaraMenuAction_Generic>> SNiagaraStackFunctionInputValue::
 		for (const FName& AvailableNamespace : AvailableNamespaces)
 		{
 			FNiagaraParameterHandle HandleToRead(AvailableNamespace, InputName);
-			bool bIsContained = AvailableHandles.Contains(HandleToRead);
+			FNiagaraVariableBase ParameterToRead = FNiagaraVariableBase(FunctionInput->GetInputType(), HandleToRead.GetParameterHandleString());
+			bool bIsContained = AvailableParameters.Contains(ParameterToRead);
 
 			if(bIsContained)
 			{
 				TSet<FName> ExistingNames;
-				for(const FNiagaraParameterHandle& Handle : AvailableHandles)
+				for(const FNiagaraVariableBase& AvailableParameter : AvailableParameters)
 				{
-					ExistingNames.Add(Handle.GetName());
+					FNiagaraParameterHandle AvailableHandle = FNiagaraParameterHandle(AvailableParameter.GetName());
+					ExistingNames.Add(AvailableHandle.GetName());
 				}
 
 				// let's get a unique name as the previous parameter already existed
 				HandleToRead = FNiagaraParameterHandle(AvailableNamespace, FNiagaraUtilities::GetUniqueName(InputName, ExistingNames));
+				ParameterToRead = FNiagaraVariableBase(FunctionInput->GetInputType(), HandleToRead.GetParameterHandleString());
 			}
 			
 			FFormatNamedArguments Args;
@@ -1383,7 +1392,7 @@ TArray<TSharedPtr<FNiagaraMenuAction_Generic>> SNiagaraStackFunctionInputValue::
 			const FText Tooltip = FText::Format(LOCTEXT("ReadToolTipFormat", "Read this input from a new parameter in the {AvailableNamespace} namespace."), Args);
 
 			TSharedPtr<FNiagaraMenuAction_Generic> MakeAction(new FNiagaraMenuAction_Generic(
-				FNiagaraMenuAction_Generic::FOnExecuteAction::CreateSP(this, &SNiagaraStackFunctionInputValue::ParameterHandleSelected, HandleToRead),         
+				FNiagaraMenuAction_Generic::FOnExecuteAction::CreateSP(this, &SNiagaraStackFunctionInputValue::ParameterSelected, ParameterToRead),         
 		        DisplayName, ENiagaraMenuSections::General, {CategoryName.ToString()}, Tooltip, FText()));
 
 			MakeAction->SourceData = NiagaraSourceData;
diff --git a/Engine/Plugins/FX/Niagara/Source/NiagaraEditorWidgets/Private/Stack/SNiagaraStackFunctionInputValue.h b/Engine/Plugins/FX/Niagara/Source/NiagaraEditorWidgets/Private/Stack/SNiagaraStackFunctionInputValue.h
--- a/Engine/Plugins/FX/Niagara/Source/NiagaraEditorWidgets/Private/Stack/SNiagaraStackFunctionInputValue.h
+++ b/Engine/Plugins/FX/Niagara/Source/NiagaraEditorWidgets/Private/Stack/SNiagaraStackFunctionInputValue.h
@@ -137,8 +137,8 @@ private:
 
 	void CreateScratchSelected();
 
-	void ParameterHandleSelected(FNiagaraParameterHandle Handle);
-	void ConversionHandleSelected(FNiagaraVariable Handle, UNiagaraScript* ConversionScript);
+	void ParameterSelected(FNiagaraVariableBase Parameter);
+	void ParameterWithConversionSelected(FNiagaraVariableBase Parameter, UNiagaraScript* ConversionScript);
 
 	EVisibility GetResetButtonVisibility() const;
 
