- Expose variants for other array types to use
- Add support for variants to mesh array

#rnx

diff --git a/Engine/Plugins/FX/Niagara/Source/Niagara/Classes/NiagaraDataInterfaceArray.h b/Engine/Plugins/FX/Niagara/Source/Niagara/Classes/NiagaraDataInterfaceArray.h
--- a/Engine/Plugins/FX/Niagara/Source/Niagara/Classes/NiagaraDataInterfaceArray.h
+++ b/Engine/Plugins/FX/Niagara/Source/Niagara/Classes/NiagaraDataInterfaceArray.h
@@ -131,6 +131,10 @@ public:
 	//~ INiagaraSimCacheCustomStorageInterface interface END
 	NIAGARA_API FString SimCacheVisualizerRead(const UNDIArraySimCacheData* CacheData, const FNDIArraySimCacheDataFrame& FrameData, int Element) const;
 
+#if WITH_EDITOR
+	NIAGARA_API void CreateAndSetVariant(UNiagaraComponent* NiagaraComponent, FName OverrideName, TFunction<void(UNiagaraDataInterfaceArray*)> SetDataDelegate);
+#endif
+
 	/** ReadWrite lock to ensure safe access to the underlying array. */
 	FTransactionallySafeRWLock ArrayRWGuard;
 
diff --git a/Engine/Plugins/FX/Niagara/Source/Niagara/Private/DataInterface/NiagaraDataInterfaceArrayMesh.cpp b/Engine/Plugins/FX/Niagara/Source/Niagara/Private/DataInterface/NiagaraDataInterfaceArrayMesh.cpp
--- a/Engine/Plugins/FX/Niagara/Source/Niagara/Private/DataInterface/NiagaraDataInterfaceArrayMesh.cpp
+++ b/Engine/Plugins/FX/Niagara/Source/Niagara/Private/DataInterface/NiagaraDataInterfaceArrayMesh.cpp
@@ -1,6 +1,7 @@
 // Copyright Epic Games, Inc. All Rights Reserved.
 
 #include "DataInterface/NiagaraDataInterfaceArrayMesh.h"
+#include "NiagaraDataInterfaceArrayFunctionLibrary.h"
 #include "NiagaraEmitterInstanceImpl.h"
 #include "NiagaraEmitterInstance.h"
 #include "NiagaraFunctionLibrary.h"
@@ -189,6 +190,16 @@ void UNiagaraDataInterfaceArrayMesh::SetNiagaraArrayMesh(UNiagaraComponent* Niag
 	#if WITH_EDITORONLY_DATA
 		ArrayDI->AddMeshChangedDelegates();
 	#endif
+	#if WITH_EDITOR
+		ArrayDI->CreateAndSetVariant(
+			NiagaraComponent,
+			OverrideName,
+			[ArrayData](UNiagaraDataInterfaceArray* VariantDI)
+			{
+				CastChecked<UNiagaraDataInterfaceArrayMesh>(VariantDI)->SetVariantArrayData(MakeConstArrayView(ArrayData));
+			}
+		);
+	#endif
 	}
 }
 
@@ -204,5 +215,15 @@ void UNiagaraDataInterfaceArrayMesh::SetNiagaraArrayMeshSM(UNiagaraComponent* Ni
 	#if WITH_EDITORONLY_DATA
 		ArrayDI->AddMeshChangedDelegates();
 	#endif
+	#if WITH_EDITOR
+		ArrayDI->CreateAndSetVariant(
+			NiagaraComponent,
+			OverrideName,
+			[ArrayData](UNiagaraDataInterfaceArray* VariantDI)
+			{
+				CastChecked<UNiagaraDataInterfaceArrayMesh>(VariantDI)->SetVariantArrayData(MakeConstArrayView(ArrayData));
+			}
+		);
+	#endif
 	}
 }
diff --git a/Engine/Plugins/FX/Niagara/Source/Niagara/Private/NiagaraDataInterfaceArray.cpp b/Engine/Plugins/FX/Niagara/Source/Niagara/Private/NiagaraDataInterfaceArray.cpp
--- a/Engine/Plugins/FX/Niagara/Source/Niagara/Private/NiagaraDataInterfaceArray.cpp
+++ b/Engine/Plugins/FX/Niagara/Source/Niagara/Private/NiagaraDataInterfaceArray.cpp
@@ -2,9 +2,12 @@
 
 #include "NiagaraDataInterfaceArray.h"
 #include "NiagaraCompileHashVisitor.h"
+#include "NiagaraComponent.h"
 #include "NiagaraShaderParametersBuilder.h"
 #include "NiagaraSystemInstance.h"
 
+#include "Engine/World.h"
+
 #include UE_INLINE_GENERATED_CPP_BY_NAME(NiagaraDataInterfaceArray)
 
 #define LOCTEXT_NAMESPACE "UNiagaraDataInterfaceArray"
@@ -167,6 +170,32 @@ FString UNiagaraDataInterfaceArray::SimCacheVisualizerRead(const UNDIArraySimCac
 	return ArrayProxy->SimCacheVisualizerRead(CacheData, FrameData, Element);
 }
 
+#if WITH_EDITOR
+
+void UNiagaraDataInterfaceArray::CreateAndSetVariant(UNiagaraComponent* NiagaraComponent, FName OverrideName, TFunction<void(UNiagaraDataInterfaceArray*)> SetDataDelegate)
+{
+	// We only need to do this for editor instances of the component as we are storing instance data on them
+	// For runtime instances they already have a unique copy of the data interface that we are modifying
+	UWorld* World = NiagaraComponent ? NiagaraComponent->GetWorld() : nullptr;
+	if (World && !World->IsGameWorld())
+	{
+		NiagaraComponent->Modify();
+
+		UNiagaraDataInterfaceArray* VariantDI = CastChecked<UNiagaraDataInterfaceArray>(DuplicateObject(this, NiagaraComponent));
+		CopyTo(VariantDI);
+		SetDataDelegate(VariantDI);
+		NiagaraComponent->SetParameterOverride(FNiagaraVariableBase(FNiagaraTypeDefinition(GetClass()), OverrideName), FNiagaraVariant(VariantDI));
+	
+		// We reinitialize the system as we need to flush out any queue'ed up GPU ticks as the DI could be GCed.
+		// This would normally occur naturally on the next system tick, but we may not tick the system before rendering
+		if (NiagaraComponent->IsActive())
+		{
+			NiagaraComponent->ReinitializeSystem();
+		}
+	}
+}
+#endif
+
 int32 UNDIArraySimCacheData::FindOrAddData(TConstArrayView<uint8> ArrayData)
 {
 	UE::TScopeLock ScopeLock(FindOrAddGuard);
diff --git a/Engine/Plugins/FX/Niagara/Source/Niagara/Private/NiagaraDataInterfaceArrayFunctionLibrary.cpp b/Engine/Plugins/FX/Niagara/Source/Niagara/Private/NiagaraDataInterfaceArrayFunctionLibrary.cpp
--- a/Engine/Plugins/FX/Niagara/Source/Niagara/Private/NiagaraDataInterfaceArrayFunctionLibrary.cpp
+++ b/Engine/Plugins/FX/Niagara/Source/Niagara/Private/NiagaraDataInterfaceArrayFunctionLibrary.cpp
@@ -37,25 +37,16 @@ void SetNiagaraArray(UNiagaraComponent* NiagaraComponent, FName OverrideName, TC
 	{
 		auto* ArrayProxy = static_cast<typename TDataInterace::FProxyType*>(ArrayDI->GetProxy());
 		ArrayProxy->SetArrayData(InArray);
-#if WITH_EDITOR
-		// We only need to do this for editor instances of the component as we are storing instance data on them
-		// For runtime instances they already have a unique copy of the data interface that we are modifying
-		UWorld* World = NiagaraComponent->GetWorld();
-		if (World && !World->IsGameWorld())
-		{
-			TDataInterace* VariantDI = CastChecked<TDataInterace>(DuplicateObject(ArrayDI, NiagaraComponent));
-			ArrayDI->CopyTo(VariantDI);
-			VariantDI->SetVariantArrayData(InArray);
-			NiagaraComponent->SetParameterOverride(FNiagaraVariableBase(FNiagaraTypeDefinition(TDataInterace::StaticClass()), OverrideName), FNiagaraVariant(VariantDI));
-
-			// We reinitialize the system as we need to flush out any queue'ed up GPU ticks as the DI could be GCed.
-			// This would normally occur naturally on the next system tick, but we may not tick the system before rendering
-			if (NiagaraComponent->IsActive())
+	#if WITH_EDITOR
+		ArrayDI->CreateAndSetVariant(
+			NiagaraComponent,
+			OverrideName,
+			[InArray](UNiagaraDataInterfaceArray* VariantDI)
 			{
-				NiagaraComponent->ReinitializeSystem();
+				CastChecked<TDataInterace>(VariantDI)->SetVariantArrayData(InArray);
 			}
-		}
-#endif
+		);
+	#endif
 	}
 }
 
