- Transfer bindings rather than clearing them
- This fixes an issue with a dangling UObject pointer when adding variants of data interfaces with objects inside them

#rnx

diff --git a/Engine/Plugins/FX/Niagara/Source/Niagara/Private/NiagaraComponent.cpp b/Engine/Plugins/FX/Niagara/Source/Niagara/Private/NiagaraComponent.cpp
--- a/Engine/Plugins/FX/Niagara/Source/Niagara/Private/NiagaraComponent.cpp
+++ b/Engine/Plugins/FX/Niagara/Source/Niagara/Private/NiagaraComponent.cpp
@@ -2045,7 +2045,6 @@ void UNiagaraComponent::PostApplyToComponent()
 	// Without this calls to set override data interface values can fail as the DI may not be present in the override parameter store
 	if (Asset != nullptr)
 	{
-		OverrideParameters.UnbindAll();
 		CopyParametersFromAsset(false);
 	}
 	ApplyOverridesToParameterStore();
@@ -3555,6 +3554,9 @@ void UNiagaraComponent::CopyParametersFromAsset(bool bResetExistingOverrideParam
 	OverrideParameters = AssetExposedParameters;
 	OverrideParameters.SetOwner(this);
 
+	// Make sure to transfer bindings over
+	ExistingOverrideParameters.TransferBindings(OverrideParameters);
+
 	// Copy Existing Variables into new parameter data
 	if (bResetExistingOverrideParameters == false)
 	{
