- When invalidating data interfaces remove public and don't allow redirectors to be created

#tests Repro is not quite known yet but some scripts have redirectors for DIs saved in them
#rnx

diff --git a/Engine/Plugins/FX/Niagara/Source/Niagara/Classes/NiagaraScript.h b/Engine/Plugins/FX/Niagara/Source/Niagara/Classes/NiagaraScript.h
--- a/Engine/Plugins/FX/Niagara/Source/Niagara/Classes/NiagaraScript.h
+++ b/Engine/Plugins/FX/Niagara/Source/Niagara/Classes/NiagaraScript.h
@@ -1515,7 +1515,7 @@ private:
 	// is not guaranteed.  We resolve this by renaming our invalidated DI, thus creating room for the newly created DI to be named properly.  The last caveat is that
 	// we have to be careful about when we rename the DI as the package containing the DI may get reloaded, which would recreate the DI (with the original name) leaving
 	// us, once again, in a state where the named slot is taken.
-	UPROPERTY(Transient)
+	UPROPERTY(Transient, DuplicateTransient)
 	TArray<TObjectPtr<UNiagaraDataInterface>> DataInterfacesPendingInvalidation;
 
 	enum EDataInterfaceInvalidationMode
diff --git a/Engine/Plugins/FX/Niagara/Source/Niagara/Private/NiagaraScript.cpp b/Engine/Plugins/FX/Niagara/Source/Niagara/Private/NiagaraScript.cpp
--- a/Engine/Plugins/FX/Niagara/Source/Niagara/Private/NiagaraScript.cpp
+++ b/Engine/Plugins/FX/Niagara/Source/Niagara/Private/NiagaraScript.cpp
@@ -1605,7 +1605,7 @@ void UNiagaraScript::InvalidateAndEmptyDefaultDataInterfaces(int32 SlackCount, E
 
 	if (InvalidationMode == EDataInterfaceInvalidationMode::ProcessInvaliation)
 	{
-		const ERenameFlags RenameFlags = REN_NonTransactional | REN_DoNotDirty;
+		const ERenameFlags RenameFlags = REN_NonTransactional | REN_DoNotDirty | REN_DontCreateRedirectors;
 
 		for (UNiagaraDataInterface* DataInterface : DataInterfacesPendingInvalidation)
 		{
@@ -1617,6 +1617,7 @@ void UNiagaraScript::InvalidateAndEmptyDefaultDataInterfaces(int32 SlackCount, E
 
 			FName NewName(InvalidatedName.ToView());
 
+			DataInterface->ClearFlags(RF_Public);
 			while (!DataInterface->Rename(*NewName.ToString(), nullptr, (RenameFlags | REN_Test)))
 			{
 				NewName.SetNumber(NewName.GetNumber() + 1);
