Added 'Show Only Modified Properties' filter to the Niagara Details panel and tweaked iconography to be consistent with the Details Panel.

Also tweaked filtering a bit to remove unnecessary stack object categories when they are empty

diff --git a/Engine/Plugins/FX/Niagara/Source/NiagaraEditor/Private/NiagaraStackEditorData.cpp b/Engine/Plugins/FX/Niagara/Source/NiagaraEditor/Private/NiagaraStackEditorData.cpp
--- a/Engine/Plugins/FX/Niagara/Source/NiagaraEditor/Private/NiagaraStackEditorData.cpp
+++ b/Engine/Plugins/FX/Niagara/Source/NiagaraEditor/Private/NiagaraStackEditorData.cpp
@@ -211,6 +211,16 @@ void UNiagaraStackEditorData::SetStackEntryDisplayName(const FString& StackEntry
 	}
 }
 
+bool UNiagaraStackEditorData::GetShowOnlyModified() const
+{
+	return bShowOnlyModified;
+}
+
+void UNiagaraStackEditorData::SetShowOnlyModified(bool bInShowOnlyModified)
+{
+	bShowOnlyModified = bInShowOnlyModified;
+}
+
 bool UNiagaraStackEditorData::GetShowAllAdvanced() const
 {
 	return bShowAllAdvanced;
diff --git a/Engine/Plugins/FX/Niagara/Source/NiagaraEditor/Private/ViewModels/Stack/NiagaraStackFunctionInput.cpp b/Engine/Plugins/FX/Niagara/Source/NiagaraEditor/Private/ViewModels/Stack/NiagaraStackFunctionInput.cpp
--- a/Engine/Plugins/FX/Niagara/Source/NiagaraEditor/Private/ViewModels/Stack/NiagaraStackFunctionInput.cpp
+++ b/Engine/Plugins/FX/Niagara/Source/NiagaraEditor/Private/ViewModels/Stack/NiagaraStackFunctionInput.cpp
@@ -191,6 +191,7 @@ void UNiagaraStackFunctionInput::Initialize(
 
 	AddChildFilter(FOnFilterChild::CreateUObject(this, &UNiagaraStackFunctionInput::FilterInlineChildren));
 	AddChildFilter(FOnFilterChild::CreateUObject(this, &UNiagaraStackFunctionInput::FilterForVisibleCondition));
+	AddChildFilter(FOnFilterChild::CreateUObject(this, &UNiagaraStackFunctionInput::FilterOnlyModified));
 	AddChildFilter(FOnFilterChild::CreateUObject(this, &UNiagaraStackFunctionInput::FilterForIsInlineEditConditionToggle));
 }
 
@@ -581,6 +582,22 @@ bool UNiagaraStackFunctionInput::FilterForVisibleCondition(const UNiagaraStackEn
 	return StackFunctionInputChild == nullptr || StackFunctionInputChild->GetShouldPassFilterForVisibleCondition();
 }
 
+bool UNiagaraStackFunctionInput::FilterOnlyModified(const UNiagaraStackEntry& Child) const
+{
+	if(GetStackEditorData().GetShowOnlyModified() == false)
+	{
+		return true;
+	}
+	
+	const UNiagaraStackFunctionInput* FunctionInput = Cast<UNiagaraStackFunctionInput>(&Child);
+	if (FunctionInput == nullptr || FunctionInput->CanReset())
+	{
+		return true;
+	}
+
+	return false;
+}
+
 bool UNiagaraStackFunctionInput::FilterForIsInlineEditConditionToggle(const UNiagaraStackEntry& Child) const
 {
 	const UNiagaraStackFunctionInput* StackFunctionInputChild = Cast<UNiagaraStackFunctionInput>(&Child);
diff --git a/Engine/Plugins/FX/Niagara/Source/NiagaraEditor/Private/ViewModels/Stack/NiagaraStackInputCategory.cpp b/Engine/Plugins/FX/Niagara/Source/NiagaraEditor/Private/ViewModels/Stack/NiagaraStackInputCategory.cpp
--- a/Engine/Plugins/FX/Niagara/Source/NiagaraEditor/Private/ViewModels/Stack/NiagaraStackInputCategory.cpp
+++ b/Engine/Plugins/FX/Niagara/Source/NiagaraEditor/Private/ViewModels/Stack/NiagaraStackInputCategory.cpp
@@ -31,6 +31,7 @@ void UNiagaraStackCategory::Initialize(FRequiredEntryData InRequiredEntryData, F
 {
 	Super::Initialize(InRequiredEntryData, InOwningStackItemEditorDataKey, InStackEditorDataKey);
 	AddChildFilter(FOnFilterChild::CreateUObject(this, &UNiagaraStackCategory::FilterForVisibleCondition));
+	AddChildFilter(FOnFilterChild::CreateUObject(this, &UNiagaraStackCategory::FilterOnlyModified));
 	AddChildFilter(FOnFilterChild::CreateUObject(this, &UNiagaraStackCategory::FilterForIsInlineEditConditionToggle));
 	bShouldShowInStack = true;
 	CategorySpacer = nullptr;
@@ -92,6 +93,22 @@ bool UNiagaraStackCategory::FilterForVisibleCondition(const UNiagaraStackEntry&
 	return StackFunctionInputChild == nullptr || StackFunctionInputChild->GetShouldPassFilterForVisibleCondition();
 }
 
+bool UNiagaraStackCategory::FilterOnlyModified(const UNiagaraStackEntry& Child) const
+{
+	if(GetStackEditorData().GetShowOnlyModified() == false)
+	{
+		return true;
+	}
+	
+	const UNiagaraStackFunctionInput* FunctionInput = Cast<UNiagaraStackFunctionInput>(&Child);
+	if (FunctionInput == nullptr || FunctionInput->CanReset())
+	{
+		return true;
+	}
+
+	return false;
+}
+
 bool UNiagaraStackCategory::FilterForIsInlineEditConditionToggle(const UNiagaraStackEntry& Child) const
 {
 	const UNiagaraStackFunctionInput* StackFunctionInputChild = Cast<UNiagaraStackFunctionInput>(&Child);
diff --git a/Engine/Plugins/FX/Niagara/Source/NiagaraEditor/Private/ViewModels/Stack/NiagaraStackObject.cpp b/Engine/Plugins/FX/Niagara/Source/NiagaraEditor/Private/ViewModels/Stack/NiagaraStackObject.cpp
--- a/Engine/Plugins/FX/Niagara/Source/NiagaraEditor/Private/ViewModels/Stack/NiagaraStackObject.cpp
+++ b/Engine/Plugins/FX/Niagara/Source/NiagaraEditor/Private/ViewModels/Stack/NiagaraStackObject.cpp
@@ -13,9 +13,11 @@
 #include "IPropertyRowGenerator.h"
 #include "PropertyEditorModule.h"
 #include "IDetailTreeNode.h"
+#include "PropertyHandle.h"
 
 #include "NiagaraPlatformSet.h"
 #include "DataHierarchyViewModelBase.h"
+#include "NiagaraStackEditorData.h"
 #include "ViewModels/Stack/NiagaraStackObjectIssueGenerator.h"
 
 #include UE_INLINE_GENERATED_CPP_BY_NAME(NiagaraStackObject)
@@ -43,6 +45,8 @@ void UNiagaraStackObject::Initialize(FRequiredEntryData InRequiredEntryData, UOb
 		, FObjectKey(InObject)
 		, MessageManagerRegistrationKey
 	).BindUObject(this, &UNiagaraStackObject::OnMessageManagerRefresh);
+
+	AddChildFilter(FOnFilterChild::CreateUObject(this, &UNiagaraStackObject::FilterOnlyModified));
 }
 
 void UNiagaraStackObject::Initialize(
@@ -63,6 +67,8 @@ void UNiagaraStackObject::Initialize(
 	bIsTopLevel = bInIsTopLevelStruct;
 	bHideTopLevelCategories = bInHideTopLevelCategories;
 	OwningNiagaraNode = InOwningNiagaraNode;
+
+	AddChildFilter(FOnFilterChild::CreateUObject(this, &UNiagaraStackObject::FilterOnlyModified));
 }
 
 void UNiagaraStackObject::SetOnFilterDetailNodes(FNiagaraStackObjectShared::FOnFilterDetailNodes InOnFilterDetailNodes, EDetailNodeFilterMode InFilterMode)
@@ -418,6 +424,33 @@ void UNiagaraStackObject::PostRefreshChildrenInternal()
 	Super::PostRefreshChildrenInternal();
 }
 
+bool UNiagaraStackObject::FilterOnlyModified(const UNiagaraStackEntry& Child) const
+{
+	if(GetStackEditorData().GetShowOnlyModified() == false)
+	{
+		return true;
+	}
+	
+	const UNiagaraStackPropertyRow* PropertyRow = Cast<UNiagaraStackPropertyRow>(&Child);
+	if (PropertyRow == nullptr)
+	{
+		return true;
+	}
+
+	TSharedPtr<IPropertyHandle> PropertyHandle = PropertyRow->GetDetailTreeNode()->CreatePropertyHandle();
+	if(PropertyHandle == nullptr)
+	{
+		return true;
+	}
+	
+	if(PropertyHandle->CanResetToDefault())
+	{
+		return true;
+	}
+
+	return false;
+}
+
 void UNiagaraStackObject::PropertyRowsRefreshed()
 {
 	RefreshChildren();
diff --git a/Engine/Plugins/FX/Niagara/Source/NiagaraEditor/Private/ViewModels/Stack/NiagaraStackPropertyRow.cpp b/Engine/Plugins/FX/Niagara/Source/NiagaraEditor/Private/ViewModels/Stack/NiagaraStackPropertyRow.cpp
--- a/Engine/Plugins/FX/Niagara/Source/NiagaraEditor/Private/ViewModels/Stack/NiagaraStackPropertyRow.cpp
+++ b/Engine/Plugins/FX/Niagara/Source/NiagaraEditor/Private/ViewModels/Stack/NiagaraStackPropertyRow.cpp
@@ -9,6 +9,7 @@
 #include "IDetailPropertyRow.h"
 #include "ScopedTransaction.h"
 #include "DataHierarchyViewModelBase.h"
+#include "NiagaraStackEditorData.h"
 #include "Customizations/NiagaraStackObjectPropertyCustomization.h"
 
 #include UE_INLINE_GENERATED_CPP_BY_NAME(NiagaraStackPropertyRow)
@@ -53,6 +54,8 @@ void UNiagaraStackPropertyRow::Initialize(
 		}
 	}
 	bIsHiddenCategory = false;
+	
+	AddChildFilter(FOnFilterChild::CreateUObject(this, &UNiagaraStackPropertyRow::FilterOnlyModified));
 }
 
 TSharedRef<IDetailTreeNode> UNiagaraStackPropertyRow::GetDetailTreeNode() const
@@ -79,7 +82,21 @@ UNiagaraStackEntry::EStackRowStyle UNiagaraStackPropertyRow::GetStackRowStyle()
 
 bool UNiagaraStackPropertyRow::GetShouldShowInStack() const
 {
-	return bIsHiddenCategory == false;
+	if(DetailTreeNode->GetNodeType() == EDetailNodeType::Category)
+	{
+		if(bIsHiddenCategory)
+		{
+			return false;
+		}
+		
+		TArray<UNiagaraStackEntry*> CurrentFilteredChildren;
+		GetFilteredChildren(CurrentFilteredChildren);
+		int32 EmptyCount = CategorySpacer == nullptr ? 0 : 1;
+
+		return CurrentFilteredChildren.Num() > EmptyCount;
+	}
+	
+	return true;
 }
 
 bool UNiagaraStackPropertyRow::HasOverridenContent() const
@@ -370,4 +387,31 @@ FHierarchyElementIdentity UNiagaraStackPropertyRow::DetermineSummaryIdentity() c
 	return Identity;
 }
 
+bool UNiagaraStackPropertyRow::FilterOnlyModified(const UNiagaraStackEntry& Child) const
+{
+	if(GetStackEditorData().GetShowOnlyModified() == false)
+	{
+		return true;
+	}
+	
+	const UNiagaraStackPropertyRow* PropertyRow = Cast<UNiagaraStackPropertyRow>(&Child);
+	if (PropertyRow == nullptr)
+	{
+		return true;
+	}
+
+	TSharedPtr<IPropertyHandle> PropertyHandle = PropertyRow->GetDetailTreeNode()->CreatePropertyHandle();
+	if(PropertyHandle == nullptr)
+	{
+		return true;
+	}
+	
+	if(PropertyHandle->CanResetToDefault())
+	{
+		return true;
+	}
+
+	return false;
+}
+
 #undef LOCTEXT_NAMESPACE
diff --git a/Engine/Plugins/FX/Niagara/Source/NiagaraEditor/Private/ViewModels/Stack/NiagaraStackScriptHierarchyRoot.cpp b/Engine/Plugins/FX/Niagara/Source/NiagaraEditor/Private/ViewModels/Stack/NiagaraStackScriptHierarchyRoot.cpp
--- a/Engine/Plugins/FX/Niagara/Source/NiagaraEditor/Private/ViewModels/Stack/NiagaraStackScriptHierarchyRoot.cpp
+++ b/Engine/Plugins/FX/Niagara/Source/NiagaraEditor/Private/ViewModels/Stack/NiagaraStackScriptHierarchyRoot.cpp
@@ -35,6 +35,8 @@ void UNiagaraStackScriptHierarchyRoot::Initialize(FRequiredEntryData InRequiredE
 	AddChildFilter(FOnFilterChild::CreateUObject(this, &UNiagaraStackScriptHierarchyRoot::FilterForVisibleCondition));
 	AddChildFilter(FOnFilterChild::CreateUObject(this, &UNiagaraStackScriptHierarchyRoot::FilterForIsInlineEditConditionToggle));
 	AddChildFilter(FOnFilterChild::CreateUObject(this, &UNiagaraStackScriptHierarchyRoot::FilterByActiveSection));
+	AddChildFilter(FOnFilterChild::CreateUObject(this, &UNiagaraStackScriptHierarchyRoot::FilterOnlyModified));
+	
 	FText LastActiveSection = GetStackEditorData().GetStackEntryActiveSection(GetStackEditorDataKey(), AllSectionName);
 	ActiveSection = FindSectionByName(FName(LastActiveSection.ToString()));
 }
@@ -338,6 +340,22 @@ bool UNiagaraStackScriptHierarchyRoot::FilterForVisibleCondition(const UNiagaraS
 	return StackFunctionInputChild == nullptr || StackFunctionInputChild->GetShouldPassFilterForVisibleCondition();
 }
 
+bool UNiagaraStackScriptHierarchyRoot::FilterOnlyModified(const UNiagaraStackEntry& Child) const
+{
+	if(GetStackEditorData().GetShowOnlyModified() == false)
+	{
+		return true;
+	}
+	
+	const UNiagaraStackFunctionInput* FunctionInput = Cast<UNiagaraStackFunctionInput>(&Child);
+	if (FunctionInput == nullptr || FunctionInput->CanReset())
+	{
+		return true;
+	}
+
+	return false;
+}
+
 bool UNiagaraStackScriptHierarchyRoot::FilterForIsInlineEditConditionToggle(const UNiagaraStackEntry& Child) const
 {
 	const UNiagaraStackFunctionInput* StackFunctionInputChild = Cast<UNiagaraStackFunctionInput>(&Child);
diff --git a/Engine/Plugins/FX/Niagara/Source/NiagaraEditor/Private/ViewModels/Stack/NiagaraStackSummaryViewInputCollection.cpp b/Engine/Plugins/FX/Niagara/Source/NiagaraEditor/Private/ViewModels/Stack/NiagaraStackSummaryViewInputCollection.cpp
--- a/Engine/Plugins/FX/Niagara/Source/NiagaraEditor/Private/ViewModels/Stack/NiagaraStackSummaryViewInputCollection.cpp
+++ b/Engine/Plugins/FX/Niagara/Source/NiagaraEditor/Private/ViewModels/Stack/NiagaraStackSummaryViewInputCollection.cpp
@@ -43,6 +43,8 @@ void UNiagaraStackSummaryViewCollection::Initialize(FRequiredEntryData InRequire
 	GetEmitterViewModel()->GetEditorData().OnSummaryViewStateChanged().AddUObject(this, &UNiagaraStackSummaryViewCollection::OnViewStateChanged);
 	GetEmitterViewModel()->GetSummaryHierarchyViewModel()->OnHierarchyChanged().AddUObject(this, &UNiagaraStackSummaryViewCollection::OnViewStateChanged);
 	GetEmitterViewModel()->GetSummaryHierarchyViewModel()->OnHierarchyPropertiesChanged().AddUObject(this, &UNiagaraStackSummaryViewCollection::OnViewStateChanged);
+
+	AddChildFilter(FOnFilterChild::CreateUObject(this, &UNiagaraStackSummaryViewCollection::FilterOnlyModified));
 }
 
 void UNiagaraStackSummaryViewCollection::FinalizeInternal()
@@ -620,6 +622,22 @@ bool UNiagaraStackSummaryViewCollection::FilterByActiveSection(const UNiagaraSta
 	return bIsAllowed;
 }
 
+bool UNiagaraStackSummaryViewCollection::FilterOnlyModified(const UNiagaraStackEntry& Child) const
+{
+	if(GetStackEditorData().GetShowOnlyModified() == false)
+	{
+		return true;
+	}
+	
+	const UNiagaraStackFunctionInput* FunctionInput = Cast<UNiagaraStackFunctionInput>(&Child);
+	if (FunctionInput == nullptr || FunctionInput->CanReset())
+	{
+		return true;
+	}
+
+	return false;
+}
+
 void UNiagaraStackSummaryViewCollection::OnViewStateChanged()
 {
 	if (!IsFinalized())
diff --git a/Engine/Plugins/FX/Niagara/Source/NiagaraEditor/Private/ViewModels/Stack/NiagaraStackViewModel.cpp b/Engine/Plugins/FX/Niagara/Source/NiagaraEditor/Private/ViewModels/Stack/NiagaraStackViewModel.cpp
--- a/Engine/Plugins/FX/Niagara/Source/NiagaraEditor/Private/ViewModels/Stack/NiagaraStackViewModel.cpp
+++ b/Engine/Plugins/FX/Niagara/Source/NiagaraEditor/Private/ViewModels/Stack/NiagaraStackViewModel.cpp
@@ -695,6 +695,32 @@ UNiagaraStackViewModel::FOnDataObjectChanged& UNiagaraStackViewModel::OnDataObje
 	return DataObjectChangedDelegate;
 }
 
+bool UNiagaraStackViewModel::GetShowOnlyModified() const
+{
+	for (TSharedRef<FTopLevelViewModel> TopLevelViewModel : TopLevelViewModels)
+	{
+		if (TopLevelViewModel->IsValid() && TopLevelViewModel->GetStackEditorData()->GetShowOnlyModified())
+		{
+			return true;
+		}
+	}
+	return false;
+}
+
+void UNiagaraStackViewModel::SetShowOnlyModified(bool bInShowOnlyModified)
+{
+	for (TSharedRef<FTopLevelViewModel> TopLevelViewModel : TopLevelViewModels)
+	{
+		if (TopLevelViewModel->IsValid()) 
+		{
+			TopLevelViewModel->GetStackEditorData()->SetShowOnlyModified(bInShowOnlyModified);
+		}
+	}
+
+	RestartSearch();
+	RootEntry->RefreshFilteredChildren();
+}
+
 bool UNiagaraStackViewModel::GetShowAllAdvanced() const
 {
 	for (TSharedRef<FTopLevelViewModel> TopLevelViewModel : TopLevelViewModels)
diff --git a/Engine/Plugins/FX/Niagara/Source/NiagaraEditor/Public/NiagaraStackEditorData.h b/Engine/Plugins/FX/Niagara/Source/NiagaraEditor/Public/NiagaraStackEditorData.h
--- a/Engine/Plugins/FX/Niagara/Source/NiagaraEditor/Public/NiagaraStackEditorData.h
+++ b/Engine/Plugins/FX/Niagara/Source/NiagaraEditor/Public/NiagaraStackEditorData.h
@@ -158,6 +158,12 @@ public:
 	*/
 	void SetStackEntryDisplayName(const FString& StackEntryKey, const FText& InDisplayName);
 
+	/** Gets whether or not only modified properties should be shown in the stack. */
+	bool GetShowOnlyModified() const;
+
+	/** Sets whether or not only modified properties should be shown in the stack. */
+	void SetShowOnlyModified(bool bInShowOnlyModified);
+	
 	/* Gets whether or not all advanced items should be shown in the stack. */
 	bool GetShowAllAdvanced() const;
 
@@ -247,14 +253,16 @@ private:
 	TMap<FString, FText> StackEntryKeyToDisplayName;
 	
 #endif
+
+	bool bShowOnlyModified = false;
 	
-	bool bShowAllAdvanced;
+	bool bShowAllAdvanced = false;
 
-	bool bShowOutputs;
+	bool bShowOutputs = false;
 
-	bool bShowLinkedInputs;
+	bool bShowLinkedInputs = false;
 
-	bool bShowOnlyIssues;
+	bool bShowOnlyIssues = false;
 
 	double LastScrollPosition;
 
diff --git a/Engine/Plugins/FX/Niagara/Source/NiagaraEditor/Public/ViewModels/Stack/NiagaraStackFunctionInput.h b/Engine/Plugins/FX/Niagara/Source/NiagaraEditor/Public/ViewModels/Stack/NiagaraStackFunctionInput.h
--- a/Engine/Plugins/FX/Niagara/Source/NiagaraEditor/Public/ViewModels/Stack/NiagaraStackFunctionInput.h
+++ b/Engine/Plugins/FX/Niagara/Source/NiagaraEditor/Public/ViewModels/Stack/NiagaraStackFunctionInput.h
@@ -415,6 +415,7 @@ private:
 
 	NIAGARAEDITOR_API bool FilterInlineChildren(const UNiagaraStackEntry& Child) const;
 	NIAGARAEDITOR_API bool FilterForVisibleCondition(const UNiagaraStackEntry& NiagaraStackEntry) const;
+	bool FilterOnlyModified(const UNiagaraStackEntry& Child) const;
 	bool FilterForIsInlineEditConditionToggle(const UNiagaraStackEntry& NiagaraStackEntry) const;
 
 	void ReportScriptVersionChange() const;
diff --git a/Engine/Plugins/FX/Niagara/Source/NiagaraEditor/Public/ViewModels/Stack/NiagaraStackInputCategory.h b/Engine/Plugins/FX/Niagara/Source/NiagaraEditor/Public/ViewModels/Stack/NiagaraStackInputCategory.h
--- a/Engine/Plugins/FX/Niagara/Source/NiagaraEditor/Public/ViewModels/Stack/NiagaraStackInputCategory.h
+++ b/Engine/Plugins/FX/Niagara/Source/NiagaraEditor/Public/ViewModels/Stack/NiagaraStackInputCategory.h
@@ -42,6 +42,7 @@ protected:
 private:
 private:
 	NIAGARAEDITOR_API bool FilterForVisibleCondition(const UNiagaraStackEntry& Child) const;
+	bool FilterOnlyModified(const UNiagaraStackEntry& Child) const;
 	NIAGARAEDITOR_API bool FilterForIsInlineEditConditionToggle(const UNiagaraStackEntry& Child) const;
 protected:
 	bool bShouldShowInStack;
diff --git a/Engine/Plugins/FX/Niagara/Source/NiagaraEditor/Public/ViewModels/Stack/NiagaraStackObject.h b/Engine/Plugins/FX/Niagara/Source/NiagaraEditor/Public/ViewModels/Stack/NiagaraStackObject.h
--- a/Engine/Plugins/FX/Niagara/Source/NiagaraEditor/Public/ViewModels/Stack/NiagaraStackObject.h
+++ b/Engine/Plugins/FX/Niagara/Source/NiagaraEditor/Public/ViewModels/Stack/NiagaraStackObject.h
@@ -94,6 +94,7 @@ protected:
 	NIAGARAEDITOR_API virtual void RefreshChildrenInternal(const TArray<UNiagaraStackEntry*>& CurrentChildren, TArray<UNiagaraStackEntry*>& NewChildren, TArray<FStackIssue>& NewIssues) override;
 	NIAGARAEDITOR_API virtual void PostRefreshChildrenInternal() override;
 
+	NIAGARAEDITOR_API bool FilterOnlyModified(const UNiagaraStackEntry& Child) const;
 private:
 	NIAGARAEDITOR_API void PropertyRowsRefreshed();
 	NIAGARAEDITOR_API void OnMessageManagerRefresh(const TArray<TSharedRef<const INiagaraMessage>>& NewMessages);
diff --git a/Engine/Plugins/FX/Niagara/Source/NiagaraEditor/Public/ViewModels/Stack/NiagaraStackPropertyRow.h b/Engine/Plugins/FX/Niagara/Source/NiagaraEditor/Public/ViewModels/Stack/NiagaraStackPropertyRow.h
--- a/Engine/Plugins/FX/Niagara/Source/NiagaraEditor/Public/ViewModels/Stack/NiagaraStackPropertyRow.h
+++ b/Engine/Plugins/FX/Niagara/Source/NiagaraEditor/Public/ViewModels/Stack/NiagaraStackPropertyRow.h
@@ -78,6 +78,8 @@ protected:
 
 	NIAGARAEDITOR_API virtual bool SupportsSummaryView() const override;
 	NIAGARAEDITOR_API virtual FHierarchyElementIdentity DetermineSummaryIdentity() const override;
+
+	NIAGARAEDITOR_API bool FilterOnlyModified(const UNiagaraStackEntry& Child) const;
 private:
 	TSharedPtr<IDetailTreeNode> DetailTreeNode;
 	UNiagaraNode* OwningNiagaraNode;
diff --git a/Engine/Plugins/FX/Niagara/Source/NiagaraEditor/Public/ViewModels/Stack/NiagaraStackScriptHierarchyRoot.h b/Engine/Plugins/FX/Niagara/Source/NiagaraEditor/Public/ViewModels/Stack/NiagaraStackScriptHierarchyRoot.h
--- a/Engine/Plugins/FX/Niagara/Source/NiagaraEditor/Public/ViewModels/Stack/NiagaraStackScriptHierarchyRoot.h
+++ b/Engine/Plugins/FX/Niagara/Source/NiagaraEditor/Public/ViewModels/Stack/NiagaraStackScriptHierarchyRoot.h
@@ -83,6 +83,7 @@ protected:
 	const UHierarchySection* FindSectionByName(FName SectionName);
 	NIAGARAEDITOR_API virtual bool FilterByActiveSection(const UNiagaraStackEntry& Child) const;
 	NIAGARAEDITOR_API bool FilterForVisibleCondition(const UNiagaraStackEntry& NiagaraStackEntry) const;
+	bool FilterOnlyModified(const UNiagaraStackEntry& Child) const;
 	bool FilterForIsInlineEditConditionToggle(const UNiagaraStackEntry& NiagaraStackEntry) const;
 	
 	const FNiagaraScriptInstanceData& GetScriptInstanceData() const { return ScriptInstanceData; }
diff --git a/Engine/Plugins/FX/Niagara/Source/NiagaraEditor/Public/ViewModels/Stack/NiagaraStackSummaryViewInputCollection.h b/Engine/Plugins/FX/Niagara/Source/NiagaraEditor/Public/ViewModels/Stack/NiagaraStackSummaryViewInputCollection.h
--- a/Engine/Plugins/FX/Niagara/Source/NiagaraEditor/Public/ViewModels/Stack/NiagaraStackSummaryViewInputCollection.h
+++ b/Engine/Plugins/FX/Niagara/Source/NiagaraEditor/Public/ViewModels/Stack/NiagaraStackSummaryViewInputCollection.h
@@ -33,6 +33,7 @@ protected:
 	NIAGARAEDITOR_API const TArray<UHierarchySection*>& GetHierarchySections() const;
 
 	NIAGARAEDITOR_API virtual bool FilterByActiveSection(const UNiagaraStackEntry& Child) const override;
+	bool FilterOnlyModified(const UNiagaraStackEntry& Child) const;
 private:
 	NIAGARAEDITOR_API void OnViewStateChanged();
 private:
diff --git a/Engine/Plugins/FX/Niagara/Source/NiagaraEditor/Public/ViewModels/Stack/NiagaraStackViewModel.h b/Engine/Plugins/FX/Niagara/Source/NiagaraEditor/Public/ViewModels/Stack/NiagaraStackViewModel.h
--- a/Engine/Plugins/FX/Niagara/Source/NiagaraEditor/Public/ViewModels/Stack/NiagaraStackViewModel.h
+++ b/Engine/Plugins/FX/Niagara/Source/NiagaraEditor/Public/ViewModels/Stack/NiagaraStackViewModel.h
@@ -100,6 +100,9 @@ public:
 	NIAGARAEDITOR_API FOnSearchCompleted& OnSearchCompleted();
 	NIAGARAEDITOR_API FOnDataObjectChanged& OnDataObjectChanged();
 
+	NIAGARAEDITOR_API bool GetShowOnlyModified() const;
+	NIAGARAEDITOR_API void SetShowOnlyModified(bool bInShowOnlyModified);
+	
 	NIAGARAEDITOR_API bool GetShowAllAdvanced() const;
 	NIAGARAEDITOR_API void SetShowAllAdvanced(bool bInShowAllAdvanced);
 
diff --git a/Engine/Plugins/FX/Niagara/Source/NiagaraEditorWidgets/Private/SNiagaraStack.cpp b/Engine/Plugins/FX/Niagara/Source/NiagaraEditorWidgets/Private/SNiagaraStack.cpp
--- a/Engine/Plugins/FX/Niagara/Source/NiagaraEditorWidgets/Private/SNiagaraStack.cpp
+++ b/Engine/Plugins/FX/Niagara/Source/NiagaraEditorWidgets/Private/SNiagaraStack.cpp
@@ -59,6 +59,7 @@
 #include "Styling/StyleColors.h"
 #include "SResetToDefaultPropertyEditor.h"
 #include "Stack/SNiagaraStackScriptHierarchyRoot.h"
+#include "Widgets/Images/SLayeredImage.h"
 
 #include "ViewModels/NiagaraSystemSelectionViewModel.h"
 #include "ViewModels/Stack/NiagaraStackNote.h"
@@ -364,7 +365,9 @@ void SNiagaraStack::SynchronizeTreeExpansion()
 
 TSharedRef<SWidget> SNiagaraStack::ConstructHeaderWidget()
 {
-	return SNew(SVerticalBox)
+	TSharedPtr<SLayeredImage> ViewOptionsImage;
+	
+	TSharedRef<SWidget> Result = SNew(SVerticalBox)
 		//~ Top level object list view
 		+ SVerticalBox::Slot()
 		.AutoHeight()
@@ -409,12 +412,16 @@ TSharedRef<SWidget> SNiagaraStack::ConstructHeaderWidget()
 				.MenuPlacement(MenuPlacement_BelowRightAnchor)
 				.ButtonContent()
 				[
-					SNew(SImage)
-					.Image(FAppStyle::Get().GetBrush("Icons.Filter"))
+					SAssignNew(ViewOptionsImage, SLayeredImage)
+					.Image(FAppStyle::Get().GetBrush("DetailsView.ViewOptions"))
 					.ColorAndOpacity(FSlateColor::UseForeground())
 				]
 			]
 		];
+
+	ViewOptionsImage->AddLayer(TAttribute<const FSlateBrush*>(this, &SNiagaraStack::GetViewOptionsBadgeIcon));
+
+	return Result;
 }
 
 void SNiagaraStack::UpdateSearchTextFromExternal(FText NewSearchText)
@@ -550,6 +557,16 @@ TSharedRef<SWidget> SNiagaraStack::GetViewOptionsMenu() const
 {
 	FMenuBuilder MenuBuilder(false, nullptr);
 
+	MenuBuilder.AddMenuEntry(
+		LOCTEXT("ShowOnlyModifiedLabel", "Show Only Modified Properties"),
+		LOCTEXT("ShowOnlyModifiedToolTip", "Displays only properties that have been changes from their default."),
+		FSlateIcon(),
+		FUIAction(
+			FExecuteAction::CreateLambda([this]() { StackViewModel->SetShowOnlyModified(!StackViewModel->GetShowOnlyModified()); }),
+			FCanExecuteAction(),
+			FGetActionCheckState::CreateLambda([this]() { return StackViewModel->GetShowOnlyModified() ? ECheckBoxState::Checked : ECheckBoxState::Unchecked; })),
+		NAME_None, EUserInterfaceActionType::Check);
+	
 	MenuBuilder.AddMenuEntry(
 		LOCTEXT("ShowAllAdvancedLabel", "Show All Advanced"),
 		LOCTEXT("ShowAllAdvancedToolTip", "Forces all advanced items to be showing in the stack."),
@@ -1285,4 +1302,14 @@ void SNiagaraStack::OnCycleThroughIssues()
 	StackTree->RequestScrollIntoView(StackEntry);
 }
 
+const FSlateBrush* SNiagaraStack::GetViewOptionsBadgeIcon() const
+{
+	// Badge the icon if any view option that narrows down the results is checked
+	bool bHasBadge = (StackViewModel->GetShowOnlyModified() || StackViewModel->GetShowOnlyIssues()
+		|| StackViewModel->GetShowAllAdvanced() || StackViewModel->GetShowLinkedInputs()
+		|| StackViewModel->GetShowOutputs());
+	
+	return bHasBadge ? FAppStyle::Get().GetBrush("Icons.BadgeModified") : nullptr;
+}
+
 #undef LOCTEXT_NAMESPACE
diff --git a/Engine/Plugins/FX/Niagara/Source/NiagaraEditorWidgets/Private/SNiagaraStack.h b/Engine/Plugins/FX/Niagara/Source/NiagaraEditorWidgets/Private/SNiagaraStack.h
--- a/Engine/Plugins/FX/Niagara/Source/NiagaraEditorWidgets/Private/SNiagaraStack.h
+++ b/Engine/Plugins/FX/Niagara/Source/NiagaraEditorWidgets/Private/SNiagaraStack.h
@@ -122,6 +122,7 @@ private:
 	FReply OnCycleThroughSystemIssues(TSharedPtr<FNiagaraSystemViewModel> SystemViewModel);
 	void OnCycleThroughIssues();
 
+	const FSlateBrush* GetViewOptionsBadgeIcon() const;
 private:
 	UNiagaraStackViewModel* StackViewModel;
 
