Fixed Show Only Modified skipping child inputs that were modified when their parent was unmodified

diff --git a/Engine/Plugins/FX/Niagara/Source/NiagaraEditor/Private/ViewModels/Stack/NiagaraStackFunctionInput.cpp b/Engine/Plugins/FX/Niagara/Source/NiagaraEditor/Private/ViewModels/Stack/NiagaraStackFunctionInput.cpp
--- a/Engine/Plugins/FX/Niagara/Source/NiagaraEditor/Private/ViewModels/Stack/NiagaraStackFunctionInput.cpp
+++ b/Engine/Plugins/FX/Niagara/Source/NiagaraEditor/Private/ViewModels/Stack/NiagaraStackFunctionInput.cpp
@@ -255,6 +255,22 @@ FNiagaraInputParameterCustomization UNiagaraStackFunctionInput::GetInputWidgetCu
 	return InputMetaData.IsSet() ? InputMetaData->WidgetCustomization : FNiagaraInputParameterCustomization();
 }
 
+bool UNiagaraStackFunctionInput::GetShouldShowInStack() const
+{
+	// Using the FilterOnlyModified child filter, we let inputs pass that are at default value but have modified children inputs
+	// If we removed them via filter, the modified children inputs also wouldn't be displayed.
+	// Instead, we remove these inputs from the display
+	if(IsFinalized() == false && GetStackEditorData().GetShowOnlyModified())
+	{
+		if(CanReset() == false)
+		{
+			return false;
+		}
+	}
+
+	return true;
+}
+
 FText UNiagaraStackFunctionInput::GetTooltipText() const
 {
 	FText Description = SummaryViewTooltipOverride.IsSet() && !SummaryViewTooltipOverride->Get().IsEmptyOrWhitespace() ? SummaryViewTooltipOverride->Get() : InputMetaData.IsSet() ? InputMetaData->Description : FText::GetEmpty();
@@ -590,7 +606,7 @@ bool UNiagaraStackFunctionInput::FilterOnlyModified(const UNiagaraStackEntry& Ch
 	}
 	
 	const UNiagaraStackFunctionInput* FunctionInput = Cast<UNiagaraStackFunctionInput>(&Child);
-	if (FunctionInput == nullptr || FunctionInput->CanReset())
+	if (FunctionInput == nullptr || FunctionInput->CanReset() || FunctionInput->HasAnyResettableChildrenInputs())
 	{
 		return true;
 	}
@@ -2431,6 +2447,22 @@ void UNiagaraStackFunctionInput::Reset()
 	}
 }
 
+bool UNiagaraStackFunctionInput::HasAnyResettableChildrenInputs() const
+{
+	TArray<UNiagaraStackFunctionInput*> FilteredChildInputs;
+	GetFilteredChildInputs(FilteredChildInputs);
+	
+	for(UNiagaraStackFunctionInput* Input : FilteredChildInputs)
+	{
+		if(Input->CanReset())
+		{
+			return true;
+		}
+	}
+
+	return false;
+}
+
 bool UNiagaraStackFunctionInput::IsStaticParameter() const
 {
 	return ParameterBehavior == EStackParameterBehavior::Static;
diff --git a/Engine/Plugins/FX/Niagara/Source/NiagaraEditor/Private/ViewModels/Stack/NiagaraStackInputCategory.cpp b/Engine/Plugins/FX/Niagara/Source/NiagaraEditor/Private/ViewModels/Stack/NiagaraStackInputCategory.cpp
--- a/Engine/Plugins/FX/Niagara/Source/NiagaraEditor/Private/ViewModels/Stack/NiagaraStackInputCategory.cpp
+++ b/Engine/Plugins/FX/Niagara/Source/NiagaraEditor/Private/ViewModels/Stack/NiagaraStackInputCategory.cpp
@@ -101,7 +101,7 @@ bool UNiagaraStackCategory::FilterOnlyModified(const UNiagaraStackEntry& Child)
 	}
 	
 	const UNiagaraStackFunctionInput* FunctionInput = Cast<UNiagaraStackFunctionInput>(&Child);
-	if (FunctionInput == nullptr || FunctionInput->CanReset())
+	if (FunctionInput == nullptr || FunctionInput->CanReset() || FunctionInput->HasAnyResettableChildrenInputs())
 	{
 		return true;
 	}
diff --git a/Engine/Plugins/FX/Niagara/Source/NiagaraEditor/Private/ViewModels/Stack/NiagaraStackScriptHierarchyRoot.cpp b/Engine/Plugins/FX/Niagara/Source/NiagaraEditor/Private/ViewModels/Stack/NiagaraStackScriptHierarchyRoot.cpp
--- a/Engine/Plugins/FX/Niagara/Source/NiagaraEditor/Private/ViewModels/Stack/NiagaraStackScriptHierarchyRoot.cpp
+++ b/Engine/Plugins/FX/Niagara/Source/NiagaraEditor/Private/ViewModels/Stack/NiagaraStackScriptHierarchyRoot.cpp
@@ -348,7 +348,7 @@ bool UNiagaraStackScriptHierarchyRoot::FilterOnlyModified(const UNiagaraStackEnt
 	}
 	
 	const UNiagaraStackFunctionInput* FunctionInput = Cast<UNiagaraStackFunctionInput>(&Child);
-	if (FunctionInput == nullptr || FunctionInput->CanReset())
+	if (FunctionInput == nullptr || FunctionInput->CanReset() || FunctionInput->HasAnyResettableChildrenInputs())
 	{
 		return true;
 	}
diff --git a/Engine/Plugins/FX/Niagara/Source/NiagaraEditor/Private/ViewModels/Stack/NiagaraStackSummaryViewInputCollection.cpp b/Engine/Plugins/FX/Niagara/Source/NiagaraEditor/Private/ViewModels/Stack/NiagaraStackSummaryViewInputCollection.cpp
--- a/Engine/Plugins/FX/Niagara/Source/NiagaraEditor/Private/ViewModels/Stack/NiagaraStackSummaryViewInputCollection.cpp
+++ b/Engine/Plugins/FX/Niagara/Source/NiagaraEditor/Private/ViewModels/Stack/NiagaraStackSummaryViewInputCollection.cpp
@@ -630,7 +630,7 @@ bool UNiagaraStackSummaryViewCollection::FilterOnlyModified(const UNiagaraStackE
 	}
 	
 	const UNiagaraStackFunctionInput* FunctionInput = Cast<UNiagaraStackFunctionInput>(&Child);
-	if (FunctionInput == nullptr || FunctionInput->CanReset())
+	if (FunctionInput == nullptr || FunctionInput->CanReset() || FunctionInput->HasAnyResettableChildrenInputs())
 	{
 		return true;
 	}
diff --git a/Engine/Plugins/FX/Niagara/Source/NiagaraEditor/Private/ViewModels/Stack/NiagaraStackViewModel.cpp b/Engine/Plugins/FX/Niagara/Source/NiagaraEditor/Private/ViewModels/Stack/NiagaraStackViewModel.cpp
--- a/Engine/Plugins/FX/Niagara/Source/NiagaraEditor/Private/ViewModels/Stack/NiagaraStackViewModel.cpp
+++ b/Engine/Plugins/FX/Niagara/Source/NiagaraEditor/Private/ViewModels/Stack/NiagaraStackViewModel.cpp
@@ -718,7 +718,7 @@ void UNiagaraStackViewModel::SetShowOnlyModified(bool bInShowOnlyModified)
 	}
 
 	RestartSearch();
-	RootEntry->RefreshFilteredChildren();
+	RootEntry->RefreshChildren();
 }
 
 bool UNiagaraStackViewModel::GetShowAllAdvanced() const
diff --git a/Engine/Plugins/FX/Niagara/Source/NiagaraEditor/Public/ViewModels/Stack/NiagaraStackFunctionInput.h b/Engine/Plugins/FX/Niagara/Source/NiagaraEditor/Public/ViewModels/Stack/NiagaraStackFunctionInput.h
--- a/Engine/Plugins/FX/Niagara/Source/NiagaraEditor/Public/ViewModels/Stack/NiagaraStackFunctionInput.h
+++ b/Engine/Plugins/FX/Niagara/Source/NiagaraEditor/Public/ViewModels/Stack/NiagaraStackFunctionInput.h
@@ -104,6 +104,7 @@ public:
 	NIAGARAEDITOR_API FNiagaraInputParameterCustomization GetInputWidgetCustomization() const;
 
 	//~ UNiagaraStackEntry interface
+	virtual bool GetShouldShowInStack() const override;
 	NIAGARAEDITOR_API virtual FText GetDisplayName() const override;
 	NIAGARAEDITOR_API virtual FText GetTooltipText() const override;
 	NIAGARAEDITOR_API virtual bool GetIsEnabled() const override;
@@ -212,6 +213,8 @@ public:
 	/** Determine if this field is editable */
 	NIAGARAEDITOR_API bool IsEditable() const;
 
+	NIAGARAEDITOR_API bool HasAnyResettableChildrenInputs() const;
+	
 	/** If true the parameter can only be set to local constant values */
 	NIAGARAEDITOR_API bool IsStaticParameter() const;
 
