Drop missing variables ensure to a log warning.

diff --git a/Engine/Plugins/FX/Niagara/Source/Niagara/Private/NiagaraDataChannelAccessor.cpp b/Engine/Plugins/FX/Niagara/Source/Niagara/Private/NiagaraDataChannelAccessor.cpp
--- a/Engine/Plugins/FX/Niagara/Source/Niagara/Private/NiagaraDataChannelAccessor.cpp
+++ b/Engine/Plugins/FX/Niagara/Source/Niagara/Private/NiagaraDataChannelAccessor.cpp
@@ -18,7 +18,7 @@
 namespace NDCCVars
 {
 #if DEBUG_NDC_ACCESS
-	int32 MissingNDCAccessorMode = 2;
+	int32 MissingNDCAccessorMode = 1;
 	static FAutoConsoleVariableRef CVarMissingNDCAccessorMode(
 		TEXT("fx.Niagara.DataChannels.MissingNDCAccessorMode"), 
 		MissingNDCAccessorMode, 
