- Add validation rule for checking stateless usage

diff --git a/Engine/Plugins/FX/Niagara/Source/NiagaraEditor/Private/NiagaraValidationRules.cpp b/Engine/Plugins/FX/Niagara/Source/NiagaraEditor/Private/NiagaraValidationRules.cpp
--- a/Engine/Plugins/FX/Niagara/Source/NiagaraEditor/Private/NiagaraValidationRules.cpp
+++ b/Engine/Plugins/FX/Niagara/Source/NiagaraEditor/Private/NiagaraValidationRules.cpp
@@ -18,6 +18,8 @@
 #include "NiagaraSystemImpl.h"
 #include "NiagaraSystemEditorData.h"
 #include "DataInterface/NiagaraDataInterfaceActorComponent.h"
+#include "Stateless/NiagaraStatelessEmitter.h"
+#include "Stateless/NiagaraStatelessModule.h"
 #include "ViewModels/NiagaraEmitterHandleViewModel.h"
 #include "ViewModels/NiagaraSystemViewModel.h"
 #include "ViewModels/Stack/NiagaraStackEmitterSettingsGroup.h"
@@ -35,6 +37,7 @@
 #include "ScopedTransaction.h"
 #include "DeviceProfiles/DeviceProfile.h"
 #include "DeviceProfiles/DeviceProfileManager.h"
+#include "ObjectEditorUtils.h"
 
 #include UE_INLINE_GENERATED_CPP_BY_NAME(NiagaraValidationRules)
 
@@ -637,6 +640,62 @@ void UNiagaraValidationRule_BannedRenderers::CheckValidity(const FNiagaraValidat
 	}
 } 
 
+void UNiagaraValidationRule_Lightweight::CheckValidity(const FNiagaraValidationContext& Context, TArray<FNiagaraValidationResult>& OutResults) const
+{
+	UNiagaraSystem& System = Context.ViewModel->GetSystem();
+	TArray<TSharedRef<FNiagaraEmitterHandleViewModel>> EmitterHandleViewModels = Context.ViewModel->GetEmitterHandleViewModels();
+	for (TSharedRef<FNiagaraEmitterHandleViewModel> EmitterHandleModel : EmitterHandleViewModels)
+	{
+		FNiagaraEmitterHandle* EmitterHandle = EmitterHandleModel->GetEmitterHandle();
+		UNiagaraStatelessEmitter* StatelessEmitter = EmitterHandle && EmitterHandle->GetEmitterMode() == ENiagaraEmitterMode::Stateless ? EmitterHandle->GetStatelessEmitter() : nullptr;
+		if (StatelessEmitter == nullptr)
+		{
+			continue;
+		}
+
+		TArray<FNiagaraPlatformSetConflictInfo> Conflicts = NiagaraValidation::GatherPlatformSetConflicts(&Platforms, &StatelessEmitter->GetPlatformSet());
+		if (Conflicts.Num() == 0)
+		{
+			continue;
+		}
+
+		if (UsedWithEmitter.IsSet())
+		{
+			OutResults.Emplace(
+				UsedWithEmitter.GetValue(),
+				LOCTEXT("StatelessNotAllowed", "Lightweight emitter is being used."),
+				FText::Format(LOCTEXT("StatelessNotAllowedDesc", "Lightweight emitter {0} is not allowed, please disable or remove."), FText::FromName(StatelessEmitter->GetFName())),
+				NiagaraValidation::GetStackEntry<UNiagaraStackSystemPropertiesItem>(Context.ViewModel->GetSystemStackViewModel())
+			);
+		}
+
+		if (UsingExperimentalModule.IsSet())
+		{
+			for (const UNiagaraStatelessModule* StatelessModule : StatelessEmitter->GetModules())
+			{
+				if (!StatelessModule || !StatelessModule->IsModuleEnabled())
+				{
+					continue;
+				}
+
+				bool bIsExperimental = false;
+				bool bIsEarlyAccess = false;
+				FString MostDerivedDevelopmentClassName;
+				FObjectEditorUtils::GetClassDevelopmentStatus(StatelessModule->GetClass(), bIsExperimental, bIsEarlyAccess, MostDerivedDevelopmentClassName);
+				if (bIsExperimental || bIsEarlyAccess)
+				{
+					OutResults.Emplace(
+						UsingExperimentalModule.GetValue(),
+						LOCTEXT("StatelessModuleNotAllowed", "Experimental lightweight modules are being used."),
+						FText::Format(LOCTEXT("StatelessModuleNotAllowedDesc", "Experimental lightweight module {0} is not allowed, please disable or remove."), StatelessModule->GetClass()->GetDisplayNameText()),
+						NiagaraValidation::GetStackEntry<UNiagaraStackSystemPropertiesItem>(Context.ViewModel->GetSystemStackViewModel())
+					);
+				}
+			}
+		}
+	}
+}
+
 void UNiagaraValidationRule_BannedModules::CheckValidity(const FNiagaraValidationContext& Context, TArray<FNiagaraValidationResult>& Results)  const
 {
 	UNiagaraSystem& System = Context.ViewModel->GetSystem();
diff --git a/Engine/Plugins/FX/Niagara/Source/NiagaraEditor/Public/NiagaraValidationRules.h b/Engine/Plugins/FX/Niagara/Source/NiagaraEditor/Public/NiagaraValidationRules.h
--- a/Engine/Plugins/FX/Niagara/Source/NiagaraEditor/Public/NiagaraValidationRules.h
+++ b/Engine/Plugins/FX/Niagara/Source/NiagaraEditor/Public/NiagaraValidationRules.h
@@ -217,6 +217,28 @@ public:
 	virtual void CheckValidity(const FNiagaraValidationContext& Context, TArray<FNiagaraValidationResult>& OutResults) const override;
 };
 
+/** Validation rule to check for lightweight usage. */
+UCLASS(Category = "Validation", DisplayName = "Lightweight Validation")
+class UNiagaraValidationRule_Lightweight : public UNiagaraValidationRule
+{
+	GENERATED_BODY()
+
+public:
+	// Platforms this validation rule will apply to.
+	UPROPERTY(EditAnywhere, Category = Validation)
+	FNiagaraPlatformSet Platforms;
+
+	// When set if we have an emitter present it will be flagged at this severity
+	UPROPERTY(EditAnywhere, Category = Validation)
+	TOptional<ENiagaraValidationSeverity> UsedWithEmitter;
+
+	// When set if an experimental module is found it will be flagged at this severity
+	UPROPERTY(EditAnywhere, Category = Validation)
+	TOptional<ENiagaraValidationSeverity> UsingExperimentalModule;
+
+	virtual void CheckValidity(const FNiagaraValidationContext& Context, TArray<FNiagaraValidationResult>& OutResults) const override;
+};
+
 /** This validation rule can ban the use of certain modules on all or a subset of platforms. */
 UCLASS(Category = "Validation", DisplayName = "Banned Modules")
 class UNiagaraValidationRule_BannedModules : public UNiagaraValidationRule
