- Add audit information for using modules that are experimental

#rnx

diff --git a/Engine/Plugins/FX/Niagara/Source/NiagaraEditor/Private/Commandlets/NiagaraStatelessAuditCommandlet.cpp b/Engine/Plugins/FX/Niagara/Source/NiagaraEditor/Private/Commandlets/NiagaraStatelessAuditCommandlet.cpp
--- a/Engine/Plugins/FX/Niagara/Source/NiagaraEditor/Private/Commandlets/NiagaraStatelessAuditCommandlet.cpp
+++ b/Engine/Plugins/FX/Niagara/Source/NiagaraEditor/Private/Commandlets/NiagaraStatelessAuditCommandlet.cpp
@@ -1,36 +1,16 @@
 // Copyright Epic Games, Inc. All Rights Reserved.
 
 #include "Commandlets/NiagaraStatelessAuditCommandlet.h"
-//#include "DeviceProfiles/DeviceProfile.h"
-//#include "DeviceProfiles/DeviceProfileManager.h"
-#include "HAL/FileManager.h"
-//#include "HAL/PlatformFileManager.h"
-//#include "Misc/FileHelper.h"
-//#include "Misc/Paths.h"
-//#include "Misc/PackageName.h"
-//#include "Modules/ModuleManager.h"
-#include "AssetRegistry/AssetData.h"
-//#include "AssetRegistry/ARFilter.h"
-#include "AssetRegistry/AssetRegistryModule.h"
-#include "UObject/UObjectGlobals.h"
-//#include "CollectionManagerTypes.h"
-//#include "ICollectionManager.h"
-//#include "CollectionManagerModule.h"
-//#include "FileHelpers.h"
-//
-//#include "ViewModels/NiagaraSystemViewModel.h"
-//#include "NiagaraSettings.h"
 #include "Stateless/NiagaraStatelessEmitter.h"
 #include "Stateless/NiagaraStatelessEmitterTemplate.h"
 #include "Stateless/NiagaraStatelessModule.h"
 #include "NiagaraSystem.h"
-//#include "NiagaraSystemImpl.h"
-//#include "NiagaraSimulationStageBase.h"
-//#include "NiagaraRendererProperties.h"
-//#include "NiagaraLightRendererProperties.h"
-//#include "NiagaraRibbonRendererProperties.h"
-//#include "NiagaraValidationRules.h"
-//#include "NiagaraValidationRule.h"
+
+#include "AssetRegistry/AssetData.h"
+#include "AssetRegistry/AssetRegistryModule.h"
+#include "HAL/FileManager.h"
+#include "UObject/UObjectGlobals.h"
+#include "ObjectEditorUtils.h"
 
 #include UE_INLINE_GENERATED_CPP_BY_NAME(NiagaraStatelessAuditCommandlet)
 
@@ -142,10 +122,22 @@ void UNiagaraStatelessAuditCommandlet::ProcessSystem(UNiagaraSystem* NiagaraSyst
 			{
 				continue;
 			}
-			FModuleUsage& ModuleUsage = ModuleUsageMap.FindOrAdd(StatelessModule->GetFName());
-			++ModuleUsage.UsageCount;
-			ModuleUsage.SystemNames.AddUnique(NiagaraSystem->GetName());
-			ModuleUsage.EmitterNames.AddUnique(StatelessEmitter->GetName());
+
+			FModuleUsage* ModuleUsage = ModuleUsageMap.Find(StatelessModule->GetFName());
+			if (ModuleUsage == nullptr)
+			{
+				ModuleUsage = &ModuleUsageMap.Add(StatelessModule->GetFName());
+
+				FString MostDerivedDevelopmentClassName;
+				FObjectEditorUtils::GetClassDevelopmentStatus(StatelessModule->GetClass(), ModuleUsage->bIsExperimental, ModuleUsage->bIsEarlyAccess, MostDerivedDevelopmentClassName);
+				bAnyExperimentalModules |= ModuleUsage->bIsExperimental;
+				bAnyEarlyAccessModules |= ModuleUsage->bIsEarlyAccess;
+			}
+
+			++ModuleUsage->UsageCount;
+			ModuleUsage->SystemNames.AddUnique(NiagaraSystem->GetName());
+			ModuleUsage->EmitterNames.AddUnique(StatelessEmitter->GetName());
+			ModuleUsage->SystemAndEmitterNames.AddUnique(StatelessEmitter->GetPathName());
 		}
 	}
 }
@@ -197,6 +189,30 @@ void UNiagaraStatelessAuditCommandlet::WriteResults() const
 			OutputStream->Logf(TEXT("%s"), *EmitterName);
 		}
 	}
+
+	// Output information about using experimental / early access modules
+	if (bAnyExperimentalModules || bAnyEarlyAccessModules)
+	{
+		TUniquePtr<FArchive> OutputStream = GetOutputFile(TEXT("InvalidModuleUsage.txt"));
+		if (OutputStream.IsValid())
+		{
+			for (auto It = ModuleUsageMap.CreateConstIterator(); It; ++It)
+			{
+				const FName ModuleName = It.Key();
+				const FModuleUsage& ModuleUsage = It.Value();
+				if (!ModuleUsage.bIsEarlyAccess && !ModuleUsage.bIsExperimental)
+				{
+					continue;
+				}
+
+				OutputStream->Logf(TEXT("Module(%s) is (%s) and used in the following:"), *ModuleName.ToString(), ModuleUsage.bIsExperimental ? TEXT("Experimental") : TEXT("EarlyAccess"));
+				for (const FString& SystemAndEmitterName : ModuleUsage.SystemAndEmitterNames)
+				{
+					OutputStream->Logf(TEXT("\t- %s"), *SystemAndEmitterName);
+				}
+			}
+		}
+	}
 }
 
 TUniquePtr<FArchive> UNiagaraStatelessAuditCommandlet::GetOutputFile(const TCHAR* Filename) const
diff --git a/Engine/Plugins/FX/Niagara/Source/NiagaraEditor/Public/Commandlets/NiagaraStatelessAuditCommandlet.h b/Engine/Plugins/FX/Niagara/Source/NiagaraEditor/Public/Commandlets/NiagaraStatelessAuditCommandlet.h
--- a/Engine/Plugins/FX/Niagara/Source/NiagaraEditor/Public/Commandlets/NiagaraStatelessAuditCommandlet.h
+++ b/Engine/Plugins/FX/Niagara/Source/NiagaraEditor/Public/Commandlets/NiagaraStatelessAuditCommandlet.h
@@ -17,9 +17,13 @@ class UNiagaraStatelessAuditCommandlet : public UCommandlet
 private:
 	struct FModuleUsage
 	{
+		bool			bIsExperimental = false;
+		bool			bIsEarlyAccess = false;
+
 		int32			UsageCount = 0;
 		TArray<FString>	SystemNames;
 		TArray<FString>	EmitterNames;
+		TArray<FString> SystemAndEmitterNames;
 	};
 
 public:
@@ -35,4 +39,6 @@ private:
 private:
 	FString						AuditOutputFolder;
 	TMap<FName, FModuleUsage>	ModuleUsageMap;
+	bool						bAnyExperimentalModules = false;
+	bool						bAnyEarlyAccessModules = false;
 };
