- Add BP editor utility function to capture a Niagara component to a static mesh

#rnx

diff --git a/Engine/Plugins/FX/Niagara/Source/NiagaraEditor/Private/NiagaraBakerFunctionLibrary.cpp b/Engine/Plugins/FX/Niagara/Source/NiagaraEditor/Private/NiagaraBakerFunctionLibrary.cpp
new file mode 100644
--- /dev/null
+++ b/Engine/Plugins/FX/Niagara/Source/NiagaraEditor/Private/NiagaraBakerFunctionLibrary.cpp
@@ -0,0 +1,25 @@
+// Copyright Epic Games, Inc. All Rights Reserved.
+
+#include "NiagaraBakerFunctionLibrary.h"
+
+#include "NiagaraBakerRendererOutputStaticMesh.h"
+
+void UNiagaraBakerFunctionLibrary::CaptureNiagaraToStaticMesh(UNiagaraComponent* ComponentToCapture, UStaticMesh* StaticMeshOutput, FNiagaraRendererReadbackParameters ReadbackParameters)
+{
+#if WITH_NIAGARA_RENDERER_READBACK
+	NiagaraRendererReadback::EnqueueReadback(
+		ComponentToCapture,
+		[ComponentToCapture, StaticMeshOutput](const FNiagaraRendererReadbackResult& ReadbackResult)
+		{
+			// Failed or no data
+			if (ReadbackResult.NumVertices == 0)
+			{
+				return;
+			}
+
+			FNiagaraBakerRendererOutputStaticMesh::ConvertReadbackResultsToStaticMesh(ReadbackResult, StaticMeshOutput);
+		},
+		ReadbackParameters
+	);
+#endif
+}
diff --git a/Engine/Plugins/FX/Niagara/Source/NiagaraEditor/Private/NiagaraBakerFunctionLibrary.h b/Engine/Plugins/FX/Niagara/Source/NiagaraEditor/Private/NiagaraBakerFunctionLibrary.h
new file mode 100644
--- /dev/null
+++ b/Engine/Plugins/FX/Niagara/Source/NiagaraEditor/Private/NiagaraBakerFunctionLibrary.h
@@ -0,0 +1,21 @@
+// Copyright Epic Games, Inc. All Rights Reserved.
+
+#pragma once
+
+#include "Kismet/BlueprintFunctionLibrary.h"
+#include "NiagaraRendererReadback.h"
+
+#include "NiagaraBakerFunctionLibrary.generated.h"
+
+class UNiagaraComponent;
+class UStaticMesh;
+
+UCLASS(MinimalAPI)
+class UNiagaraBakerFunctionLibrary : public UBlueprintFunctionLibrary
+{
+	GENERATED_BODY()
+
+public:
+	UFUNCTION(BlueprintCallable, Category = "Niagara Baker", meta = (Keywords = "Niagara Baker", WorldContext = "WorldContextObject", UnsafeDuringActorConstruction = "true"))
+	static void CaptureNiagaraToStaticMesh(UNiagaraComponent* ComponentToCapture, UStaticMesh* StaticMeshOutput, FNiagaraRendererReadbackParameters ReadbackParameters = FNiagaraRendererReadbackParameters());
+};
diff --git a/Engine/Plugins/FX/Niagara/Source/NiagaraEditor/Private/NiagaraBakerRendererOutputStaticMesh.cpp b/Engine/Plugins/FX/Niagara/Source/NiagaraEditor/Private/NiagaraBakerRendererOutputStaticMesh.cpp
--- a/Engine/Plugins/FX/Niagara/Source/NiagaraEditor/Private/NiagaraBakerRendererOutputStaticMesh.cpp
+++ b/Engine/Plugins/FX/Niagara/Source/NiagaraEditor/Private/NiagaraBakerRendererOutputStaticMesh.cpp
@@ -95,10 +95,10 @@ void FNiagaraBakerRendererOutputStaticMesh::BakeFrame(FNiagaraBakerFeedbackConte
 #if WITH_NIAGARA_RENDERER_READBACK
 	NiagaraRendererReadback::EnqueueReadback(
 		PreviewComponent,
-		[BakerOutput, FrameIndex](const FNiagaraRendererReadbackResult& Result)
+		[BakerOutput, FrameIndex](const FNiagaraRendererReadbackResult& ReadbackResult)
 		{
 			// Failed or no data
-			if (Result.NumVertices == 0)
+			if (ReadbackResult.NumVertices == 0)
 			{
 				return;
 			}
@@ -111,96 +111,7 @@ void FNiagaraBakerRendererOutputStaticMesh::BakeFrame(FNiagaraBakerFeedbackConte
 				return;
 			}
 
-			const uint32 NumTexCoords = Result.VertexTexCoordNum;
-			const uint32 NumTriangles = Result.NumVertices / 3;
-
-			// Create Mesh Description
-			FMeshDescription MeshDescription;
-			TVertexAttributesRef<FVector3f> VertexPositions = MeshDescription.GetVertexPositions();
-			TVertexInstanceAttributesRef<FVector3f> VertexInstanceNormals = MeshDescription.VertexInstanceAttributes().RegisterAttribute<FVector3f>(MeshAttribute::VertexInstance::Normal, 1, FVector3f::ZeroVector, EMeshAttributeFlags::Mandatory);
-			TVertexInstanceAttributesRef<FVector3f> VertexInstanceTangents = MeshDescription.VertexInstanceAttributes().RegisterAttribute<FVector3f>(MeshAttribute::VertexInstance::Tangent, 1, FVector3f::ZeroVector, EMeshAttributeFlags::Mandatory);
-			TVertexInstanceAttributesRef<float> VertexInstanceBinormalSigns = MeshDescription.VertexInstanceAttributes().RegisterAttribute<float>(MeshAttribute::VertexInstance::BinormalSign, 1, 1.0f, EMeshAttributeFlags::Mandatory);
-			TVertexInstanceAttributesRef<FVector4f> VertexInstanceColors = MeshDescription.VertexInstanceAttributes().RegisterAttribute<FVector4f>(MeshAttribute::VertexInstance::Color, 1, FVector4f(1.0f, 1.0f, 1.0f, 1.0f), EMeshAttributeFlags::Lerpable | EMeshAttributeFlags::Mandatory);
-			TVertexInstanceAttributesRef<FVector2f> VertexInstanceUVs = MeshDescription.VertexInstanceAttributes().RegisterAttribute<FVector2f>(MeshAttribute::VertexInstance::TextureCoordinate, NumTexCoords, FVector2f::ZeroVector, EMeshAttributeFlags::Lerpable | EMeshAttributeFlags::Mandatory);
-
-			MeshDescription.EdgeAttributes().RegisterAttribute<bool>(MeshAttribute::Edge::IsHard, 1, false, EMeshAttributeFlags::Mandatory);
-			TPolygonGroupAttributesRef<FName> PolygonGroupSlotNames = MeshDescription.PolygonGroupAttributes().RegisterAttribute<FName>(MeshAttribute::PolygonGroup::ImportedMaterialSlotName, 1, NAME_None, EMeshAttributeFlags::Mandatory); //The unique key to match the mesh material slot
-
-			// Reserve space
-			MeshDescription.ReserveNewVertices(Result.NumVertices);
-			MeshDescription.ReserveNewVertexInstances(NumTriangles);
-			MeshDescription.ReserveNewEdges(NumTriangles);
-			MeshDescription.ReserveNewPolygons(NumTriangles);
-			MeshDescription.ReserveNewPolygonGroups(Result.Sections.Num());
-
-			// Build vertices
-			TArray<FVertexInstanceID> VertexInstances;
-			VertexInstances.Reserve(Result.NumVertices);
-			for (uint32 iVertex=0; iVertex < Result.NumVertices; ++iVertex)
-			{
-				const FVertexID VertexID = MeshDescription.CreateVertex();
-				check(VertexID.GetValue() == iVertex);
-				VertexPositions[VertexID] = Result.HasPosition() ? Result.GetPosition(iVertex) : FVector3f::ZeroVector;
-
-				FVertexInstanceID VertexInstanceID = MeshDescription.CreateVertexInstance(VertexID);
-				VertexInstances.Add(VertexInstanceID);
-
-				const FVector3f TangentX = Result.HasTangentBasis() ? Result.GetTangentX(iVertex) : FVector3f::XAxisVector;
-				const FVector3f TangentY = Result.HasTangentBasis() ? Result.GetTangentY(iVertex) : FVector3f::YAxisVector;
-				const FVector3f TangentZ = Result.HasTangentBasis() ? Result.GetTangentZ(iVertex) : FVector3f::ZAxisVector;
-				const float TangentSign = FVector3f::DotProduct(FVector3f::CrossProduct(TangentX, TangentZ), TangentY) < 0.0f ? -1.0f : 1.0f;
-				VertexInstanceNormals.Set(VertexInstanceID, TangentX);
-				VertexInstanceTangents.Set(VertexInstanceID, TangentZ);
-				VertexInstanceBinormalSigns.Set(VertexInstanceID, TangentSign);
-
-				VertexInstanceColors.Set(VertexInstanceID, Result.HasColor() ? Result.GetColor(iVertex) : FLinearColor::White);
-				for (uint32 iTexCoord=0; iTexCoord < NumTexCoords; ++iTexCoord)
-				{
-					VertexInstanceUVs.Set(VertexInstanceID, iTexCoord, Result.GetTexCoord(iVertex, iTexCoord));
-				}
-			}
-
-			// Build sections / triangles
-			TArray<FStaticMaterial> StaticMaterials;
-			for (int32 iSection=0; iSection < Result.Sections.Num(); ++iSection)
-			{
-				const FNiagaraRendererReadbackResult::FSection& Section = Result.Sections[iSection];
-
-				const FName MaterialSlotName(*FString::Printf(TEXT("Section%d"), iSection));
-				const FPolygonGroupID PolyGroupId = MeshDescription.CreatePolygonGroup();
-				PolygonGroupSlotNames[PolyGroupId] = MaterialSlotName;
-
-				FStaticMaterial& StaticMaterial			= StaticMaterials.AddDefaulted_GetRef();
-				StaticMaterial.MaterialInterface		= Section.WeakMaterialInterface.Get();
-				StaticMaterial.MaterialSlotName			= MaterialSlotName;
-				StaticMaterial.ImportedMaterialSlotName	= MaterialSlotName;
-
-				for (uint32 iTriangle=0; iTriangle < Section.NumTriangles; ++iTriangle)
-				{
-					const uint32 BaseIndex = (Section.FirstTriangle + iTriangle) * 3;
-					const FPolygonID PolygonID = MeshDescription.CreatePolygon(PolyGroupId, { VertexInstances[BaseIndex + 0], VertexInstances[BaseIndex + 1], VertexInstances[BaseIndex + 2] });
-					MeshDescription.ComputePolygonTriangulation(PolygonID);
-				}
-			}
-
-			StaticMesh->SetNumSourceModels(1);
-			{
-				FMeshBuildSettings& MeshBuildSettings = StaticMesh->GetSourceModel(0).BuildSettings;
-				MeshBuildSettings.bRecomputeNormals = false;
-				MeshBuildSettings.bRecomputeTangents = false;
-			}
-
-			StaticMesh->SetStaticMaterials(StaticMaterials);
-
-			// Build Mesh
-			UStaticMesh::FBuildMeshDescriptionsParams Params;
-			Params.bFastBuild = WITH_EDITOR ? false : true;
-			Params.bUseHashAsGuid = true;
-			Params.bMarkPackageDirty = false;
-			Params.bCommitMeshDescription = true;
-			Params.bAllowCpuAccess = false;
-
-			StaticMesh->BuildFromMeshDescriptions({ &MeshDescription }, Params);
+			ConvertReadbackResultsToStaticMesh(ReadbackResult, StaticMesh);
 		},
 		BakerOutput->ExportParameters
 	);
@@ -218,3 +129,107 @@ void FNiagaraBakerRendererOutputStaticMesh::EndBake(FNiagaraBakerFeedbackContext
 		BakeRenderTarget = nullptr;
 	}
 }
+
+bool FNiagaraBakerRendererOutputStaticMesh::ConvertReadbackResultsToStaticMesh(const FNiagaraRendererReadbackResult& ReadbackResult, UStaticMesh* StaticMesh)
+{
+#if WITH_NIAGARA_RENDERER_READBACK
+	// Failed or no data
+	if (ReadbackResult.NumVertices == 0 || StaticMesh == nullptr)
+	{
+		return false;
+	}
+
+	const uint32 NumTexCoords = ReadbackResult.VertexTexCoordNum;
+	const uint32 NumTriangles = ReadbackResult.NumVertices / 3;
+
+	// Create Mesh Description
+	FMeshDescription MeshDescription;
+	TVertexAttributesRef<FVector3f> VertexPositions = MeshDescription.GetVertexPositions();
+	TVertexInstanceAttributesRef<FVector3f> VertexInstanceNormals = MeshDescription.VertexInstanceAttributes().RegisterAttribute<FVector3f>(MeshAttribute::VertexInstance::Normal, 1, FVector3f::ZeroVector, EMeshAttributeFlags::Mandatory);
+	TVertexInstanceAttributesRef<FVector3f> VertexInstanceTangents = MeshDescription.VertexInstanceAttributes().RegisterAttribute<FVector3f>(MeshAttribute::VertexInstance::Tangent, 1, FVector3f::ZeroVector, EMeshAttributeFlags::Mandatory);
+	TVertexInstanceAttributesRef<float> VertexInstanceBinormalSigns = MeshDescription.VertexInstanceAttributes().RegisterAttribute<float>(MeshAttribute::VertexInstance::BinormalSign, 1, 1.0f, EMeshAttributeFlags::Mandatory);
+	TVertexInstanceAttributesRef<FVector4f> VertexInstanceColors = MeshDescription.VertexInstanceAttributes().RegisterAttribute<FVector4f>(MeshAttribute::VertexInstance::Color, 1, FVector4f(1.0f, 1.0f, 1.0f, 1.0f), EMeshAttributeFlags::Lerpable | EMeshAttributeFlags::Mandatory);
+	TVertexInstanceAttributesRef<FVector2f> VertexInstanceUVs = MeshDescription.VertexInstanceAttributes().RegisterAttribute<FVector2f>(MeshAttribute::VertexInstance::TextureCoordinate, NumTexCoords, FVector2f::ZeroVector, EMeshAttributeFlags::Lerpable | EMeshAttributeFlags::Mandatory);
+
+	MeshDescription.EdgeAttributes().RegisterAttribute<bool>(MeshAttribute::Edge::IsHard, 1, false, EMeshAttributeFlags::Mandatory);
+	TPolygonGroupAttributesRef<FName> PolygonGroupSlotNames = MeshDescription.PolygonGroupAttributes().RegisterAttribute<FName>(MeshAttribute::PolygonGroup::ImportedMaterialSlotName, 1, NAME_None, EMeshAttributeFlags::Mandatory); //The unique key to match the mesh material slot
+
+	// Reserve space
+	MeshDescription.ReserveNewVertices(ReadbackResult.NumVertices);
+	MeshDescription.ReserveNewVertexInstances(NumTriangles);
+	MeshDescription.ReserveNewEdges(NumTriangles);
+	MeshDescription.ReserveNewPolygons(NumTriangles);
+	MeshDescription.ReserveNewPolygonGroups(ReadbackResult.Sections.Num());
+
+	// Build vertices
+	TArray<FVertexInstanceID> VertexInstances;
+	VertexInstances.Reserve(ReadbackResult.NumVertices);
+	for (uint32 iVertex=0; iVertex < ReadbackResult.NumVertices; ++iVertex)
+	{
+		const FVertexID VertexID = MeshDescription.CreateVertex();
+		check(VertexID.GetValue() == iVertex);
+		VertexPositions[VertexID] = ReadbackResult.HasPosition() ? ReadbackResult.GetPosition(iVertex) : FVector3f::ZeroVector;
+
+		FVertexInstanceID VertexInstanceID = MeshDescription.CreateVertexInstance(VertexID);
+		VertexInstances.Add(VertexInstanceID);
+
+		const FVector3f TangentX = ReadbackResult.HasTangentBasis() ? ReadbackResult.GetTangentX(iVertex) : FVector3f::XAxisVector;
+		const FVector3f TangentY = ReadbackResult.HasTangentBasis() ? ReadbackResult.GetTangentY(iVertex) : FVector3f::YAxisVector;
+		const FVector3f TangentZ = ReadbackResult.HasTangentBasis() ? ReadbackResult.GetTangentZ(iVertex) : FVector3f::ZAxisVector;
+		const float TangentSign = FVector3f::DotProduct(FVector3f::CrossProduct(TangentX, TangentZ), TangentY) < 0.0f ? -1.0f : 1.0f;
+		VertexInstanceNormals.Set(VertexInstanceID, TangentX);
+		VertexInstanceTangents.Set(VertexInstanceID, TangentZ);
+		VertexInstanceBinormalSigns.Set(VertexInstanceID, TangentSign);
+
+		VertexInstanceColors.Set(VertexInstanceID, ReadbackResult.HasColor() ? ReadbackResult.GetColor(iVertex) : FLinearColor::White);
+		for (uint32 iTexCoord=0; iTexCoord < NumTexCoords; ++iTexCoord)
+		{
+			VertexInstanceUVs.Set(VertexInstanceID, iTexCoord, ReadbackResult.GetTexCoord(iVertex, iTexCoord));
+		}
+	}
+
+	// Build sections / triangles
+	TArray<FStaticMaterial> StaticMaterials;
+	for (int32 iSection=0; iSection < ReadbackResult.Sections.Num(); ++iSection)
+	{
+		const FNiagaraRendererReadbackResult::FSection& Section = ReadbackResult.Sections[iSection];
+
+		const FName MaterialSlotName(*FString::Printf(TEXT("Section%d"), iSection));
+		const FPolygonGroupID PolyGroupId = MeshDescription.CreatePolygonGroup();
+		PolygonGroupSlotNames[PolyGroupId] = MaterialSlotName;
+
+		FStaticMaterial& StaticMaterial			= StaticMaterials.AddDefaulted_GetRef();
+		StaticMaterial.MaterialInterface		= Section.WeakMaterialInterface.Get();
+		StaticMaterial.MaterialSlotName			= MaterialSlotName;
+		StaticMaterial.ImportedMaterialSlotName	= MaterialSlotName;
+
+		for (uint32 iTriangle=0; iTriangle < Section.NumTriangles; ++iTriangle)
+		{
+			const uint32 BaseIndex = (Section.FirstTriangle + iTriangle) * 3;
+			const FPolygonID PolygonID = MeshDescription.CreatePolygon(PolyGroupId, { VertexInstances[BaseIndex + 0], VertexInstances[BaseIndex + 1], VertexInstances[BaseIndex + 2] });
+			MeshDescription.ComputePolygonTriangulation(PolygonID);
+		}
+	}
+
+	StaticMesh->SetNumSourceModels(1);
+	{
+		FMeshBuildSettings& MeshBuildSettings = StaticMesh->GetSourceModel(0).BuildSettings;
+		MeshBuildSettings.bRecomputeNormals = false;
+		MeshBuildSettings.bRecomputeTangents = false;
+	}
+
+	StaticMesh->SetStaticMaterials(StaticMaterials);
+
+	// Build Mesh
+	UStaticMesh::FBuildMeshDescriptionsParams Params;
+	Params.bFastBuild = WITH_EDITOR ? false : true;
+	Params.bUseHashAsGuid = true;
+	Params.bMarkPackageDirty = false;
+	Params.bCommitMeshDescription = true;
+	Params.bAllowCpuAccess = false;
+
+	StaticMesh->BuildFromMeshDescriptions({ &MeshDescription }, Params);
+
+	return true;
+#endif
+}
diff --git a/Engine/Plugins/FX/Niagara/Source/NiagaraEditor/Private/NiagaraBakerRendererOutputStaticMesh.h b/Engine/Plugins/FX/Niagara/Source/NiagaraEditor/Private/NiagaraBakerRendererOutputStaticMesh.h
--- a/Engine/Plugins/FX/Niagara/Source/NiagaraEditor/Private/NiagaraBakerRendererOutputStaticMesh.h
+++ b/Engine/Plugins/FX/Niagara/Source/NiagaraEditor/Private/NiagaraBakerRendererOutputStaticMesh.h
@@ -9,6 +9,9 @@
 #include "NiagaraBakerRendererOutputStaticMesh.generated.h"
 
 class UTextureRenderTarget2D;
+class UStaticMesh;
+
+struct FNiagaraRendererReadbackResult;
 
 UCLASS()
 class UNiagaraBakerStaticMeshFactoryNew : public UFactory
@@ -34,6 +37,8 @@ public:
 	virtual void BakeFrame(FNiagaraBakerFeedbackContext& FeedbackContext, UNiagaraBakerOutput* InBakerOutput, int FrameIndex, const FNiagaraBakerRenderer& BakerRenderer) override;
 	virtual void EndBake(FNiagaraBakerFeedbackContext& FeedbackContext, UNiagaraBakerOutput* InBakerOutput) override;
 
+	static bool ConvertReadbackResultsToStaticMesh(const FNiagaraRendererReadbackResult& ReadbackResult, UStaticMesh* StaticMesh);
+
 private:
 	UTextureRenderTarget2D* BakeRenderTarget = nullptr;
 };
