Added favorites feature to Niagara Menus.
Works on input dropdowns (dynamic inputs, linked attributes), graph actions and modules.
You can specify favorite by default actions, but are free to favorite or unfavorite any action.
Favorited/suggested actions get a search scorre bonus.

- Also added missing FavoriteOutline icon to starship
- Also added TagSimple icon to starship
- got rid of section specification for niagara actions. The sections were only used for the Suggested and the General section, but the action should not determine whether it is suggested or not. It has been replaced by FNiagaraFavoritesActionData, which includes an action identifier and a "favorite by default" member.
Depending on whether an action is 'favorite by default' or not, toggling favorites writes into one of two locations. Either the explicitly favorited actions, or the explicitly unfavorited actions. This architecture allows us to freely favorite and unfavorite any action while still allowing default favorites.
- Added "OnItemRowHoverEvent" delegate to SItemSelector. This works specifically and only on items and is used to surface information about hover status of an item row.
This is used to display the favorite button only on items that are currently hovered on.
- Added 'Suggested' feature for parameters in Parameter Definition assets, works like Suggested for script assets
- added Particles.Position, Particles.ID and Particles.Color to suggested parameters.

#virtualized

diff --git a/Engine/Plugins/FX/Niagara/Source/Niagara/Public/NiagaraVariableMetaData.h b/Engine/Plugins/FX/Niagara/Source/Niagara/Public/NiagaraVariableMetaData.h
--- a/Engine/Plugins/FX/Niagara/Source/Niagara/Public/NiagaraVariableMetaData.h
+++ b/Engine/Plugins/FX/Niagara/Source/Niagara/Public/NiagaraVariableMetaData.h
@@ -274,6 +274,10 @@ struct FNiagaraVariableMetaData
 	UPROPERTY(EditAnywhere, Category = "Variable", meta = (SkipForCompileHash = "true"))
 	FNiagaraInputConditionMetadata VisibleCondition;
 
+	/** If this variable is suggested by default, it will get prioritized in menus. */
+	UPROPERTY(EditAnywhere, Category = "Variable", meta = (SkipForCompileHash = "true"))
+	bool bSuggested = false;
+	
 	UPROPERTY(EditAnywhere, Category = "Variable", DisplayName = "Property Metadata", meta = (ToolTip = "Property Metadata", SkipForCompileHash = "true"))
 	TMap<FName, FString> PropertyMetaData;
 
diff --git a/Engine/Plugins/FX/Niagara/Source/NiagaraEditor/Private/Config/NiagaraFavoriteActionsConfig.cpp b/Engine/Plugins/FX/Niagara/Source/NiagaraEditor/Private/Config/NiagaraFavoriteActionsConfig.cpp
new file mode 100644
--- /dev/null
+++ b/Engine/Plugins/FX/Niagara/Source/NiagaraEditor/Private/Config/NiagaraFavoriteActionsConfig.cpp
@@ -0,0 +1,71 @@
+﻿// Copyright Epic Games, Inc. All Rights Reserved.
+
+#include "Config/NiagaraFavoriteActionsConfig.h"
+
+TStrongObjectPtr<UNiagaraFavoriteActionsConfig> UNiagaraFavoriteActionsConfig::Instance = nullptr;
+
+bool FNiagaraFavoriteActionsProfile::IsFavorite(FNiagaraFavoritesActionData InAction)
+{
+	if(InAction.ActionIdentifier.IsValid() == false)
+	{
+		return false;
+	}
+	
+	if(InAction.bFavoriteByDefault)
+	{
+		return UnfavoriteActions.Contains(InAction.ActionIdentifier) == false;
+	}
+	
+	return FavoriteActions.Contains(InAction.ActionIdentifier);
+}
+
+void FNiagaraFavoriteActionsProfile::ToggleFavoriteAction(FNiagaraFavoritesActionData InAction)
+{
+	if(InAction.bFavoriteByDefault)
+	{
+		if(UnfavoriteActions.Contains(InAction.ActionIdentifier))
+		{
+			UnfavoriteActions.Remove(InAction.ActionIdentifier);
+		}
+		else
+		{
+			UnfavoriteActions.Add(InAction.ActionIdentifier);
+		}
+	}
+	else
+	{
+		if(FavoriteActions.Contains(InAction.ActionIdentifier))
+		{
+			FavoriteActions.Remove(InAction.ActionIdentifier);
+		}
+		else
+		{
+			FavoriteActions.Add(InAction.ActionIdentifier);
+		}
+	}
+}
+
+FNiagaraFavoriteActionsProfile& UNiagaraFavoriteActionsConfig::GetActionsProfile(FName ProfileName)
+{
+	return Profiles.FindOrAdd(ProfileName);
+}
+
+UNiagaraFavoriteActionsConfig* UNiagaraFavoriteActionsConfig::Get()
+{
+	if(Instance.IsValid() == false)
+	{
+		Instance.Reset(NewObject<UNiagaraFavoriteActionsConfig>());
+		Instance->LoadEditorConfig();
+	}
+
+	return Instance.Get();
+}
+
+void UNiagaraFavoriteActionsConfig::Shutdown()
+{
+	if(Instance != nullptr)
+	{
+		Instance->SaveEditorConfig();
+		Instance.Reset();
+	}
+}
diff --git a/Engine/Plugins/FX/Niagara/Source/NiagaraEditor/Private/Customizations/NiagaraScriptVariableCustomization.cpp b/Engine/Plugins/FX/Niagara/Source/NiagaraEditor/Private/Customizations/NiagaraScriptVariableCustomization.cpp
--- a/Engine/Plugins/FX/Niagara/Source/NiagaraEditor/Private/Customizations/NiagaraScriptVariableCustomization.cpp
+++ b/Engine/Plugins/FX/Niagara/Source/NiagaraEditor/Private/Customizations/NiagaraScriptVariableCustomization.cpp
@@ -114,6 +114,9 @@ void FNiagaraScriptVariableDetails::CustomizeDetails(IDetailLayoutBuilder& Detai
 		return;
 	}
 
+	// We hide bSuggested by default unless it is part of the a parameter definition
+	DetailBuilder.HideProperty(GET_MEMBER_NAME_CHECKED(UNiagaraScriptVariable, Metadata.bSuggested));
+	
 	const bool bInLibrary = Variable->GetOuter()->IsA<UNiagaraParameterDefinitions>();
 	if (Variable->GetIsStaticSwitch())
 	{
@@ -224,6 +227,7 @@ void FNiagaraScriptVariableDetails::CustomizeDetailsParameterDefinitionsSynchron
 	UObject* VariableOuter = Variable->GetOuter();
 	const bool bInLibraryAsset = VariableOuter->IsA<UNiagaraParameterDefinitions>();
 
+	TSharedPtr<IPropertyHandle> SuggestedPropertyHandle = DetailBuilder.GetProperty(GET_MEMBER_NAME_CHECKED(UNiagaraScriptVariable, Metadata.bSuggested));
 	/**
 	 *  Default mode is handled as a special case for UNiagaraScriptVariables synchronized with Parameter Libraries. 
 	 * 	Hide the actual default mode and instead edit a library default mode in place which maps to the actual default mode and whether or not to override the library default.
@@ -231,7 +235,9 @@ void FNiagaraScriptVariableDetails::CustomizeDetailsParameterDefinitionsSynchron
 	const TSharedPtr<IPropertyHandle> DefaultModeHandle = DetailBuilder.GetProperty(GET_MEMBER_NAME_CHECKED(UNiagaraScriptVariable, DefaultMode), UNiagaraScriptVariable::StaticClass());
 	DetailBuilder.HideProperty(DefaultModeHandle);
 	if (bInLibraryAsset)
-	{ 
+	{
+		DetailBuilder.AddPropertyToCategory(SuggestedPropertyHandle);
+		
 		LibraryDefaultModeValue = GetLibrarySourcedDefaultModeInitialValue();
 		FDetailWidgetRow& DefaultModeWidget = CategoryBuilder.AddCustomRow(LOCTEXT("DefaultModeFilterText", "Default Mode"));
 		DefaultModeWidget.NameContent()
diff --git a/Engine/Plugins/FX/Niagara/Source/NiagaraEditor/Private/EdGraphSchema_Niagara.cpp b/Engine/Plugins/FX/Niagara/Source/NiagaraEditor/Private/EdGraphSchema_Niagara.cpp
--- a/Engine/Plugins/FX/Niagara/Source/NiagaraEditor/Private/EdGraphSchema_Niagara.cpp
+++ b/Engine/Plugins/FX/Niagara/Source/NiagaraEditor/Private/EdGraphSchema_Niagara.cpp
@@ -48,6 +48,7 @@
 #include "Textures/SlateIcon.h"
 #include "UObject/LinkerLoad.h"
 #include "Algo/Sort.h"
+#include "Config/NiagaraFavoriteActionsConfig.h"
 
 #include UE_INLINE_GENERATED_CPP_BY_NAME(EdGraphSchema_Niagara)
 
@@ -221,15 +222,15 @@ TSharedPtr<FNiagaraAction_NewNode> AddNewNodeMenuAction(
 	TArray<TSharedPtr<FNiagaraAction_NewNode>>& NewActions,
 	UEdGraphNode* InNodeTemplate,
 	const FText& DisplayName,
-	ENiagaraMenuSections Section,
 	TArray<FString> NestedCategories,
+	TOptional<FNiagaraFavoritesActionData> InFavoritesActionData,
 	const FText& Tooltip,
 	FText Keywords = FText(),
 	FNiagaraActionSourceData SourceData = FNiagaraActionSourceData(
 		EScriptSource::Niagara, FText::FromString(TEXT("Niagara")), true))
 {
 	const UNiagaraEditorSettings* NiagaraEditorSettings = GetDefault<UNiagaraEditorSettings>();
-	TSharedPtr<FNiagaraAction_NewNode> NewAction = MakeShared<FNiagaraAction_NewNode>(DisplayName, Section, NestedCategories, Tooltip, Keywords);
+	TSharedPtr<FNiagaraAction_NewNode> NewAction = MakeShared<FNiagaraAction_NewNode>(DisplayName, NestedCategories, InFavoritesActionData, Tooltip, Keywords);
 	if (ensureMsgf(InNodeTemplate == nullptr || NiagaraEditorSettings->IsVisibleClass(InNodeTemplate->GetClass()),
 		TEXT("Can not create a menu action for a node of class %s."), *InNodeTemplate->GetClass()->GetName()))
 	{
@@ -350,7 +351,7 @@ void AddParametersForGraph(TArray<TSharedPtr<FNiagaraAction_NewNode>>& NewAction
 			}
 
 			UNiagaraNodeInput* InputNodeTemplate = NewObject<UNiagaraNodeInput>(OwnerOfTemporaries);
-			AddNewNodeMenuAction(NewActions, InputNodeTemplate, DisplayName, ENiagaraMenuSections::General, {GraphParameterCategory.ToString()}, FText::GetEmpty());
+			AddNewNodeMenuAction(NewActions, InputNodeTemplate, DisplayName, {GraphParameterCategory.ToString()}, {}, FText::GetEmpty());
 			InputNodeTemplate->Input = InputNode->Input;
 			InputNodeTemplate->Usage = InputNode->Usage;
 			InputNodeTemplate->ExposureOptions = InputNode->ExposureOptions;
@@ -485,10 +486,13 @@ TArray<TSharedPtr<FNiagaraAction_NewNode>> UEdGraphSchema_Niagara::GetGraphActio
 
 			if(bAddOp)
 			{
-				// todo suggestion info per op?
+				FNiagaraFavoritesActionData FavoritesActionData;
+				FavoritesActionData.ActionIdentifier.Names.Add(OpInfo.Name);
+				FavoritesActionData.bFavoriteByDefault = false;
+				
 				UNiagaraNodeOp* OpNode = NewObject<UNiagaraNodeOp>(OwnerOfTemporaries);
 				OpNode->OpName = OpInfo.Name;
-				TSharedPtr<FNiagaraAction_NewNode> MathOpAction = AddNewNodeMenuAction(NewActions, OpNode, OpInfo.FriendlyName, ENiagaraMenuSections::General, {OpInfo.Category.ToString()}, OpInfo.Description,  OpInfo.Keywords);
+				TSharedPtr<FNiagaraAction_NewNode> MathOpAction = AddNewNodeMenuAction(NewActions, OpNode, OpInfo.FriendlyName, {OpInfo.Category.ToString()}, FavoritesActionData, OpInfo.Description,  OpInfo.Keywords);
 				MathOpAction->AlternateSearchName = OpInfo.AlternateSearchName.IsSet() ? FText::FromName(OpInfo.AlternateSearchName.GetValue()) : TOptional<FText>();
 			}
 		}
@@ -502,7 +506,7 @@ TArray<TSharedPtr<FNiagaraAction_NewNode>> UEdGraphSchema_Niagara::GetGraphActio
 		
 		UNiagaraNodeCustomHlsl* CustomHlslNode = NewObject<UNiagaraNodeCustomHlsl>(OwnerOfTemporaries);
 		CustomHlslNode->SetCustomHlsl(TEXT("// Insert the body of the function here and add any inputs\r\n// and outputs by name using the add pins above.\r\n// Currently, complicated branches, for loops, switches, etc are not advised."));
-		AddNewNodeMenuAction(NewActions, CustomHlslNode, DisplayName, ENiagaraMenuSections::General, {LOCTEXT("Function Menu Title", "Functions").ToString()}, TooltipDesc,  FText::GetEmpty());
+		AddNewNodeMenuAction(NewActions, CustomHlslNode, DisplayName, {LOCTEXT("Function Menu Title", "Functions").ToString()}, {}, TooltipDesc,  FText::GetEmpty());
 	}
 
 	auto AddScriptFunctionAction = [&NewActions, OwnerOfTemporaries](const TArray<FString>& Categories, const FAssetData& ScriptAsset)
@@ -521,12 +525,14 @@ TArray<TSharedPtr<FNiagaraAction_NewNode>> UEdGraphSchema_Niagara::GetGraphActio
 		const TTuple<EScriptSource, FText> Source = FNiagaraEditorUtilities::GetScriptSource(ScriptAsset);
 		FNiagaraActionSourceData SourceData(Source.Key, Source.Value, true);
 		
-		const ENiagaraMenuSections Section = bSuggested ? ENiagaraMenuSections::Suggested: ENiagaraMenuSections::General;
-
+		FNiagaraFavoritesActionData FavoritesActionData;
+		FavoritesActionData.ActionIdentifier.Names.Add(FName(ScriptAsset.GetObjectPathString()));
+		FavoritesActionData.bFavoriteByDefault = bSuggested;
+		
 		UNiagaraNodeFunctionCall* FunctionCallNode = NewObject<UNiagaraNodeFunctionCall>(OwnerOfTemporaries);
 		FunctionCallNode->FunctionScriptAssetObjectPath = FName(ScriptAsset.GetSoftObjectPath().ToString());
 
-		TSharedPtr<FNiagaraAction_NewNode> Action = AddNewNodeMenuAction(NewActions, FunctionCallNode, DisplayName, Section, Categories, TooltipDesc, Keywords, SourceData);
+		TSharedPtr<FNiagaraAction_NewNode> Action = AddNewNodeMenuAction(NewActions, FunctionCallNode, DisplayName, Categories, FavoritesActionData, TooltipDesc, Keywords, SourceData);
 		Action->bIsInLibrary = bIsInLibrary;
 
 		return Action;
@@ -556,7 +562,7 @@ TArray<TSharedPtr<FNiagaraAction_NewNode>> UEdGraphSchema_Niagara::GetGraphActio
 			FunctionCallNode->Signature = Sig;
 
 			const FText MenuDesc = FText::FromString(FName::NameToDisplayString(Sig.Name.ToString(), false));
-			TSharedPtr<FNiagaraAction_NewNode> FunctionCallAction = AddNewNodeMenuAction(NewActions, FunctionCallNode, MenuDesc, ENiagaraMenuSections::General, TArray<FString>(), MenuDesc);
+			TSharedPtr<FNiagaraAction_NewNode> FunctionCallAction = AddNewNodeMenuAction(NewActions, FunctionCallNode, MenuDesc, TArray<FString>(), {},  MenuDesc);
 			FunctionCallAction->WeakNodeTemplate = FunctionCallNode;
 		}
 	}
@@ -607,7 +613,7 @@ TArray<TSharedPtr<FNiagaraAction_NewNode>> UEdGraphSchema_Niagara::GetGraphActio
 					UNiagaraNodeReadDataSet* EventReadNode = NewObject<UNiagaraNodeReadDataSet>(OwnerOfTemporaries);
 					EventReadNode->InitializeFromStruct(PayloadType.GetStruct());
 
-					AddNewNodeMenuAction(NewActions, EventReadNode, DisplayName, ENiagaraMenuSections::General, ReadCategories, FText::GetEmpty(), FText::GetEmpty());
+					AddNewNodeMenuAction(NewActions, EventReadNode, DisplayName, ReadCategories, {}, FText::GetEmpty(), FText::GetEmpty());
 				}
 				{
 					const FText DisplayNameFormat = LOCTEXT("AddEventWriteFmt", "Add {0} Event Write");
@@ -618,7 +624,7 @@ TArray<TSharedPtr<FNiagaraAction_NewNode>> UEdGraphSchema_Niagara::GetGraphActio
 					UNiagaraNodeWriteDataSet* EventWriteNode = NewObject<UNiagaraNodeWriteDataSet>(OwnerOfTemporaries);
 					EventWriteNode->InitializeFromStruct(PayloadType.GetStruct());
 					
-					AddNewNodeMenuAction(NewActions, EventWriteNode, DisplayName, ENiagaraMenuSections::General, WriteCategories, FText::GetEmpty(), FText::GetEmpty());
+					AddNewNodeMenuAction(NewActions, EventWriteNode, DisplayName,  WriteCategories, {}, FText::GetEmpty(), FText::GetEmpty());
 				}
 			}
 		}
@@ -649,7 +655,7 @@ TArray<TSharedPtr<FNiagaraAction_NewNode>> UEdGraphSchema_Niagara::GetGraphActio
 
 			FText DisplayName = FText::Format(LOCTEXT("AddOutput", "Add {0}"), OutputNode->GetNodeTitle(ENodeTitleType::FullTitle));
 
-			AddNewNodeMenuAction(NewActions, OutputNode, DisplayName, ENiagaraMenuSections::General, {MenuCat.ToString()}, FText::GetEmpty(), FText::GetEmpty());
+			AddNewNodeMenuAction(NewActions, OutputNode, DisplayName,  {MenuCat.ToString()}, {}, FText::GetEmpty(), FText::GetEmpty());
 
 			UNiagaraNodeOutput* UpdateOutputNode = NiagaraGraph->FindOutputNode(ENiagaraScriptUsage::ParticleUpdateScript);
 			if (UpdateOutputNode)
@@ -687,7 +693,7 @@ TArray<TSharedPtr<FNiagaraAction_NewNode>> UEdGraphSchema_Niagara::GetGraphActio
 				FText Desc = FText::Format(DescFmt, DisplayName);
 				
 				UNiagaraNodeConvert* ConvertNode = NewObject<UNiagaraNodeConvert>(OwnerOfTemporaries);
-				TSharedPtr<FNiagaraAction_NewNode> Action = AddNewNodeMenuAction(NewActions, ConvertNode, DisplayName, ENiagaraMenuSections::General, {bMake ? MakeCat.ToString() : BreakCat.ToString()}, FText::GetEmpty(), FText::GetEmpty());
+				TSharedPtr<FNiagaraAction_NewNode> Action = AddNewNodeMenuAction(NewActions, ConvertNode, DisplayName,  {bMake ? MakeCat.ToString() : BreakCat.ToString()}, {}, FText::GetEmpty(), FText::GetEmpty());
 				if (bMake)
 				{
 					ConvertNode->InitAsMake(Type);
@@ -792,7 +798,7 @@ TArray<TSharedPtr<FNiagaraAction_NewNode>> UEdGraphSchema_Niagara::GetGraphActio
 			FText Desc = LOCTEXT("NiagaraConvert", "Convert");
 			
 			UNiagaraNodeConvert* ConvertNode = NewObject<UNiagaraNodeConvert>(OwnerOfTemporaries);
-			AddNewNodeMenuAction(NewActions, ConvertNode, Desc, ENiagaraMenuSections::General, {}, FText::GetEmpty(), FText::GetEmpty());
+			AddNewNodeMenuAction(NewActions, ConvertNode, Desc,  {}, {},FText::GetEmpty(), FText::GetEmpty());
 		}
 	}
 
@@ -810,7 +816,7 @@ TArray<TSharedPtr<FNiagaraAction_NewNode>> UEdGraphSchema_Niagara::GetGraphActio
 					continue;
 
 				UNiagaraNodeFunctionCall* FuncNode = NewObject<UNiagaraNodeFunctionCall>(OwnerOfTemporaries);
-				AddNewNodeMenuAction(NewActions, FuncNode, FText::FromString(FName::NameToDisplayString(Sig.GetNameString(), false)), ENiagaraMenuSections::General, {MenuCat.ToString()}, Sig.GetDescription(), FText::GetEmpty());
+				AddNewNodeMenuAction(NewActions, FuncNode, FText::FromString(FName::NameToDisplayString(Sig.GetNameString(), false)),  {MenuCat.ToString()}, {}, Sig.GetDescription(), FText::GetEmpty());
 				FuncNode->Signature = Sig;
 			}
 		}
@@ -847,7 +853,7 @@ TArray<TSharedPtr<FNiagaraAction_NewNode>> UEdGraphSchema_Niagara::GetGraphActio
 					const FText Category = LOCTEXT("NiagaraSwizzles", "Swizzles");
 					
 					UNiagaraNodeConvert* ConvertNode = NewObject<UNiagaraNodeConvert>(OwnerOfTemporaries);
-					AddNewNodeMenuAction(NewActions, ConvertNode, FText::FromString(Swiz), ENiagaraMenuSections::General, {Category.ToString()}, FText::GetEmpty(), FText::GetEmpty());
+					AddNewNodeMenuAction(NewActions, ConvertNode, FText::FromString(Swiz),  {Category.ToString()}, {}, FText::GetEmpty(), FText::GetEmpty());
 
 					ConvertNode->InitAsSwizzle(Swiz);
 				}
@@ -862,32 +868,42 @@ TArray<TSharedPtr<FNiagaraAction_NewNode>> UEdGraphSchema_Niagara::GetGraphActio
 		if (NiagaraEditorSettings->IsVisibleClass(UNiagaraNodeParameterMapGet::StaticClass()))
 		{
 			FString Name = TEXT("Parameter Map Get");
+			
+			FNiagaraFavoritesActionData FavoritesActionData;
+			FavoritesActionData.ActionIdentifier.Names.Add(FName(Name));
+			FavoritesActionData.bFavoriteByDefault = true;
+			
 			UNiagaraNodeParameterMapGet* BaseNode = NewObject<UNiagaraNodeParameterMapGet>(OwnerOfTemporaries);
-			AddNewNodeMenuAction(NewActions, BaseNode, FText::FromString(Name), ENiagaraMenuSections::Suggested, {MenuCat.ToString()}, FText::GetEmpty(), FText::GetEmpty());
+			AddNewNodeMenuAction(NewActions, BaseNode, FText::FromString(Name), {MenuCat.ToString()}, FavoritesActionData, FText::GetEmpty(), FText::GetEmpty());
 		}
 		if (NiagaraEditorSettings->IsVisibleClass(UNiagaraNodeParameterMapSet::StaticClass()))
 		{
 			FString Name = TEXT("Parameter Map Set");
+			
+			FNiagaraFavoritesActionData FavoritesActionData;
+			FavoritesActionData.ActionIdentifier.Names.Add(FName(Name));
+			FavoritesActionData.bFavoriteByDefault = true;
+			
 			UNiagaraNodeParameterMapSet* BaseNode = NewObject<UNiagaraNodeParameterMapSet>(OwnerOfTemporaries);
-			AddNewNodeMenuAction(NewActions, BaseNode, FText::FromString(Name), ENiagaraMenuSections::Suggested, {MenuCat.ToString()}, FText::GetEmpty(), FText::GetEmpty());
+			AddNewNodeMenuAction(NewActions, BaseNode, FText::FromString(Name), {MenuCat.ToString()}, FavoritesActionData, FText::GetEmpty(), FText::GetEmpty());
 		}
 		if (NiagaraEditorSettings->IsVisibleClass(UNiagaraNodeParameterMapFor::StaticClass()))
 		{
 			FString Name = TEXT("Parameter Map For");
 			UNiagaraNodeParameterMapFor* BaseNode = NewObject<UNiagaraNodeParameterMapFor>(OwnerOfTemporaries);
-			AddNewNodeMenuAction(NewActions, BaseNode, FText::FromString(Name), ENiagaraMenuSections::General, {MenuCat.ToString()}, FText::GetEmpty(), FText::GetEmpty());
+			AddNewNodeMenuAction(NewActions, BaseNode, FText::FromString(Name),  {MenuCat.ToString()}, {}, FText::GetEmpty(), FText::GetEmpty());
 		}
 		if (NiagaraEditorSettings->IsVisibleClass(UNiagaraNodeParameterMapForWithContinue::StaticClass()))
 		{
 			FString Name = TEXT("Parameter Map For With Continue");
 			UNiagaraNodeParameterMapForWithContinue* BaseNode = NewObject<UNiagaraNodeParameterMapForWithContinue>(OwnerOfTemporaries);
-			AddNewNodeMenuAction(NewActions, BaseNode, FText::FromString(Name), ENiagaraMenuSections::General, { MenuCat.ToString() }, FText::GetEmpty(), FText::GetEmpty());
+			AddNewNodeMenuAction(NewActions, BaseNode, FText::FromString(Name),  { MenuCat.ToString() }, {}, FText::GetEmpty(), FText::GetEmpty());
 		}
 		if (NiagaraEditorSettings->IsVisibleClass(UNiagaraNodeParameterMapForIndex::StaticClass()))
 		{
 			FString Name = TEXT("Parameter Map For Current Index");
 			UNiagaraNodeParameterMapForIndex* BaseNode = NewObject<UNiagaraNodeParameterMapForIndex>(OwnerOfTemporaries);
-			AddNewNodeMenuAction(NewActions, BaseNode, FText::FromString(Name), ENiagaraMenuSections::General, { MenuCat.ToString() }, FText::GetEmpty(), FText::GetEmpty());
+			AddNewNodeMenuAction(NewActions, BaseNode, FText::FromString(Name),  { MenuCat.ToString() }, {}, FText::GetEmpty(), FText::GetEmpty());
 		}
 	}
 
@@ -899,7 +915,7 @@ TArray<TSharedPtr<FNiagaraAction_NewNode>> UEdGraphSchema_Niagara::GetGraphActio
 		{
 			FString Name = TEXT("Add Comment");
 			UEdGraphNode_Comment* BaseNode = NewObject<UEdGraphNode_Comment>(OwnerOfTemporaries);
-			AddNewNodeMenuAction(NewActions, BaseNode, FText::FromString(Name), ENiagaraMenuSections::General, {MenuCat.ToString()}, FText::GetEmpty(), FText::GetEmpty());
+			AddNewNodeMenuAction(NewActions, BaseNode, FText::FromString(Name),  {MenuCat.ToString()}, {}, FText::GetEmpty(), FText::GetEmpty());
 		}		
 	}
 
@@ -911,7 +927,7 @@ TArray<TSharedPtr<FNiagaraAction_NewNode>> UEdGraphSchema_Niagara::GetGraphActio
 		{
 			FString Name = TEXT("Add Compiler Output Tag");
 			UNiagaraNodeOutputTag* BaseNode = NewObject<UNiagaraNodeOutputTag>(OwnerOfTemporaries);
-			AddNewNodeMenuAction(NewActions, BaseNode, FText::FromString(Name), ENiagaraMenuSections::General, {MenuCat.ToString()}, FText::GetEmpty(), FText::GetEmpty());
+			AddNewNodeMenuAction(NewActions, BaseNode, FText::FromString(Name),  {MenuCat.ToString()}, {}, FText::GetEmpty(), FText::GetEmpty());
 		}
 	}
 	
@@ -932,7 +948,7 @@ TArray<TSharedPtr<FNiagaraAction_NewNode>> UEdGraphSchema_Niagara::GetGraphActio
 				const FText DisplayName = FText::Format(LOCTEXT("GetSystemConstant", "Get {Constant}"), Args);
 				
 				UNiagaraNodeInput* InputNode = NewObject<UNiagaraNodeInput>(OwnerOfTemporaries);
-				AddNewNodeMenuAction(NewActions, InputNode, DisplayName, ENiagaraMenuSections::General, {LOCTEXT("System Parameters Menu Title", "System Parameters").ToString()}, FText::GetEmpty(), FText::GetEmpty());
+				AddNewNodeMenuAction(NewActions, InputNode, DisplayName,  {LOCTEXT("System Parameters Menu Title", "System Parameters").ToString()}, {}, FText::GetEmpty(), FText::GetEmpty());
 
 				InputNode->Usage = ENiagaraInputNodeUsage::SystemConstant;
 				InputNode->Input = SysConst;
@@ -948,7 +964,7 @@ TArray<TSharedPtr<FNiagaraAction_NewNode>> UEdGraphSchema_Niagara::GetGraphActio
 			const FText DisplayName = FText::Format(LOCTEXT("GetTranslatorConstant", "{Constant}"), Args);
 			
 			UNiagaraNodeInput* InputNode = NewObject<UNiagaraNodeInput>(OwnerOfTemporaries);
-			AddNewNodeMenuAction(NewActions, InputNode, DisplayName, ENiagaraMenuSections::General, {LOCTEXT("Translator Parameters Menu Title", "Special Purpose Parameters").ToString()}, FText::GetEmpty(), FText::GetEmpty());
+			AddNewNodeMenuAction(NewActions, InputNode, DisplayName,  {LOCTEXT("Translator Parameters Menu Title", "Special Purpose Parameters").ToString()}, {}, FText::GetEmpty(), FText::GetEmpty());
 
 			InputNode->Usage = ENiagaraInputNodeUsage::TranslatorConstant;
 			InputNode->ExposureOptions.bCanAutoBind = true;
@@ -1005,7 +1021,7 @@ TArray<TSharedPtr<FNiagaraAction_NewNode>> UEdGraphSchema_Niagara::GetGraphActio
 					const FText DisplayName = FText::Format(MenuDescFmt, Type.GetNameText());
 					
 					UNiagaraNodeInput* InputNode = NewObject<UNiagaraNodeInput>(OwnerOfTemporaries);
-					AddNewNodeMenuAction(NewActions, InputNode, DisplayName, ENiagaraMenuSections::General, Categories, FText::GetEmpty(), FText::GetEmpty());
+					AddNewNodeMenuAction(NewActions, InputNode, DisplayName,  Categories, {}, FText::GetEmpty(), FText::GetEmpty());
 					FNiagaraEditorUtilities::InitializeParameterInputNode(*InputNode, Type, NiagaraGraph);
 				}
 
@@ -1016,7 +1032,7 @@ TArray<TSharedPtr<FNiagaraAction_NewNode>> UEdGraphSchema_Niagara::GetGraphActio
 					const FText DisplayName = FText::Format(MenuDescFmt, PinType.GetNameText());
 					
 					UNiagaraNodeInput* InputNode = NewObject<UNiagaraNodeInput>(OwnerOfTemporaries);
-					AddNewNodeMenuAction(NewActions, InputNode, DisplayName, ENiagaraMenuSections::General, {}, FText::GetEmpty(), FText::GetEmpty());
+					AddNewNodeMenuAction(NewActions, InputNode, DisplayName,  {}, {}, FText::GetEmpty(), FText::GetEmpty());
 					FNiagaraEditorUtilities::InitializeParameterInputNode(*InputNode, PinType, NiagaraGraph);
 				}
 			}
@@ -1031,7 +1047,7 @@ TArray<TSharedPtr<FNiagaraAction_NewNode>> UEdGraphSchema_Niagara::GetGraphActio
 		const FText RerouteMenuDesc = LOCTEXT("NiagaraRerouteMenuDesc", "Reroute");
 		
 		UNiagaraNodeReroute* RerouteNode = NewObject<UNiagaraNodeReroute>(OwnerOfTemporaries);
-		AddNewNodeMenuAction(NewActions, RerouteNode, RerouteMenuDesc, ENiagaraMenuSections::General, {UtilMenuCat.ToString()}, FText::GetEmpty(), FText::GetEmpty());
+		AddNewNodeMenuAction(NewActions, RerouteNode, RerouteMenuDesc,  {UtilMenuCat.ToString()}, {}, FText::GetEmpty(), FText::GetEmpty());
 	}
 	
 	// Add select node
@@ -1039,9 +1055,13 @@ TArray<TSharedPtr<FNiagaraAction_NewNode>> UEdGraphSchema_Niagara::GetGraphActio
 	if (bAllowSelectNodes && NiagaraEditorSettings->IsVisibleClass(UNiagaraNodeSelect::StaticClass()))
 	{
 		const FText SelectMenuDesc = LOCTEXT("NiagaraSelectMenuDesc", "Select / If");
+
+		FNiagaraFavoritesActionData FavoritesActionData;
+		FavoritesActionData.ActionIdentifier.Names.Add(FName(FText::AsCultureInvariant(SelectMenuDesc).ToString()));
+		FavoritesActionData.bFavoriteByDefault = true;
 		
 		UNiagaraNodeSelect* Node = NewObject<UNiagaraNodeSelect>(OwnerOfTemporaries);
-		AddNewNodeMenuAction(NewActions, Node, SelectMenuDesc, ENiagaraMenuSections::Suggested, {UtilMenuCat.ToString()}, FText::GetEmpty(), FText::FromString(TEXT("If Branch Bool Select")));
+		AddNewNodeMenuAction(NewActions, Node, SelectMenuDesc,  {UtilMenuCat.ToString()}, FavoritesActionData, FText::GetEmpty(), FText::FromString(TEXT("If Branch Bool Select")));
 	}
 
 	// Add static switch node
@@ -1049,11 +1069,15 @@ TArray<TSharedPtr<FNiagaraAction_NewNode>> UEdGraphSchema_Niagara::GetGraphActio
 	if (bAllowStaticSwitchNodes && NiagaraEditorSettings->IsVisibleClass(UNiagaraNodeStaticSwitch::StaticClass()))
 	{
 		const FText UsageSelectorMenuDesc = LOCTEXT("NiagaraStaticSwitchMenuDesc", "Static Switch");
+
+		FNiagaraFavoritesActionData FavoritesActionData;
+		FavoritesActionData.ActionIdentifier.Names.Add(FName(FText::AsCultureInvariant(UsageSelectorMenuDesc).ToString()));
+		FavoritesActionData.bFavoriteByDefault = true;
 		
 		UNiagaraNodeStaticSwitch* Node = NewObject<UNiagaraNodeStaticSwitch>(OwnerOfTemporaries);
 		// new nodes should auto refresh
 		Node->SwitchTypeData.bAutoRefreshEnabled = true;
-		AddNewNodeMenuAction(NewActions, Node, UsageSelectorMenuDesc, ENiagaraMenuSections::Suggested, {UtilMenuCat.ToString()}, FText::GetEmpty(), FText::FromString(TEXT("")));
+		AddNewNodeMenuAction(NewActions, Node, UsageSelectorMenuDesc,  {UtilMenuCat.ToString()}, FavoritesActionData, FText::GetEmpty(), FText::FromString(TEXT("")));
 	}
 
 	return NewActions;
diff --git a/Engine/Plugins/FX/Niagara/Source/NiagaraEditor/Private/NiagaraActions.cpp b/Engine/Plugins/FX/Niagara/Source/NiagaraEditor/Private/NiagaraActions.cpp
--- a/Engine/Plugins/FX/Niagara/Source/NiagaraEditor/Private/NiagaraActions.cpp
+++ b/Engine/Plugins/FX/Niagara/Source/NiagaraEditor/Private/NiagaraActions.cpp
@@ -24,6 +24,7 @@
 #include "Widgets/SToolTip.h"
 
 #include "NiagaraDataChannel.h"
+#include "Config/NiagaraFavoriteActionsConfig.h"
 
 #include UE_INLINE_GENERATED_CPP_BY_NAME(NiagaraActions)
 
@@ -102,16 +103,23 @@ void FNiagaraMenuActionCollector::AddAllActionsTo(FGraphActionListBuilderBase& A
 	}
 }
 
-FNiagaraMenuAction_Base::FNiagaraMenuAction_Base(FText InDisplayName, ENiagaraMenuSections InSection,
-                                                 TArray<FString> InNodeCategories, FText InToolTip, FText InKeywords, float InIntrinsicWeightMultiplier)
+FNiagaraMenuAction_Base::FNiagaraMenuAction_Base(FText InDisplayName, TArray<FString> InNodeCategories, TOptional<FNiagaraFavoritesActionData> InFavoritesActionData, FText InToolTip, FText InKeywords, float InIntrinsicWeightMultiplier)
 {
 	DisplayName = InDisplayName;
-	Section = InSection;
 	Categories = InNodeCategories;
+	FavoritesActionData = InFavoritesActionData;
 	ToolTip = InToolTip;
 	Keywords = InKeywords;
 	SearchWeightMultiplier = InIntrinsicWeightMultiplier;
 
+	if(FavoritesActionData.IsSet() == false)
+	{
+		FNiagaraFavoritesActionData DefaultFavoritesActionData;
+		DefaultFavoritesActionData.ActionIdentifier.Names.Add(FName(DisplayName.ToString()));
+		DefaultFavoritesActionData.bFavoriteByDefault = false;
+		FavoritesActionData = DefaultFavoritesActionData;
+	}
+	
 	UpdateFullSearchText();
 }
 
diff --git a/Engine/Plugins/FX/Niagara/Source/NiagaraEditor/Private/NiagaraEditorModule.cpp b/Engine/Plugins/FX/Niagara/Source/NiagaraEditor/Private/NiagaraEditorModule.cpp
--- a/Engine/Plugins/FX/Niagara/Source/NiagaraEditor/Private/NiagaraEditorModule.cpp
+++ b/Engine/Plugins/FX/Niagara/Source/NiagaraEditor/Private/NiagaraEditorModule.cpp
@@ -1100,6 +1100,12 @@ void FNiagaraEditorModule::StartupModule()
 
 	UNiagaraAssetBrowserConfig::Initialize();
 
+	// We have to call shutdown before the exit actually happens to be able to save to disk
+	FEditorDelegates::OnEditorPreExit.AddLambda([]()
+	{
+		UNiagaraFavoriteActionsConfig::Shutdown();
+	});
+	
 	FNiagaraAssetDetailDatabase::Init();
 
 	TempPackage = NewObject<UPackage>(nullptr, TEXT("/Temp/NiagaraEditor"), RF_Transient);
@@ -1656,7 +1662,7 @@ void FNiagaraEditorModule::ShutdownModule()
 	FNiagaraEditorStyle::Unregister();
 
 	UnregisterSettings();
-
+	
 	if (UObjectInitialized())
 	{
 		FComponentAssetBrokerage::UnregisterBroker(NiagaraComponentBroker);
diff --git a/Engine/Plugins/FX/Niagara/Source/NiagaraEditor/Private/NiagaraEditorUtilities.cpp b/Engine/Plugins/FX/Niagara/Source/NiagaraEditor/Private/NiagaraEditorUtilities.cpp
--- a/Engine/Plugins/FX/Niagara/Source/NiagaraEditor/Private/NiagaraEditorUtilities.cpp
+++ b/Engine/Plugins/FX/Niagara/Source/NiagaraEditor/Private/NiagaraEditorUtilities.cpp
@@ -1655,7 +1655,7 @@ bool FNiagaraEditorUtilities::IsEnginePluginAsset(const FTopLevelAssetPath& InTo
 		FPaths::IsUnderDirectory(PackagePathLocal, FPaths::EnginePluginsDir() + "FX/Niagara");
 }
 
-int32 FNiagaraEditorUtilities::GetWeightForItem(const TSharedPtr<FNiagaraMenuAction_Generic>& InCurrentAction, const TArray<FString>& InFilterTerms)
+int32 FNiagaraEditorUtilities::GetWeightForItem(const TSharedPtr<FNiagaraMenuAction_Generic>& InCurrentAction, FName FavoriteActionsProfile, const TArray<FString>& InFilterTerms)
 {
 	// The overall 'weight'
 	int32 TotalWeight = 0;
@@ -1881,11 +1881,14 @@ int32 FNiagaraEditorUtilities::GetWeightForItem(const TSharedPtr<FNiagaraMenuAct
 		}
 	}
 
-	// suggested actions get favored massively
-	if(InCurrentAction->Section == ENiagaraMenuSections::Suggested)
+	// favorite actions get favored massively
+	if(InCurrentAction->FavoritesActionData.IsSet())
 	{
-		//TotalWeight *= TotalWeight;
-		TotalWeight += (int32)FMath::Pow(10.0f, 4.0f);
+		FNiagaraFavoritesActionData ActionData = InCurrentAction->FavoritesActionData.GetValue();
+		if(UNiagaraFavoriteActionsConfig::Get()->GetActionsProfile(FavoriteActionsProfile).IsFavorite(ActionData))
+		{
+			TotalWeight += (int32)FMath::Pow(10.0f, 4.0f);
+		}
 	}
 	
 	return TotalWeight > 0.f ? TotalWeight : INDEX_NONE;
@@ -3128,6 +3131,19 @@ void FNiagaraEditorUtilities::GetAvailableParameterCollections(TArray<UNiagaraPa
 	}
 }
 
+void FNiagaraEditorUtilities::GetAvailableParameterDefinitions(TArray<UNiagaraParameterDefinitions*>& OutParameterDefinitions)
+{
+	const TArray<TWeakObjectPtr<UNiagaraParameterDefinitions>>& CachedParameterDefinitionAssets = FNiagaraEditorModule::Get().GetCachedParameterDefinitionsAssets();
+
+	for (const TWeakObjectPtr<UNiagaraParameterDefinitions>& CachedParameterDefinitionAsset : CachedParameterDefinitionAssets)
+	{
+		if (UNiagaraParameterDefinitions* ParameterDefinition = CachedParameterDefinitionAsset.Get())
+		{
+			OutParameterDefinitions.Add(ParameterDefinition);
+		}
+	}
+}
+
 TSharedPtr<INiagaraParameterDefinitionsSubscriberViewModel> FNiagaraEditorUtilities::GetOwningLibrarySubscriberViewModelForGraph(const UNiagaraGraph* Graph)
 {
 	if (UNiagaraScript* BaseScript = Graph->GetTypedOuter<UNiagaraScript>())
diff --git a/Engine/Plugins/FX/Niagara/Source/NiagaraEditor/Private/ViewModels/Stack/NiagaraStackFunctionInput.cpp b/Engine/Plugins/FX/Niagara/Source/NiagaraEditor/Private/ViewModels/Stack/NiagaraStackFunctionInput.cpp
--- a/Engine/Plugins/FX/Niagara/Source/NiagaraEditor/Private/ViewModels/Stack/NiagaraStackFunctionInput.cpp
+++ b/Engine/Plugins/FX/Niagara/Source/NiagaraEditor/Private/ViewModels/Stack/NiagaraStackFunctionInput.cpp
@@ -31,6 +31,7 @@
 #include "NiagaraScriptMergeManager.h"
 #include "NiagaraScriptSource.h"
 #include "NiagaraSettings.h"
+#include "NiagaraTrace.h"
 #include "ViewModels/NiagaraSystemScriptViewModel.h"
 #include "ScopedTransaction.h"
 #include "Algo/RemoveIf.h"
@@ -60,6 +61,17 @@
 
 static FText TooManyConversionScripts = LOCTEXT("TooManyConversionScripts", "There is more than one dynamic input script available auto-convert the dragged parameter. Please fix this by disabling conversion for all but one of them:\n{0}");
 
+UNiagaraStackFunctionInput::FNiagaraAvailableParameterInfo::FNiagaraAvailableParameterInfo(FNiagaraVariableBase InVariableBase)
+{
+	Variable = InVariableBase;
+}
+
+UNiagaraStackFunctionInput::FNiagaraAvailableParameterInfo::FNiagaraAvailableParameterInfo(FNiagaraVariableBase InVariableBase,	UNiagaraScript* InConversionScript)
+{
+	Variable = InVariableBase;
+	ConversionScript = InConversionScript;
+}
+
 UNiagaraStackFunctionInput::UNiagaraStackFunctionInput()
 	: OwningModuleNode(nullptr)
 	, OwningFunctionCallNode(nullptr)
@@ -1808,8 +1820,10 @@ UNiagaraScript* UNiagaraStackFunctionInput::FindConversionScript(const FNiagaraT
 	return ConversionScriptCache[FromType];
 }
 
-void UNiagaraStackFunctionInput::GetAvailableParameters(TArray<FNiagaraVariableBase>& OutAvailableParameters, TMap<FNiagaraVariableBase, UNiagaraScript*>& OutAvailableParametersWithConversions, bool bIncludeConversionScripts) const
+void UNiagaraStackFunctionInput::GetAvailableParameters(TSet<FNiagaraAvailableParameterInfo>& OutAvailableParameters, FGetAvailableParameterArgs Arguments) const
 {
+	TRACE_CPUPROFILER_EVENT_SCOPE_ON_CHANNEL(UNiagaraStackFunctionInput::GetAvailableParameters, NiagaraChannel);
+	
 	TMap<FNiagaraTypeDefinition, UNiagaraScript*> ConversionScriptCache;
 	
 	// Engine Handles.
@@ -1819,9 +1833,9 @@ void UNiagaraStackFunctionInput::GetAvailableParameters(TArray<FNiagaraVariableB
 		{
 			OutAvailableParameters.Add(SystemVariable);
 		}
-		else if (UNiagaraScript* ConversionScript = FindConversionScript(SystemVariable.GetType(), ConversionScriptCache, bIncludeConversionScripts))
+		else if (UNiagaraScript* ConversionScript = FindConversionScript(SystemVariable.GetType(), ConversionScriptCache, Arguments.bIncludeConversionScripts))
 		{
-			OutAvailableParametersWithConversions.Add(SystemVariable, ConversionScript);
+			OutAvailableParameters.Add({SystemVariable, ConversionScript});
 		}
 	}
 
@@ -1834,9 +1848,9 @@ void UNiagaraStackFunctionInput::GetAvailableParameters(TArray<FNiagaraVariableB
 		{
 			OutAvailableParameters.Add(ExposedVar);
 		}
-		else if (UNiagaraScript* ConversionScript = FindConversionScript(ExposedVar.GetType(), ConversionScriptCache, bIncludeConversionScripts))
+		else if (UNiagaraScript* ConversionScript = FindConversionScript(ExposedVar.GetType(), ConversionScriptCache, Arguments.bIncludeConversionScripts))
 		{
-			OutAvailableParametersWithConversions.Add(ExposedVar, ConversionScript);
+			OutAvailableParameters.Add({ExposedVar, ConversionScript});
 		}
 	}
 
@@ -1934,7 +1948,7 @@ void UNiagaraStackFunctionInput::GetAvailableParameters(TArray<FNiagaraVariableB
 
 						if (FNiagaraEditorUtilities::AreTypesAssignable(HistoryVariable.GetType(), InputType))
 						{
-							OutAvailableParameters.AddUnique(HistoryVariable);
+							OutAvailableParameters.Add(HistoryVariable);
 							if (AvailableHandle.GetNamespace() == OutputNamespace)
 							{
 								AvailableParametersForThisOutput.Add(TPair<FNiagaraVariableBase, FNiagaraParameterHandle>(HistoryVariable, AvailableHandle));
@@ -1949,13 +1963,13 @@ void UNiagaraStackFunctionInput::GetAvailableParameters(TArray<FNiagaraVariableB
 									FString NewName = HistoryVariable.GetName().ToString().Replace(*StackContextRoots[StackRootIdx].ToString(), *FNiagaraConstants::StackContextNamespace.ToString());
 
 									FNiagaraVariableBase AvailableAliasedParameter = FNiagaraVariableBase(HistoryVariable.GetType(), *NewName);
-									OutAvailableParameters.AddUnique(AvailableAliasedParameter);
+									OutAvailableParameters.Add(AvailableAliasedParameter);
 								}
 							}
 						}
-						else if (UNiagaraScript* ConversionScript = FindConversionScript(HistoryVariable.GetType(), ConversionScriptCache, bIncludeConversionScripts))
+						else if (UNiagaraScript* ConversionScript = FindConversionScript(HistoryVariable.GetType(), ConversionScriptCache, Arguments.bIncludeConversionScripts))
 						{
-							OutAvailableParametersWithConversions.Add(HistoryVariable, ConversionScript);
+							OutAvailableParameters.Add({HistoryVariable, ConversionScript});
 						}
 					}
 				}
@@ -1970,7 +1984,7 @@ void UNiagaraStackFunctionInput::GetAvailableParameters(TArray<FNiagaraVariableB
 					if (AvailableParameter.GetType().IsStatic() == false)
 					{
 						FNiagaraParameterHandle InitialParameterHandle = FNiagaraParameterHandle::CreateInitialParameterHandle(AvailableHandle);
-						OutAvailableParameters.AddUnique(FNiagaraVariableBase(AvailableParameter.GetType(), InitialParameterHandle.GetParameterHandleString()));
+						OutAvailableParameters.Add(FNiagaraVariableBase(AvailableParameter.GetType(), InitialParameterHandle.GetParameterHandleString()));
 					}
 				}
 			}
@@ -1986,14 +2000,39 @@ void UNiagaraStackFunctionInput::GetAvailableParameters(TArray<FNiagaraVariableB
 		{
 			if (FNiagaraEditorUtilities::AreTypesAssignable(CollectionParam.GetType(), InputType))
 			{
-				OutAvailableParameters.AddUnique(CollectionParam);
+				OutAvailableParameters.Add(CollectionParam);
 			}
-			else if (UNiagaraScript* ConversionScript = FindConversionScript(CollectionParam.GetType(), ConversionScriptCache, bIncludeConversionScripts))
+			else if (UNiagaraScript* ConversionScript = FindConversionScript(CollectionParam.GetType(), ConversionScriptCache, Arguments.bIncludeConversionScripts))
 			{
-				OutAvailableParametersWithConversions.Add(CollectionParam, ConversionScript);
+				OutAvailableParameters.Add({CollectionParam, ConversionScript});
 			}
 		}
 	}
+
+	if(Arguments.bIncludeParameterDefinitions)
+	{
+		TRACE_CPUPROFILER_EVENT_SCOPE_ON_CHANNEL(UNiagaraStackFunctionInput::GetAvailableParameters::IncludeParameterDefinitions, NiagaraChannel);
+		// At last, we patch up all the found parameters with available parameter definitions
+		TArray<UNiagaraParameterDefinitions*> AvailableParameterDefinitions;
+		FNiagaraEditorUtilities::GetAvailableParameterDefinitions(AvailableParameterDefinitions);
+		TMap<FNiagaraVariableBase, UNiagaraScriptVariable*> ParameterDefinitionsCache;
+
+		for(UNiagaraParameterDefinitions* DefinitionsAsset : AvailableParameterDefinitions)
+		{
+			for(UNiagaraScriptVariable* ScriptVariable : DefinitionsAsset->GetParametersConst())
+			{
+				ParameterDefinitionsCache.Add(ScriptVariable->Variable, ScriptVariable);
+			}
+		}
+		
+		for(FNiagaraAvailableParameterInfo& ParameterInfo : OutAvailableParameters)
+		{
+			if(ParameterDefinitionsCache.Contains(ParameterInfo.Variable))
+			{
+				ParameterInfo.MatchingParameterDefinition = ParameterDefinitionsCache[ParameterInfo.Variable];
+			}		
+		}
+	}
 }
 
 UNiagaraNodeFunctionCall* UNiagaraStackFunctionInput::GetDefaultFunctionNode() const
@@ -3548,13 +3587,14 @@ void UNiagaraStackFunctionInput::GetDefaultLinkedParameterOrLinkedFunctionFromDe
 	else if(LinkedValues.Num() > 1)
 	{
 		// If there are a chain of linked values use the first one that's available, otherwise just use the last one.
-		TArray<FNiagaraVariableBase> AvailableParameters;
-		TMap<FNiagaraVariableBase, UNiagaraScript*> ConversionHandles;
-		GetAvailableParameters(AvailableParameters, ConversionHandles, false);
+		TSet<FNiagaraAvailableParameterInfo> AvailableParameterInfos;
+		FGetAvailableParameterArgs Args;
+		Args.bIncludeConversionScripts = false;
+		GetAvailableParameters(AvailableParameterInfos, Args);
 		for (FLinkedParameterOrFunctionNode& LinkedValue : LinkedValues)
 		{
 			if (LinkedValue.LinkedFunctionCallNode.IsValid() ||
-				AvailableParameters.Contains(LinkedValue.LinkedParameter.GetValue()))
+				AvailableParameterInfos.Contains(LinkedValue.LinkedParameter.GetValue()))
 			{
 				ValueToUse = &LinkedValue;
 				break;
diff --git a/Engine/Plugins/FX/Niagara/Source/NiagaraEditor/Private/ViewModels/Stack/NiagaraStackGraphUtilities.cpp b/Engine/Plugins/FX/Niagara/Source/NiagaraEditor/Private/ViewModels/Stack/NiagaraStackGraphUtilities.cpp
--- a/Engine/Plugins/FX/Niagara/Source/NiagaraEditor/Private/ViewModels/Stack/NiagaraStackGraphUtilities.cpp
+++ b/Engine/Plugins/FX/Niagara/Source/NiagaraEditor/Private/ViewModels/Stack/NiagaraStackGraphUtilities.cpp
@@ -49,6 +49,7 @@
 #include "NiagaraSimulationStageBase.h"
 #include "NiagaraEmitter.h"
 #include "DataHierarchyViewModelBase.h"
+#include "NiagaraTrace.h"
 #include "ViewModels/HierarchyEditor/NiagaraSummaryViewViewModel.h"
 #include "TraversalCache/StaticContextBuilder.h"
 #include "TraversalCache/TraversalCache.h"
@@ -547,6 +548,8 @@ TArray<FName> FNiagaraStackGraphUtilities::StackContextResolution(FVersionedNiag
 
 void FNiagaraStackGraphUtilities::BuildParameterMapHistoryWithStackContextResolution(FVersionedNiagaraEmitter OwningEmitter, UNiagaraNodeOutput* OutputNodeInChain, UNiagaraNode* NodeToVisit, TArray<FNiagaraParameterMapHistory>& OutHistories, bool bRecursive /*= true*/, bool bFilterForCompilation /*= true*/)
 {
+	TRACE_CPUPROFILER_EVENT_SCOPE_ON_CHANNEL(FNiagaraStackGraphUtilities::BuildParameterMapHistoryWithStackContextResolution, NiagaraChannel);
+	
 	FNiagaraParameterMapHistoryBuilder Builder;
 
 	bool bSetUsage = false;
diff --git a/Engine/Plugins/FX/Niagara/Source/NiagaraEditor/Private/ViewModels/Stack/NiagaraStackRenderItemGroup.cpp b/Engine/Plugins/FX/Niagara/Source/NiagaraEditor/Private/ViewModels/Stack/NiagaraStackRenderItemGroup.cpp
--- a/Engine/Plugins/FX/Niagara/Source/NiagaraEditor/Private/ViewModels/Stack/NiagaraStackRenderItemGroup.cpp
+++ b/Engine/Plugins/FX/Niagara/Source/NiagaraEditor/Private/ViewModels/Stack/NiagaraStackRenderItemGroup.cpp
@@ -54,6 +54,21 @@ public:
 		return FText();
 	}
 
+	virtual TOptional<FNiagaraFavoritesActionData> GetFavoritesData() const override
+	{
+		return TOptional<FNiagaraFavoritesActionData>();
+	}
+	
+	virtual bool IsInLibrary() const override
+	{
+		return true;
+	}
+	
+	virtual FNiagaraActionSourceData GetSourceData() const override
+	{
+		return FNiagaraActionSourceData();
+	}
+	
 	const FNiagaraRendererCreationInfo& GetRendererCreationInfo() { return RendererCreationInfo; }
 
 private:
diff --git a/Engine/Plugins/FX/Niagara/Source/NiagaraEditor/Private/ViewModels/Stack/NiagaraStackScriptItemGroup.cpp b/Engine/Plugins/FX/Niagara/Source/NiagaraEditor/Private/ViewModels/Stack/NiagaraStackScriptItemGroup.cpp
--- a/Engine/Plugins/FX/Niagara/Source/NiagaraEditor/Private/ViewModels/Stack/NiagaraStackScriptItemGroup.cpp
+++ b/Engine/Plugins/FX/Niagara/Source/NiagaraEditor/Private/ViewModels/Stack/NiagaraStackScriptItemGroup.cpp
@@ -15,6 +15,7 @@
 #include "ViewModels/NiagaraScriptGraphViewModel.h"
 #include "NiagaraStackEditorData.h"
 #include "ScopedTransaction.h"
+#include "Config/NiagaraFavoriteActionsConfig.h"
 #include "DragAndDrop/AssetDragDropOp.h"
 #include "Framework/Notifications/NotificationManager.h"
 #include "Interfaces/IMainFrameModule.h"
@@ -42,6 +43,8 @@
 
 using namespace UE::Niagara::Wizard;
 
+const FName INiagaraStackItemGroupAddAction::FavoriteActionsProfile = FName("StackGroupAdd");
+
 class FScriptGroupAddAction : public INiagaraStackItemGroupAddAction
 {
 public:
@@ -66,9 +69,14 @@ public:
 		AssetData.GetTagValue(GET_MEMBER_NAME_CHECKED(FVersionedNiagaraScriptData, Keywords), Keywords);
 
 		bool bSuggested = AssetData.GetTagValueRef<bool>(GET_MEMBER_NAME_CHECKED(FVersionedNiagaraScriptData, bSuggested));
+
+		FNiagaraFavoritesActionData FavoritesActionData;
+		FavoritesActionData.ActionIdentifier.Names.Add(FName(AssetData.GetObjectPathString()));
+		FavoritesActionData.bFavoriteByDefault = bSuggested;
+		
 		TTuple<EScriptSource, FText> Source = FNiagaraEditorUtilities::GetScriptSource(AssetData);
 		FNiagaraActionSourceData SourceData(Source.Key, Source.Value, true);
-		return MakeShareable(new FScriptGroupAddAction(DisplayName, {Category.ToString()}, Description, Keywords, bSuggested, bIsInLibrary, FNiagaraVariable(), false, AssetData, nullptr, false, false, SourceData));
+		return MakeShareable(new FScriptGroupAddAction(DisplayName, {Category.ToString()}, Description, Keywords, FavoritesActionData, bIsInLibrary, FNiagaraVariable(), false, AssetData, nullptr, false, false, SourceData));
 	}
 
 	static TSharedRef<FScriptGroupAddAction> CreateModuleActionFromScratchPadScript(UNiagaraScript* ModuleScript)
@@ -85,11 +93,14 @@ public:
 		FText DisplayName = FNiagaraEditorUtilities::FormatScriptName(ModuleScript->GetFName(), bIsInLibrary);
 		FText Description = FNiagaraEditorUtilities::FormatScriptDescription(ScriptData->Description, FSoftObjectPath(ModuleScript), bIsInLibrary);
 		FText Keywords = ScriptData->Keywords;
-		bool bSuggested = ScriptData->bSuggested;
 
+		FNiagaraFavoritesActionData FavoritesActionData;
+		FavoritesActionData.ActionIdentifier.Names.Add(FName(GetPathNameSafe(ModuleScript)));
+		FavoritesActionData.bFavoriteByDefault = ScriptData->bSuggested;
+		
 		FAssetData Data = FAssetData(ModuleScript);
 		FNiagaraActionSourceData SourceData(EScriptSource::Niagara, FText::FromString("Scratch Pad"), true);
-		return MakeShareable(new FScriptGroupAddAction(DisplayName, {Category.ToString()}, Description, Keywords, bSuggested, bIsInLibrary, FNiagaraVariable(), false, FAssetData(), ModuleScript, false, false, SourceData));
+		return MakeShareable(new FScriptGroupAddAction(DisplayName, {Category.ToString()}, Description, Keywords, FavoritesActionData, bIsInLibrary, FNiagaraVariable(), false, FAssetData(), ModuleScript, false, false, SourceData));
 	}
 
 	static TSharedRef<FScriptGroupAddAction> CreateExistingParameterModuleAction(FNiagaraVariable ParameterVariable)
@@ -100,7 +111,7 @@ public:
 		FText AttributeDescription = FNiagaraConstants::GetAttributeDescription(ParameterVariable);
 		FText Description = FText::Format(LOCTEXT("ExistingParameterModuleDescriptoinFormat", "Description: Set the parameter {0}. {1}"), FText::FromName(ParameterVariable.GetName()), AttributeDescription);
 
-		return MakeShareable(new FScriptGroupAddAction(DisplayName, {Category.ToString()}, Description, FText(), false, true, ParameterVariable, false, FAssetData(), nullptr, false, false));
+		return MakeShareable(new FScriptGroupAddAction(DisplayName, {Category.ToString()}, Description, FText(), {}, true, ParameterVariable, false, FAssetData(), nullptr, false, false));
 	}
 
 	static TSharedRef<FScriptGroupAddAction> CreateNewParameterModuleAction(FName NewParameterNamespace, FNiagaraTypeDefinition NewParameterType)
@@ -112,26 +123,40 @@ public:
 		FNiagaraParameterHandle NewParameterHandle(NewParameterNamespace, *(TEXT("New") + NewParameterType.GetName()));
 		FNiagaraVariable NewParameter(NewParameterType, NewParameterHandle.GetParameterHandleString());
 
-		return MakeShareable(new FScriptGroupAddAction( DisplayName, {Category.ToString()}, Description, FText(), false, true, NewParameter, true, FAssetData(), nullptr, false, false));
+		return MakeShareable(new FScriptGroupAddAction( DisplayName, {Category.ToString()}, Description, FText(), {}, true, NewParameter, true, FAssetData(), nullptr, false, false));
 	}
 
 	static TSharedRef<FScriptGroupAddAction> CreateNewScratchModuleAction()
 	{
 		FText DisplayName = LOCTEXT("NewScratchModuleName", "New Scratch Pad Module");
 		FText Description = LOCTEXT("NewScratchModuleDescription", "Description: Create a new scratch pad module.");
-		return MakeShareable(new FScriptGroupAddAction( DisplayName, {}, Description, FText(), true, true, FNiagaraVariable(), false, FAssetData(), nullptr, true, false));
+		
+		FNiagaraFavoritesActionData FavoritesActionData;
+		FavoritesActionData.ActionIdentifier.Names.Add(FName(DisplayName.ToString()));
+		FavoritesActionData.bFavoriteByDefault = true;
+		
+		return MakeShareable(new FScriptGroupAddAction( DisplayName, {}, Description, FText(), FavoritesActionData, true, FNiagaraVariable(), false, FAssetData(), nullptr, true, false));
 	}
 
 	static TSharedRef<FScriptGroupAddAction> CreateNewSetSpecificModuleAction()
 	{
 		FText DisplayName = LOCTEXT("NewSetSpecificModuleName", "Set new or existing parameter directly");
 		FText Description = LOCTEXT("NewSetSpecificModuleDescription", "Description: Create a new module that can set new or existing parameters directly.");
-		return MakeShareable(new FScriptGroupAddAction( DisplayName, {}, Description, FText(), true, true, FNiagaraVariable(), false, FAssetData(), nullptr, false, true));
+		
+		FNiagaraFavoritesActionData FavoritesActionData;
+		FavoritesActionData.ActionIdentifier.Names.Add(FName(DisplayName.ToString()));
+		FavoritesActionData.bFavoriteByDefault = true;
+		
+		return MakeShareable(new FScriptGroupAddAction( DisplayName, {}, Description, FText(), FavoritesActionData, true, FNiagaraVariable(), false, FAssetData(), nullptr, false, true));
 	}
 
 	static TSharedRef<FScriptGroupAddAction> CreateWizardAction(const FModuleWizardGenerator::FAction& Action)
 	{
-		TSharedRef<FScriptGroupAddAction> AddAction = MakeShareable(new FScriptGroupAddAction( Action.DisplayName, {}, Action.Description, Action.Keywords, Action.bSuggestedAction, true, FNiagaraVariable(), false, FAssetData(), nullptr, false, true));
+		FNiagaraFavoritesActionData FavoritesActionData;
+		FavoritesActionData.ActionIdentifier.Names.Add(Action.WizardModel->GetIdentifier());
+		FavoritesActionData.bFavoriteByDefault = Action.bSuggestedAction;
+		
+		TSharedRef<FScriptGroupAddAction> AddAction = MakeShareable(new FScriptGroupAddAction( Action.DisplayName, {}, Action.Description, Action.Keywords, FavoritesActionData, true, FNiagaraVariable(), false, FAssetData(), nullptr, false, true));
 		AddAction->ModuleWizardModel = Action.WizardModel;
 		return AddAction;
 	}
@@ -156,9 +181,9 @@ public:
 		return Keywords;
 	}
 
-	virtual bool GetSuggested() const override
+	virtual TOptional<FNiagaraFavoritesActionData> GetFavoritesData() const override
 	{
-		return bSuggested;
+		return FavoritesData;
 	}
 
 	virtual bool IsInLibrary() const override
@@ -213,7 +238,7 @@ public:
 	}
 
 private:
-	FScriptGroupAddAction(FText InDisplayName, TArray<FString> InCategories, FText InDescription, FText InKeywords, bool bInSuggested, bool bInIsInLibrary,
+	FScriptGroupAddAction(FText InDisplayName, TArray<FString> InCategories, FText InDescription, FText InKeywords, TOptional<FNiagaraFavoritesActionData> InFavoritesData, bool bInIsInLibrary,
 		FNiagaraVariable InModuleParameterVariable, bool bInRenameParameterOnAdd,
 		FAssetData InModuleAssetData, UNiagaraScript* InModuleScript,
 		bool bInIsNewScratchModuleAction, bool bInIsNewSetSpecificModuleAction, FNiagaraActionSourceData InSourceData = FNiagaraActionSourceData(EScriptSource::Niagara, FText::FromString(TEXT("Niagara")), true))
@@ -222,7 +247,7 @@ private:
 		, Categories(InCategories)
 		, Description(InDescription)
 		, Keywords(InKeywords)
-		, bSuggested(bInSuggested)
+		, FavoritesData(InFavoritesData)
 		, bIsInLibrary(bInIsInLibrary)
 		, ModuleParameterVariable(InModuleParameterVariable)
 		, bRenameParameterOnAdd(bInRenameParameterOnAdd)
@@ -238,7 +263,7 @@ private:
 	TArray<FString> Categories;
 	FText Description;
 	FText Keywords;
-	bool bSuggested;
+	TOptional<FNiagaraFavoritesActionData> FavoritesData;
 	bool bIsInLibrary;
 	FNiagaraVariable ModuleParameterVariable;
 	bool bRenameParameterOnAdd;
diff --git a/Engine/Plugins/FX/Niagara/Source/NiagaraEditor/Private/ViewModels/Stack/NiagaraStackStatelessEmitterSimulateGroup.cpp b/Engine/Plugins/FX/Niagara/Source/NiagaraEditor/Private/ViewModels/Stack/NiagaraStackStatelessEmitterSimulateGroup.cpp
--- a/Engine/Plugins/FX/Niagara/Source/NiagaraEditor/Private/ViewModels/Stack/NiagaraStackStatelessEmitterSimulateGroup.cpp
+++ b/Engine/Plugins/FX/Niagara/Source/NiagaraEditor/Private/ViewModels/Stack/NiagaraStackStatelessEmitterSimulateGroup.cpp
@@ -95,6 +95,9 @@ public:
 	virtual FText GetDisplayName() const override { return DisplayName; }
 	virtual FText GetDescription() const override { return FText(); }
 	virtual FText GetKeywords() const override { return FText(); }
+	virtual TOptional<FNiagaraFavoritesActionData> GetFavoritesData() const override { return TOptional<FNiagaraFavoritesActionData>(); }
+	virtual bool IsInLibrary() const override { return true; }
+	virtual FNiagaraActionSourceData GetSourceData() const override { return FNiagaraActionSourceData(); }
 
 private:
 	TWeakObjectPtr<UNiagaraStatelessModule> StatelessModuleWeak;
diff --git a/Engine/Plugins/FX/Niagara/Source/NiagaraEditor/Private/ViewModels/Stack/NiagaraStackStatelessEmitterSpawnGroup.cpp b/Engine/Plugins/FX/Niagara/Source/NiagaraEditor/Private/ViewModels/Stack/NiagaraStackStatelessEmitterSpawnGroup.cpp
--- a/Engine/Plugins/FX/Niagara/Source/NiagaraEditor/Private/ViewModels/Stack/NiagaraStackStatelessEmitterSpawnGroup.cpp
+++ b/Engine/Plugins/FX/Niagara/Source/NiagaraEditor/Private/ViewModels/Stack/NiagaraStackStatelessEmitterSpawnGroup.cpp
@@ -106,6 +106,9 @@ public:
 	virtual FText GetDisplayName() const override { return DisplayName; }
 	virtual FText GetDescription() const override { return FText(); }
 	virtual FText GetKeywords() const override { return FText(); }
+	virtual TOptional<FNiagaraFavoritesActionData> GetFavoritesData() const override { return TOptional<FNiagaraFavoritesActionData>(); }
+	virtual bool IsInLibrary() const override { return true; }
+	virtual FNiagaraActionSourceData GetSourceData() const override { return FNiagaraActionSourceData(); }
 
 private:
 	ENiagaraStatelessSpawnInfoType SpawnInfoType;
diff --git a/Engine/Plugins/FX/Niagara/Source/NiagaraEditor/Private/ViewModels/Stack/NiagaraStackSystemSettingsGroup.cpp b/Engine/Plugins/FX/Niagara/Source/NiagaraEditor/Private/ViewModels/Stack/NiagaraStackSystemSettingsGroup.cpp
--- a/Engine/Plugins/FX/Niagara/Source/NiagaraEditor/Private/ViewModels/Stack/NiagaraStackSystemSettingsGroup.cpp
+++ b/Engine/Plugins/FX/Niagara/Source/NiagaraEditor/Private/ViewModels/Stack/NiagaraStackSystemSettingsGroup.cpp
@@ -51,6 +51,21 @@ public:
 		return FText();
 	}
 
+	virtual TOptional<FNiagaraFavoritesActionData> GetFavoritesData() const override
+	{
+		return TOptional<FNiagaraFavoritesActionData>();
+	}
+	
+	virtual bool IsInLibrary() const override
+	{
+		return true;
+	}
+	
+	virtual FNiagaraActionSourceData GetSourceData() const override
+	{
+		return FNiagaraActionSourceData();
+	}
+
 private:
 	FNiagaraVariable NewParameterVariable;
 };
diff --git a/Engine/Plugins/FX/Niagara/Source/NiagaraEditor/Private/Widgets/DataChannel/NiagaraDataChannelWizard.cpp b/Engine/Plugins/FX/Niagara/Source/NiagaraEditor/Private/Widgets/DataChannel/NiagaraDataChannelWizard.cpp
--- a/Engine/Plugins/FX/Niagara/Source/NiagaraEditor/Private/Widgets/DataChannel/NiagaraDataChannelWizard.cpp
+++ b/Engine/Plugins/FX/Niagara/Source/NiagaraEditor/Private/Widgets/DataChannel/NiagaraDataChannelWizard.cpp
@@ -542,6 +542,11 @@ TSharedRef<FModuleWizardModel> DataChannel::CreateReadNDCModuleWizardModel()
 			Pages.Add(VariablesPage.ToSharedRef());
 		}
 		virtual ~FReadNDCModel() override = default;
+
+		virtual FName GetIdentifier() const override
+		{
+			return "ReadNDCWizardModel";
+		}
 		
 		virtual void GenerateNewModuleContent(TSharedPtr<FNiagaraScratchPadScriptViewModel> ScratchPadScriptViewModel, const TArray<const UNiagaraNodeFunctionCall*>& PreviousModules) override
 		{
@@ -701,6 +706,11 @@ TSharedRef<FModuleWizardModel> DataChannel::CreateWriteNDCModuleWizardModel()
 			Pages.Add(VariablesPage.ToSharedRef());
 		}
 		virtual ~FWriteNDCModel() override = default;
+
+		virtual FName GetIdentifier() const override
+		{
+			return "WriteNDCWizardModel";
+		}
 		
 		virtual void GenerateNewModuleContent(TSharedPtr<FNiagaraScratchPadScriptViewModel> ScratchPadScriptViewModel, const TArray<const UNiagaraNodeFunctionCall*>& PreviousModules) override
 		{
@@ -876,6 +886,11 @@ TSharedRef<FModuleWizardModel> DataChannel::CreateSpawnNDCModuleWizardModel()
 		}
 		virtual ~FSpawnNDCModel() override = default;
 
+		virtual FName GetIdentifier() const override
+		{
+			return "SpawnNDCWizardModel";
+		}
+		
 		virtual TArray<FModuleCreationEntry> GetModulesToCreate(UNiagaraNodeOutput* ProvidedOutputNode, int32 ProvidedTargetIndex, TSharedPtr<FNiagaraSystemViewModel> SystemModel, TSharedPtr<FNiagaraEmitterViewModel> EmitterViewModel) override
 		{
 			TArray<FModuleCreationEntry> Result;
diff --git a/Engine/Plugins/FX/Niagara/Source/NiagaraEditor/Private/Widgets/SNiagaraGraphActionMenu.cpp b/Engine/Plugins/FX/Niagara/Source/NiagaraEditor/Private/Widgets/SNiagaraGraphActionMenu.cpp
--- a/Engine/Plugins/FX/Niagara/Source/NiagaraEditor/Private/Widgets/SNiagaraGraphActionMenu.cpp
+++ b/Engine/Plugins/FX/Niagara/Source/NiagaraEditor/Private/Widgets/SNiagaraGraphActionMenu.cpp
@@ -11,6 +11,7 @@
 #include "NiagaraEditorStyle.h"
 #include "NiagaraEditorUtilities.h"
 #include "SNiagaraGraphActionWidget.h"
+#include "Config/NiagaraFavoriteActionsConfig.h"
 #include "Widgets/Text/SRichTextBlock.h"
 #include "Widgets/Text/STextBlock.h"
 #include "Widgets/SItemSelector.h"
@@ -18,6 +19,7 @@
 #define LOCTEXT_NAMESPACE "NiagaraGraphActionMenu"
 
 bool SNiagaraGraphActionMenu::bLibraryOnly = true;
+const FName SNiagaraGraphActionMenu::FavoriteActionsProfile = FName("NiagaraGraphActions");
 
 SNiagaraGraphActionMenu::~SNiagaraGraphActionMenu()
 {
@@ -74,7 +76,7 @@ void SNiagaraGraphActionMenu::Construct( const FArguments& InArgs )
 					.OnGenerateWidgetForItem(this, &SNiagaraGraphActionMenu::OnGenerateWidgetForItem)
 					.OnGetItemWeight_Lambda([](const TSharedPtr<FNiagaraAction_NewNode>& InCurrentAction, const TArray<FString>& InFilterTerms)
 					{
-						return FNiagaraEditorUtilities::GetWeightForItem(InCurrentAction, InFilterTerms);
+						return FNiagaraEditorUtilities::GetWeightForItem(InCurrentAction, SNiagaraGraphActionMenu::FavoriteActionsProfile, InFilterTerms);
 					})
 					.OnItemActivated(this, &SNiagaraGraphActionMenu::OnItemActivated)
 					.AllowMultiselect(false)
@@ -82,6 +84,10 @@ void SNiagaraGraphActionMenu::Construct( const FArguments& InArgs )
 					.OnDoesSectionPassCustomFilter(this, &SNiagaraGraphActionMenu::DoesSectionPassCustomFilter)
 					.ClickActivateMode(EItemSelectorClickActivateMode::SingleClick)
 					.ExpandInitially(false)
+					.OnItemRowHoverEvent_Lambda([](const TSharedPtr<FNiagaraAction_NewNode>& Item, bool bHover)
+					{
+						Item->bIsHovered = bHover;
+					})
 					.OnGetSectionData_Lambda([](const ENiagaraMenuSections& Section)
 					{
 						if(Section == ENiagaraMenuSections::Suggested)
@@ -120,12 +126,19 @@ TArray<FString> SNiagaraGraphActionMenu::OnGetCategoriesForItem(const TSharedPtr
 
 TArray<ENiagaraMenuSections> SNiagaraGraphActionMenu::OnGetSectionsForItem(const TSharedPtr<FNiagaraAction_NewNode>& Item)
 {	
-	if(Item->Section == ENiagaraMenuSections::Suggested)
+	TArray<ENiagaraMenuSections> Sections{ENiagaraMenuSections::General};
+
+	if(Item->FavoritesActionData.IsSet())
 	{
-		return { ENiagaraMenuSections::General, ENiagaraMenuSections::Suggested };
+		FNiagaraFavoriteActionsProfile& ActionsProfile = UNiagaraFavoriteActionsConfig::Get()->GetActionsProfile(SNiagaraGraphActionMenu::FavoriteActionsProfile);
+		TOptional<bool> bIsFavorite = ActionsProfile.IsFavorite(Item->FavoritesActionData.GetValue());
+		if(bIsFavorite.IsSet() && bIsFavorite.GetValue())
+		{
+			Sections.Add(ENiagaraMenuSections::Suggested);
+		}
 	}
 		
-	return {Item->Section};
+	return Sections;
 }
 
 bool SNiagaraGraphActionMenu::OnCompareSectionsForEquality(const ENiagaraMenuSections& SectionA, const ENiagaraMenuSections& SectionB)
@@ -183,6 +196,7 @@ TSharedRef<SWidget> SNiagaraGraphActionMenu::OnGenerateWidgetForItem(const TShar
 {
 	FCreateNiagaraWidgetForActionData ActionData(Item);
 	ActionData.HighlightText = TAttribute<FText>::Create(TAttribute<FText>::FGetter::CreateRaw(this, &SNiagaraGraphActionMenu::GetFilterText));
+	ActionData.FavoriteActionsProfileName = SNiagaraGraphActionMenu::FavoriteActionsProfile;
 	return SNew(SNiagaraActionWidget, ActionData);
 }
 
diff --git a/Engine/Plugins/FX/Niagara/Source/NiagaraEditor/Private/Widgets/SNiagaraGraphActionWidget.cpp b/Engine/Plugins/FX/Niagara/Source/NiagaraEditor/Private/Widgets/SNiagaraGraphActionWidget.cpp
--- a/Engine/Plugins/FX/Niagara/Source/NiagaraEditor/Private/Widgets/SNiagaraGraphActionWidget.cpp
+++ b/Engine/Plugins/FX/Niagara/Source/NiagaraEditor/Private/Widgets/SNiagaraGraphActionWidget.cpp
@@ -10,6 +10,7 @@
 #include "Widgets/SNiagaraParameterName.h"
 #include "NiagaraEditorStyle.h"
 #include "NiagaraEditorUtilities.h"
+#include "Config/NiagaraFavoriteActionsConfig.h"
 #include "Widgets/Layout/SSpacer.h"
 #include "Widgets/Layout/SSeparator.h"
 
@@ -82,6 +83,7 @@ FReply SNiagaraGraphActionWidget::OnMouseButtonDown(const FGeometry& MyGeometry,
 void SNiagaraActionWidget::Construct(const FArguments& InArgs, const FCreateNiagaraWidgetForActionData& InCreateData)
 {
 	ActionPtr = InCreateData.Action;
+	FavoriteProfileName = InCreateData.FavoriteActionsProfileName;
 
 	TSharedPtr<SWidget> NameWidget;
 	if (ActionPtr->GetParameterVariable().IsSet())
@@ -123,11 +125,11 @@ void SNiagaraActionWidget::Construct(const FArguments& InArgs, const FCreateNiag
 			NameWidget.ToSharedRef()
 		]
 		+ SHorizontalBox::Slot()
-       .FillWidth(1.f)
-       [
+		.FillWidth(1.f)
+		[
            SNew(SSpacer)
-       ]
-       + SHorizontalBox::Slot()
+		]
+		+ SHorizontalBox::Slot()
         .AutoWidth()
         .VAlign(VAlign_Center)
         .HAlign(HAlign_Right)
@@ -137,15 +139,23 @@ void SNiagaraActionWidget::Construct(const FArguments& InArgs, const FCreateNiag
             .Visibility(ActionPtr->bIsExperimental ? EVisibility::Visible : EVisibility::Collapsed)
             .ToolTipText(ActionPtr->bIsExperimental ? LOCTEXT("ScriptExperimentalToolTip", "This script is experimental, use with care!") : FText::GetEmpty())
         ]
-       + SHorizontalBox::Slot()
-       .HAlign(HAlign_Right)
-       .VAlign(VAlign_Fill)
-       [
+		+ SHorizontalBox::Slot()
+		.HAlign(HAlign_Right)
+		.VAlign(VAlign_Center)
+		.AutoWidth()
+		[
+			CreateFavoriteActionWidget()
+		]
+		+ SHorizontalBox::Slot()
+		.HAlign(HAlign_Right)
+		.VAlign(VAlign_Fill)
+		.AutoWidth()
+		[
            SNew(SSeparator)
            .SeparatorImage(FAppStyle::Get().GetBrush("Separator"))
            .Orientation(EOrientation::Orient_Vertical)
-           .Visibility(ActionPtr->SourceData.bDisplaySource ? EVisibility::Visible : EVisibility::Collapsed)
-       ]
+			.Visibility(ActionPtr->SourceData.bDisplaySource ? EVisibility::Visible : EVisibility::Collapsed)
+		]
        + SHorizontalBox::Slot()
        .HAlign(HAlign_Fill)
        .VAlign(VAlign_Center)
@@ -165,4 +175,91 @@ void SNiagaraActionWidget::Construct(const FArguments& InArgs, const FCreateNiag
     ];
 }
 
+TSharedRef<SWidget> SNiagaraActionWidget::CreateFavoriteActionWidget()
+{	
+	return
+		SNew(SBox)
+		.WidthOverride(16.f)
+		.HeightOverride(16.f)
+		[
+			SNew(SButton)
+			.ContentPadding(FMargin(0.f))
+			.ButtonStyle(&FAppStyle::Get().GetWidgetStyle<FButtonStyle>("HoverHintOnly"))
+			.OnClicked(this, &SNiagaraActionWidget::OnFavoriteStateChanged)
+			.Visibility(this, &SNiagaraActionWidget::OnShouldShowFavoriteButton)
+			.ToolTipText(this, &SNiagaraActionWidget::GetFavoriteButtonTooltipText)
+			[
+				SNew(SImage)
+				.Image(this, &SNiagaraActionWidget::GetFavoriteBrush)
+			]
+		];
+}
+
+FReply SNiagaraActionWidget::OnFavoriteStateChanged()
+{
+	// We still want the button to consume the input, but generally this will never happen as the button isn't displayed
+	if(FavoriteProfileName.IsSet() == false || ActionPtr->FavoritesActionData.IsSet() == false)
+	{
+		return FReply::Handled();
+	}
+	
+	FNiagaraFavoriteActionsProfile& ActionsProfile = UNiagaraFavoriteActionsConfig::Get()->GetActionsProfile(FavoriteProfileName.GetValue());
+	ActionsProfile.ToggleFavoriteAction(ActionPtr->FavoritesActionData.GetValue());
+
+	return FReply::Handled();
+}
+
+ECheckBoxState SNiagaraActionWidget::OnGetIsFavorite() const
+{
+	if(FavoriteProfileName.IsSet() == false || ActionPtr->FavoritesActionData.IsSet() == false)
+	{
+		return ECheckBoxState::Undetermined;
+	}
+		
+	FNiagaraFavoriteActionsProfile& ActionsProfile = UNiagaraFavoriteActionsConfig::Get()->GetActionsProfile(FavoriteProfileName.GetValue());
+	return ActionsProfile.IsFavorite(ActionPtr->FavoritesActionData.GetValue()) ? ECheckBoxState::Checked : ECheckBoxState::Unchecked;
+}
+
+EVisibility SNiagaraActionWidget::OnShouldShowFavoriteButton() const
+{
+	if(FavoriteProfileName.IsSet() == false || ActionPtr->FavoritesActionData.IsSet() == false)
+	{
+		return EVisibility::Collapsed;
+	}
+
+	return ActionPtr->bIsHovered ? EVisibility::Visible : EVisibility::Collapsed;
+}
+
+const FSlateBrush* SNiagaraActionWidget::GetFavoriteBrush() const
+{
+	if(FavoriteProfileName.IsSet() == false || ActionPtr->FavoritesActionData.IsSet() == false)
+	{
+		return FAppStyle::GetNoBrush();
+	}
+	
+	FNiagaraFavoriteActionsProfile& ActionsProfile = UNiagaraFavoriteActionsConfig::Get()->GetActionsProfile(FavoriteProfileName.GetValue());	
+	if(ActionsProfile.IsFavorite(ActionPtr->FavoritesActionData.GetValue()))
+	{
+		return FAppStyle::GetBrush("Icons.Star");
+	}
+	
+	return FAppStyle::GetBrush("Icons.Star.Outline");
+}
+
+FText SNiagaraActionWidget::GetFavoriteButtonTooltipText() const
+{
+	if(FavoriteProfileName.IsSet() == false || ActionPtr->FavoritesActionData.IsSet() == false)
+	{
+		return FText::GetEmpty();
+	}
+
+	FNiagaraFavoriteActionsProfile& ActionsProfile = UNiagaraFavoriteActionsConfig::Get()->GetActionsProfile(FavoriteProfileName.GetValue());
+	if(ActionsProfile.IsFavorite(ActionPtr->FavoritesActionData.GetValue()))
+	{
+		return LOCTEXT("UnfavoriteButtonTooltip", "Remove this action from favorites");
+	}
+
+	return LOCTEXT("FavoriteButtonTooltip", "Add this action to favorites");
+}
+
 #undef LOCTEXT_NAMESPACE
diff --git a/Engine/Plugins/FX/Niagara/Source/NiagaraEditor/Public/Config/NiagaraFavoriteActionsConfig.h b/Engine/Plugins/FX/Niagara/Source/NiagaraEditor/Public/Config/NiagaraFavoriteActionsConfig.h
new file mode 100644
--- /dev/null
+++ b/Engine/Plugins/FX/Niagara/Source/NiagaraEditor/Public/Config/NiagaraFavoriteActionsConfig.h
@@ -0,0 +1,119 @@
+﻿// Copyright Epic Games, Inc. All Rights Reserved.
+
+#pragma once
+
+#include "CoreMinimal.h"
+#include "EditorConfigBase.h"
+#include "NiagaraFavoriteActionsConfig.generated.h"
+
+USTRUCT()
+struct FNiagaraActionIdentifier
+{
+	GENERATED_BODY()
+
+	UPROPERTY()
+	TArray<FName> Names;
+
+	UPROPERTY()
+	TArray<FGuid> Guids;	
+	
+	bool IsValid() const
+	{
+		return Guids.Num() > 0 || Names.Num() > 0;
+	}
+	
+	bool operator==(const FNiagaraActionIdentifier& OtherIdentity) const
+	{
+		if(Guids.Num() != OtherIdentity.Guids.Num() || Names.Num() != OtherIdentity.Names.Num())
+		{
+			return false;
+		}
+
+		for(int32 GuidIndex = 0; GuidIndex < Guids.Num(); GuidIndex++)
+		{
+			if(Guids[GuidIndex] != OtherIdentity.Guids[GuidIndex])
+			{
+				return false;
+			}
+		}
+
+		for(int32 NameIndex = 0; NameIndex < Names.Num(); NameIndex++)
+		{
+			if(!Names[NameIndex].IsEqual(OtherIdentity.Names[NameIndex]))
+			{
+				return false;
+			}
+		}
+
+		return true;
+	}
+
+	bool operator!=(const FNiagaraActionIdentifier& OtherIdentity) const
+	{
+		return !(*this == OtherIdentity);
+	}
+};
+
+FORCEINLINE uint32 GetTypeHash(const FNiagaraActionIdentifier& Identity)
+{
+	uint32 Hash = 0;
+	
+	for(const FGuid& Guid : Identity.Guids)
+	{
+		Hash = HashCombine(Hash, GetTypeHash(Guid));
+	}
+	
+	for(const FName& Name : Identity.Names)
+	{
+		Hash = HashCombine(Hash, GetTypeHash(Name));
+	}
+	
+	return Hash;
+}
+
+struct FNiagaraFavoritesActionData
+{
+	FNiagaraActionIdentifier ActionIdentifier;
+	bool bFavoriteByDefault = false;
+};
+
+USTRUCT()
+struct NIAGARAEDITOR_API FNiagaraFavoriteActionsProfile
+{
+	GENERATED_BODY()
+public:
+	
+	bool IsFavorite(FNiagaraFavoritesActionData InAction);
+	void ToggleFavoriteAction(FNiagaraFavoritesActionData InAction);
+	
+private:
+	/** Explicitly favorited actions */
+	UPROPERTY()
+	TSet<FNiagaraActionIdentifier> FavoriteActions;
+
+	/** For unfavorited actions */
+	UPROPERTY()
+	TSet<FNiagaraActionIdentifier> UnfavoriteActions;
+};
+/**
+ * 
+ */
+UCLASS(EditorConfig="FavoriteNiagaraActions")
+class NIAGARAEDITOR_API UNiagaraFavoriteActionsConfig : public UEditorConfigBase
+{
+	GENERATED_BODY()
+
+public:
+	bool HasActionsProfile(FName ProfileName) const { return Profiles.Contains(ProfileName); }
+	FNiagaraFavoriteActionsProfile& GetActionsProfile(FName ProfileName);
+	
+	static UNiagaraFavoriteActionsConfig* Get();
+	static void Shutdown();
+	
+private:
+	UPROPERTY(meta=(EditorConfig))
+	TMap<FName, FNiagaraFavoriteActionsProfile> Profiles;
+
+private:
+	static TStrongObjectPtr<UNiagaraFavoriteActionsConfig> Instance;
+};
diff --git a/Engine/Plugins/FX/Niagara/Source/NiagaraEditor/Public/NiagaraActions.h b/Engine/Plugins/FX/Niagara/Source/NiagaraEditor/Public/NiagaraActions.h
--- a/Engine/Plugins/FX/Niagara/Source/NiagaraEditor/Public/NiagaraActions.h
+++ b/Engine/Plugins/FX/Niagara/Source/NiagaraEditor/Public/NiagaraActions.h
@@ -9,9 +9,9 @@
 #include "DragAndDrop/DecoratedDragDropOp.h"
 #include "NiagaraEditorCommon.h"
 #include "ViewModels/Stack/NiagaraParameterHandle.h"
+#include "Config/NiagaraFavoriteActionsConfig.h"
 #include "NiagaraActions.generated.h"
 
-
 class UToolMenu;
 class UGraphNodeContextMenuContext;
 
@@ -112,21 +112,18 @@ USTRUCT()
 struct FNiagaraMenuAction_Base
 {
 	GENERATED_USTRUCT_BODY();
-
+	
 	DECLARE_DELEGATE(FOnExecuteAction);
 	DECLARE_DELEGATE_RetVal(bool, FCanExecuteAction);
 
 	FNiagaraMenuAction_Base() {}
-	NIAGARAEDITOR_API FNiagaraMenuAction_Base(FText DisplayName, ENiagaraMenuSections Section, TArray<FString> InNodeCategories, FText InToolTip, FText InKeywords, float InIntrinsicWeightMultiplier = 1.f);
+	NIAGARAEDITOR_API FNiagaraMenuAction_Base(FText DisplayName, TArray<FString> InNodeCategories, TOptional<FNiagaraFavoritesActionData> FavoritesActionData, FText InToolTip, FText InKeywords, float InIntrinsicWeightMultiplier = 1.f);
 
 	NIAGARAEDITOR_API void UpdateFullSearchText();
 
 	bool bIsExperimental = false;
 
 	bool bIsInLibrary = true;
-
-	/** Top level section this action belongs to. */
-	ENiagaraMenuSections Section;
 	
 	/** Nested categories below a top level section. Can be empty */
 	TArray<FString> Categories;
@@ -151,6 +148,10 @@ struct FNiagaraMenuAction_Base
 
 	/** A multiplier intrinsic to the action. Can be used to tweak search relevance depending on context (like module actions in a module script) */
 	float SearchWeightMultiplier = 1.f;
+	
+	TOptional<FNiagaraFavoritesActionData> FavoritesActionData;
+	
+	mutable bool bIsHovered = false;
 };
 
 USTRUCT()
@@ -161,16 +162,16 @@ struct FNiagaraMenuAction_Generic : public FNiagaraMenuAction_Base
 	FNiagaraMenuAction_Generic() {}
 
 	FNiagaraMenuAction_Generic(FOnExecuteAction ExecuteAction, FCanExecuteAction InCanExecuteAction,
-		FText InDisplayName, ENiagaraMenuSections Section, TArray<FString> InNodeCategories, FText InToolTip, FText InKeywords, float InIntrinsicWeightMultiplier = 1.f)
-    : FNiagaraMenuAction_Base(InDisplayName, Section, InNodeCategories, InToolTip, InKeywords, InIntrinsicWeightMultiplier)
+		FText InDisplayName, TArray<FString> InNodeCategories, TOptional<FNiagaraFavoritesActionData> FavoritesActionData,  FText InToolTip, FText InKeywords, float InIntrinsicWeightMultiplier = 1.f)
+    : FNiagaraMenuAction_Base(InDisplayName, InNodeCategories, FavoritesActionData, InToolTip, InKeywords, InIntrinsicWeightMultiplier)
 	{
 		Action = ExecuteAction;
 		CanExecuteAction = InCanExecuteAction;
 	}
 
 	FNiagaraMenuAction_Generic(FOnExecuteAction ExecuteAction,
-        FText InDisplayName, ENiagaraMenuSections Section, TArray<FString> InNodeCategories, FText InToolTip, FText InKeywords, float InIntrinsicWeightMultiplier = 1.f)
-    : FNiagaraMenuAction_Base(InDisplayName, Section, InNodeCategories, InToolTip, InKeywords, InIntrinsicWeightMultiplier)
+        FText InDisplayName, TArray<FString> InNodeCategories, TOptional<FNiagaraFavoritesActionData> FavoritesActionData, FText InToolTip, FText InKeywords, float InIntrinsicWeightMultiplier = 1.f)
+    : FNiagaraMenuAction_Base(InDisplayName, InNodeCategories, FavoritesActionData, InToolTip, InKeywords, InIntrinsicWeightMultiplier)
 	{
 		Action = ExecuteAction;
 	}
@@ -206,8 +207,8 @@ struct FNiagaraAction_NewNode : public FNiagaraMenuAction_Generic
 
 	FNiagaraAction_NewNode() {}
 	FNiagaraAction_NewNode(
-		FText InDisplayName, ENiagaraMenuSections Section, TArray<FString> InNodeCategories, FText InToolTip, FText InKeywords, float InIntrinsicWeightMultiplier = 1.f)
-		: FNiagaraMenuAction_Generic(FOnExecuteAction(), InDisplayName, Section, InNodeCategories, InToolTip, InKeywords, InIntrinsicWeightMultiplier)
+		FText InDisplayName, TArray<FString> InNodeCategories, TOptional<FNiagaraFavoritesActionData> FavoritesActionData, FText InToolTip, FText InKeywords, float InIntrinsicWeightMultiplier = 1.f)
+		: FNiagaraMenuAction_Generic(FOnExecuteAction(), InDisplayName, InNodeCategories, FavoritesActionData, InToolTip, InKeywords, InIntrinsicWeightMultiplier)
 	{		
 	}
 
diff --git a/Engine/Plugins/FX/Niagara/Source/NiagaraEditor/Public/NiagaraEditorUtilities.h b/Engine/Plugins/FX/Niagara/Source/NiagaraEditor/Public/NiagaraEditorUtilities.h
--- a/Engine/Plugins/FX/Niagara/Source/NiagaraEditor/Public/NiagaraEditorUtilities.h
+++ b/Engine/Plugins/FX/Niagara/Source/NiagaraEditor/Public/NiagaraEditorUtilities.h
@@ -253,7 +253,7 @@ namespace FNiagaraEditorUtilities
 
 	NIAGARAEDITOR_API bool IsEnginePluginAsset(const FTopLevelAssetPath& InTopLevelAssetPath);
 
-	NIAGARAEDITOR_API int32 GetWeightForItem(const TSharedPtr<FNiagaraMenuAction_Generic>& Item, const TArray<FString>& FilterTerms);
+	NIAGARAEDITOR_API int32 GetWeightForItem(const TSharedPtr<FNiagaraMenuAction_Generic>& Item, FName FavoriteActionsProfile, const TArray<FString>& FilterTerms);
 
 	NIAGARAEDITOR_API bool DoesItemMatchFilterText(const FText& FilterText, const TSharedPtr<FNiagaraMenuAction_Generic>& Item);
 	
@@ -356,6 +356,8 @@ namespace FNiagaraEditorUtilities
 
 	NIAGARAEDITOR_API void GetAvailableParameterCollections(TArray<UNiagaraParameterCollection*>& OutParameterCollections);
 
+	NIAGARAEDITOR_API void GetAvailableParameterDefinitions(TArray<UNiagaraParameterDefinitions*>& OutParameterDefinitions);
+
 	TSharedPtr<INiagaraParameterDefinitionsSubscriberViewModel> GetOwningLibrarySubscriberViewModelForGraph(const UNiagaraGraph* Graph);
 
 	TArray<UNiagaraParameterDefinitions*> DowncastParameterDefinitionsBaseArray(const TArray<UNiagaraParameterDefinitionsBase*> BaseArray);
diff --git a/Engine/Plugins/FX/Niagara/Source/NiagaraEditor/Public/NiagaraParameterDefinitions.h b/Engine/Plugins/FX/Niagara/Source/NiagaraEditor/Public/NiagaraParameterDefinitions.h
--- a/Engine/Plugins/FX/Niagara/Source/NiagaraEditor/Public/NiagaraParameterDefinitions.h
+++ b/Engine/Plugins/FX/Niagara/Source/NiagaraEditor/Public/NiagaraParameterDefinitions.h
@@ -73,7 +73,7 @@ public:
 	/** Getters for script variables. 
 	 *  NOTE: These public getters are implemented to support details views. Editing UNiagaraScriptVariable definitions publicly should be done via the add/remove/rename interface otherwise.
 	 */
-	UNiagaraScriptVariable* GetScriptVariable(const FNiagaraVariable& Var);
+	NIAGARAEDITOR_API UNiagaraScriptVariable* GetScriptVariable(const FNiagaraVariable& Var);
 	UNiagaraScriptVariable* GetScriptVariable(const FGuid& ScriptVarId);
 
 	//~ Begin Subscribed ParameterDefinitions methods
diff --git a/Engine/Plugins/FX/Niagara/Source/NiagaraEditor/Public/SNiagaraGraphActionWidget.h b/Engine/Plugins/FX/Niagara/Source/NiagaraEditor/Public/SNiagaraGraphActionWidget.h
--- a/Engine/Plugins/FX/Niagara/Source/NiagaraEditor/Public/SNiagaraGraphActionWidget.h
+++ b/Engine/Plugins/FX/Niagara/Source/NiagaraEditor/Public/SNiagaraGraphActionWidget.h
@@ -29,6 +29,8 @@ struct FCreateNiagaraWidgetForActionData
 	TAttribute<FText> HighlightText;
 	/** True if the widget should be read only - no renaming allowed */
 	bool bIsReadOnly = true;
+	/** If specified, adds a 'favorite' toggle to the widget */
+	TOptional<FName> FavoriteActionsProfileName;
 };
 
 /** Custom widget for GraphActionMenu */
@@ -37,7 +39,7 @@ class SNiagaraGraphActionWidget : public SCompoundWidget
 	public:
 
 	SLATE_BEGIN_ARGS( SNiagaraGraphActionWidget ) {}
-	SLATE_ATTRIBUTE(FText, HighlightText)
+		SLATE_ATTRIBUTE(FText, HighlightText)
 SLATE_END_ARGS()
 
 NIAGARAEDITOR_API void Construct(const FArguments& InArgs, const FCreateWidgetForActionData* InCreateData);
@@ -57,17 +59,25 @@ public:
 
 	SLATE_BEGIN_ARGS( SNiagaraActionWidget )
 		: _bShowTypeIfParameter(true)
-	{} 
+		{} 
 		SLATE_ARGUMENT(bool, bShowTypeIfParameter)
 	SLATE_END_ARGS()
 
 	NIAGARAEDITOR_API void Construct(const FArguments& InArgs, const FCreateNiagaraWidgetForActionData& InCreateData);
 	//virtual FReply OnMouseButtonDown( const FGeometry& MyGeometry, const FPointerEvent& MouseEvent ) override;
 
+private:
+	TSharedRef<SWidget> CreateFavoriteActionWidget();
+	FReply OnFavoriteStateChanged();
+	ECheckBoxState OnGetIsFavorite() const;
+	EVisibility OnShouldShowFavoriteButton() const;
+	const FSlateBrush* GetFavoriteBrush() const;
+	FText GetFavoriteButtonTooltipText() const;
 private:
 	/** The item that we want to display with this widget */
 	TSharedPtr<struct FNiagaraMenuAction_Generic> ActionPtr;
 	/** Delegate executed when mouse button goes down */
 	FOnActionClicked MouseButtonDownDelegate;
 	FNiagaraActionSourceData SourceData;
+	TOptional<FName> FavoriteProfileName;
 };
diff --git a/Engine/Plugins/FX/Niagara/Source/NiagaraEditor/Public/ViewModels/Stack/INiagaraStackItemGroupAddUtilities.h b/Engine/Plugins/FX/Niagara/Source/NiagaraEditor/Public/ViewModels/Stack/INiagaraStackItemGroupAddUtilities.h
--- a/Engine/Plugins/FX/Niagara/Source/NiagaraEditor/Public/ViewModels/Stack/INiagaraStackItemGroupAddUtilities.h
+++ b/Engine/Plugins/FX/Niagara/Source/NiagaraEditor/Public/ViewModels/Stack/INiagaraStackItemGroupAddUtilities.h
@@ -7,6 +7,8 @@
 class INiagaraStackItemGroupAddAction
 {
 public:
+	NIAGARAEDITOR_API static const FName FavoriteActionsProfile;
+	
 	/** Gets the category for this action. */
 	virtual TArray<FString> GetCategories() const = 0;
 
@@ -18,24 +20,15 @@ public:
 
 	/** Gets a space separated string of keywords which expose additional search terms for this action. */
 	virtual FText GetKeywords() const = 0;
-
-	/** Gets a bool that indicates whether this action is suggested or not */
-	virtual bool GetSuggested() const
-	{
-		return false;
-	}
+	
+	/** Required for favorites feature. */
+	virtual TOptional<FNiagaraFavoritesActionData> GetFavoritesData() const = 0;
 
 	/** Indicates if this action represents a library action */
-	virtual bool IsInLibrary() const
-	{
-		return true;
-	}
+	virtual bool IsInLibrary() const = 0;
 
 	/** Gets a source data struct that is useful to display additional information about an action */
-	virtual FNiagaraActionSourceData GetSourceData() const
-	{
-		return FNiagaraActionSourceData();
-	}
+	virtual FNiagaraActionSourceData GetSourceData() const = 0;
 
 	virtual ~INiagaraStackItemGroupAddAction() { }
 };
diff --git a/Engine/Plugins/FX/Niagara/Source/NiagaraEditor/Public/ViewModels/Stack/NiagaraStackEmitterPropertiesGroup.h b/Engine/Plugins/FX/Niagara/Source/NiagaraEditor/Public/ViewModels/Stack/NiagaraStackEmitterPropertiesGroup.h
--- a/Engine/Plugins/FX/Niagara/Source/NiagaraEditor/Public/ViewModels/Stack/NiagaraStackEmitterPropertiesGroup.h
+++ b/Engine/Plugins/FX/Niagara/Source/NiagaraEditor/Public/ViewModels/Stack/NiagaraStackEmitterPropertiesGroup.h
@@ -34,6 +34,9 @@ private:
 		virtual FText GetDisplayName() const override {	return DisplayName;	}
 		virtual FText GetDescription() const override {	return Description;	}
 		virtual FText GetKeywords() const override { return Keywords; }
+		virtual TOptional<FNiagaraFavoritesActionData> GetFavoritesData() const override { return TOptional<FNiagaraFavoritesActionData>(); }
+		virtual bool IsInLibrary() const override { return true; }
+		virtual FNiagaraActionSourceData GetSourceData() const override { return FNiagaraActionSourceData(); }
 
 		EStageAddMode AddMode;
 		UClass* SimulationStageClass;
diff --git a/Engine/Plugins/FX/Niagara/Source/NiagaraEditor/Public/ViewModels/Stack/NiagaraStackFunctionInput.h b/Engine/Plugins/FX/Niagara/Source/NiagaraEditor/Public/ViewModels/Stack/NiagaraStackFunctionInput.h
--- a/Engine/Plugins/FX/Niagara/Source/NiagaraEditor/Public/ViewModels/Stack/NiagaraStackFunctionInput.h
+++ b/Engine/Plugins/FX/Niagara/Source/NiagaraEditor/Public/ViewModels/Stack/NiagaraStackFunctionInput.h
@@ -59,6 +59,25 @@ public:
 		None
 	};
 
+	struct NIAGARAEDITOR_API FNiagaraAvailableParameterInfo
+	{
+		FNiagaraAvailableParameterInfo(FNiagaraVariableBase InVariableBase);
+		FNiagaraAvailableParameterInfo(FNiagaraVariableBase InVariableBase, UNiagaraScript* InConversionScript);
+
+		/** The variable this info struct is about. */
+		FNiagaraVariableBase Variable;
+		/** If this is set, the implication is that the variable can only be applied to this stack function input through conversion. */
+		UNiagaraScript* ConversionScript = nullptr;
+		/** If this variable is defined in a parameter definition asset, the data is included here for further processing. */
+		UNiagaraScriptVariable* MatchingParameterDefinition = nullptr;
+
+		/** We use the variable as identifier so we can use array operations such as Contains or Find based on just the variable. */
+		bool operator==(const FNiagaraAvailableParameterInfo& Other) const
+		{
+			return Variable == Other.Variable;
+		}
+	};
+	
 	DECLARE_MULTICAST_DELEGATE(FOnValueChanged);
 
 public:
@@ -154,7 +173,13 @@ public:
 	NIAGARAEDITOR_API void SetLinkedParameterValueViaConversionScript(const FNiagaraVariableBase& InParameter, UNiagaraScript& ConversionScript);
 
 	/** Gets the current set of available parameter handles which can be assigned to this input. Optionally returns possible conversion scripts. */
-	NIAGARAEDITOR_API void GetAvailableParameters(TArray<FNiagaraVariableBase>& OutAvailableParameters, TMap<FNiagaraVariableBase, UNiagaraScript*>& OutAvailableParametersWithConversions, bool bIncludeConversionScripts = true) const;
+	struct NIAGARAEDITOR_API FGetAvailableParameterArgs
+	{
+		bool bIncludeConversionScripts = true;
+		bool bIncludeParameterDefinitions = false;
+	};
+	
+	NIAGARAEDITOR_API void GetAvailableParameters(TSet<FNiagaraAvailableParameterInfo>& OutAvailableParameters, FGetAvailableParameterArgs Arguments) const;
 
 	/** Gets the function node form the script graph if the current value mode is DefaultFunction. */
 	NIAGARAEDITOR_API UNiagaraNodeFunctionCall* GetDefaultFunctionNode() const;
@@ -564,3 +589,8 @@ private:
 	// The value of the change id for the called graph of the function call which owns this input as of the last refresh.
 	TOptional<FGuid> LastFunctionGraphChangeId;
 };
+
+FORCEINLINE uint32 GetTypeHash(const UNiagaraStackFunctionInput::FNiagaraAvailableParameterInfo& ParameterInfo)
+{
+	return GetTypeHash(ParameterInfo.Variable);
+}
diff --git a/Engine/Plugins/FX/Niagara/Source/NiagaraEditor/Public/Widgets/SItemSelector.h b/Engine/Plugins/FX/Niagara/Source/NiagaraEditor/Public/Widgets/SItemSelector.h
--- a/Engine/Plugins/FX/Niagara/Source/NiagaraEditor/Public/Widgets/SItemSelector.h
+++ b/Engine/Plugins/FX/Niagara/Source/NiagaraEditor/Public/Widgets/SItemSelector.h
@@ -75,6 +75,7 @@ public:
 	DECLARE_DELEGATE_RetVal_OneParam(const CategoryKeyType&, FOnGetKeyForCategory, const CategoryType& /* Category */);
 	DECLARE_DELEGATE_RetVal_OneParam(const SectionKeyType&, FOnGetKeyForSection, const SectionType& /* Section */);
 	DECLARE_DELEGATE_OneParam(FOnSuggestionUpdated, int32 /* Suggestion Index*/)
+	DECLARE_DELEGATE_TwoParams(FOnItemRowHoverEvent, const ItemType&, bool /* bHover */)
 
 	DECLARE_DELEGATE_RetVal_OneParam(bool, FOnCategoryPassesFilter, const CategoryType& /* CategoryA */);
 
@@ -224,6 +225,8 @@ public:
 		/** Whether we want to expand the tree initially or not. */
 		SLATE_ARGUMENT(bool, ExpandInitially)
 
+		SLATE_EVENT(FOnItemRowHoverEvent, OnItemRowHoverEvent)
+		
 		/** An optional delegate called when initially creating this and ExpandInitially is false. */
 
 		/** Slot for additional widget content to go adjacent to right of the search box. */
@@ -1359,21 +1362,44 @@ private:
 			SLATE_ARGUMENT(bool, ShowSelection)
 			SLATE_ARGUMENT(FMargin, Padding)
 			SLATE_EVENT(FOnDragDetected, OnDragDetected)
+			SLATE_EVENT(FOnItemRowHoverEvent, OnItemRowHoverEventDelegate)
+			SLATE_ARGUMENT(TSharedPtr<FItemSelectorItemContainerViewModel>, ItemViewModel)
 		SLATE_END_ARGS();
 
 		void Construct(const FArguments& InArgs, const TSharedRef<STableViewBase>& OwnerTree)
 		{
+			ItemViewModelWeak = InArgs._ItemViewModel;
+			
 			typename SItemSelectorTableRow::FArguments Arguments;
 			Arguments = Arguments[InArgs._Content.Widget];
 			Arguments._Style = InArgs._Style;
 			Arguments._ShowSelection = InArgs._ShowSelection;
 			Arguments._Padding = InArgs._Padding;
 			Arguments._OnDragDetected = InArgs._OnDragDetected;
+
+			OnItemRowHoverEventDelegate = InArgs._OnItemRowHoverEventDelegate;
 			SItemSelectorTableRow::Construct(Arguments, OwnerTree);
 		}
+
+		virtual void OnMouseEnter(const FGeometry& InGeometry, const FPointerEvent& PointerEvent) override
+		{
+			SItemSelectorTableRow::OnMouseEnter(InGeometry, PointerEvent);
+
+			OnItemRowHoverEventDelegate.ExecuteIfBound(ItemViewModelWeak.Pin()->GetItem(), true);
+		}
+
+		virtual void OnMouseLeave(const FPointerEvent& PointerEvent) override
+		{
+			SItemSelectorTableRow::OnMouseLeave(PointerEvent);
+
+			OnItemRowHoverEventDelegate.ExecuteIfBound(ItemViewModelWeak.Pin()->GetItem(), false);
+		}
+
+		TWeakPtr<FItemSelectorItemContainerViewModel> ItemViewModelWeak;
+		FOnItemRowHoverEvent OnItemRowHoverEventDelegate;
 	};
 
-	class SItemSelectorItemCategoryTableRow : public SItemSelectorItemTableRow
+	class SItemSelectorItemCategoryTableRow : public SItemSelectorTableRow
 	{
 	public:
 		SLATE_BEGIN_ARGS(SItemSelectorItemCategoryTableRow)
@@ -1388,24 +1414,26 @@ private:
 			SLATE_ARGUMENT(FLinearColor, BorderBackgroundColor)
 			SLATE_ARGUMENT(FMargin, ChildSlotPadding)
 			SLATE_ARGUMENT(FMargin, BorderBackgroundPadding)
+			SLATE_ARGUMENT(TSharedPtr<FItemSelectorItemCategoryViewModel>, CategoryViewModel)
 			SLATE_EVENT(FOnGetCategoryBackgroundImage, OnGetCategoryBackgroundImage)
 		SLATE_END_ARGS();
 
 		void Construct(const FArguments& InArgs, const TSharedRef<STableViewBase>& OwnerTree)
 		{
+			CategoryViewModel = InArgs._CategoryViewModel;
 			OnGetCategoryBackgroundImage = InArgs._OnGetCategoryBackgroundImage;
 
-			typename SItemSelectorItemTableRow::FArguments Arguments;
+			typename SItemSelectorTableRow::FArguments Arguments;
 			Arguments = Arguments[InArgs._Content.Widget];
 			Arguments._Style = InArgs._Style;
 			Arguments._ShowSelection = InArgs._ShowSelection;
 			Arguments._Padding = InArgs._Padding;
-			SItemSelectorItemTableRow::Construct(Arguments, OwnerTree);
+			SItemSelectorTableRow::Construct(Arguments, OwnerTree);
 
 			// If category background images are used, construct a border to display them.
 			if (OnGetCategoryBackgroundImage.IsBound())
 			{
-				SItemSelectorItemTableRow::ChildSlot
+				SItemSelectorTableRow::ChildSlot
 				.Padding(InArgs._ChildSlotPadding)
 				[
 					SAssignNew(ContentBorder, SBorder)
@@ -1444,8 +1472,40 @@ private:
 	private:
 		FOnGetCategoryBackgroundImage OnGetCategoryBackgroundImage;
 		TSharedPtr<SBorder> ContentBorder;
+		TWeakPtr<FItemSelectorItemCategoryViewModel> CategoryViewModel;
 	};
 
+	class SItemSelectorItemSectionTableRow : public SItemSelectorTableRow
+	{
+		SLATE_BEGIN_ARGS(SItemSelectorItemSectionTableRow)
+			: _Style(&FCoreStyle::Get().GetWidgetStyle<FTableRowStyle>("TableView.Row"))
+			, _ShowSelection(false)
+			{}
+			SLATE_DEFAULT_SLOT(typename SItemSelectorItemSectionTableRow::FArguments, Content)
+			SLATE_STYLE_ARGUMENT(FTableRowStyle, Style)
+			SLATE_ARGUMENT(bool, ShowSelection)
+			SLATE_ARGUMENT(FMargin, Padding)
+			SLATE_EVENT(FOnDragDetected, OnDragDetected)
+			SLATE_EVENT(FOnItemRowHoverEvent, OnItemRowHoverEventDelegate)
+			SLATE_ARGUMENT(TSharedPtr<FSectionViewModel>, SectionViewModel)
+		SLATE_END_ARGS()
+
+		void Construct(const FArguments& InArgs, const TSharedRef<STableViewBase>& OwnerTree)
+		{
+			SectionViewModelWeak = InArgs._SectionViewModel;
+			
+			typename SItemSelectorTableRow::FArguments Arguments;
+			Arguments = Arguments[InArgs._Content.Widget];
+			Arguments._Style = InArgs._Style;
+			Arguments._ShowSelection = InArgs._ShowSelection;
+			Arguments._Padding = InArgs._Padding;
+			Arguments._OnDragDetected = InArgs._OnDragDetected;
+			
+			SItemSelectorTableRow::Construct(Arguments, OwnerTree);
+		}
+
+		TWeakPtr<FSectionViewModel> SectionViewModelWeak;
+	};
 public:
 	void Construct(const FArguments& InArgs)
 	{
@@ -1494,6 +1554,7 @@ public:
 		OnGetKeyForCategory = InArgs._OnGetKeyForCategory;
 		OnGetKeyForSection = InArgs._OnGetKeyForSection;
 		SearchBoxAdjacentContentWidget = InArgs._SearchBoxAdjacentContent.Widget;
+		OnItemRowHoverEventDelegate = InArgs._OnItemRowHoverEvent;
 
 		OnSuggestionUpdated = FOnSuggestionUpdated::CreateSP(this, &SItemSelector::OnSuggestionChanged);
 
@@ -1996,6 +2057,7 @@ private:
 				.ChildSlotPadding(CategoryChildSlotPadding)
 				.BorderBackgroundPadding(CategoryBorderBackgroundPadding)
 				.OnGetCategoryBackgroundImage(OnGetCategoryBackgroundImage)
+				.CategoryViewModel(ItemCategoryViewModel)
 				[
 					OnGenerateWidgetForCategory.Execute(ItemCategoryViewModel->GetCategory())
 				];
@@ -2008,6 +2070,8 @@ private:
 				.ShowSelection(true)
 				.Padding(Margin)
 				.OnDragDetected(this, &SItemSelector::OnItemDragDetected)
+				.OnItemRowHoverEventDelegate(OnItemRowHoverEventDelegate)
+				.ItemViewModel(ItemViewModel)
 				[
 					OnGenerateWidgetForItem.Execute(ItemViewModel->GetItem())
 				];
@@ -2015,10 +2079,11 @@ private:
 		case EItemSelectorItemViewModelType::Section:
 		{
 			TSharedRef<FSectionViewModel> SectionViewModel = StaticCastSharedRef<FSectionViewModel>(Item);
-			return SNew(SItemSelectorItemTableRow, OwnerTable)
+			return SNew(SItemSelectorItemSectionTableRow, OwnerTable)
 				.Style(SectionRowStyle)
 				.ShowSelection(false)
 				.Padding(Margin)
+				.SectionViewModel(SectionViewModel)
 				[
 					OnGenerateWidgetForSection.Execute(SectionViewModel->GetSection())
 				];
@@ -2270,7 +2335,8 @@ private:
 	FOnGetKeyForItem OnGetKeyForItem;
 	FOnGetKeyForCategory OnGetKeyForCategory;
 	FOnGetKeyForSection OnGetKeyForSection;
-
+	FOnItemRowHoverEvent OnItemRowHoverEventDelegate;
+	
 	TSharedPtr<FItemSelectorViewModel> ViewModelUtilities;
 	TSharedPtr<STreeView<TSharedRef<FItemSelectorItemViewModel>>> ItemTree;
 	FOnSuggestionUpdated OnSuggestionUpdated;
diff --git a/Engine/Plugins/FX/Niagara/Source/NiagaraEditor/Public/Widgets/SNiagaraGraphActionMenu.h b/Engine/Plugins/FX/Niagara/Source/NiagaraEditor/Public/Widgets/SNiagaraGraphActionMenu.h
--- a/Engine/Plugins/FX/Niagara/Source/NiagaraEditor/Public/Widgets/SNiagaraGraphActionMenu.h
+++ b/Engine/Plugins/FX/Niagara/Source/NiagaraEditor/Public/Widgets/SNiagaraGraphActionMenu.h
@@ -35,6 +35,10 @@ struct FNiagaraGraphActionMenuFilter
 class SNiagaraGraphActionMenu : public SBorder
 {
 public:
+	/** The name of the profile used to save favorite action data. Used with UNiagaraFavoriteActionsConfig.
+	 * Each action gets an identifier, and the user-authored favorites data is saved per profile. */
+	static const FName FavoriteActionsProfile;
+	
 	SLATE_BEGIN_ARGS( SNiagaraGraphActionMenu )
 		: _GraphObj( static_cast<UEdGraph*>(NULL) )
 		, _NewNodePosition( FVector2D::ZeroVector )
diff --git a/Engine/Plugins/FX/Niagara/Source/NiagaraEditor/Public/Widgets/Wizard/SNiagaraModuleWizard.h b/Engine/Plugins/FX/Niagara/Source/NiagaraEditor/Public/Widgets/Wizard/SNiagaraModuleWizard.h
--- a/Engine/Plugins/FX/Niagara/Source/NiagaraEditor/Public/Widgets/Wizard/SNiagaraModuleWizard.h
+++ b/Engine/Plugins/FX/Niagara/Source/NiagaraEditor/Public/Widgets/Wizard/SNiagaraModuleWizard.h
@@ -57,6 +57,9 @@ namespace UE::Niagara::Wizard
 		// called after the generated module is added to the stack - returning true will apply the modified clipboard content to the module inputs
 		virtual bool UpdateModuleInputs(UNiagaraClipboardContent* NewModule, const TArray<const UNiagaraNodeFunctionCall*>& PreviouslyCreatedModules);
 
+		// a unique name to identify this wizard with. Used in menus.
+		virtual FName GetIdentifier() const { return "ModuleWizardModel"; }
+		
 		// target usage of the generated module
 		ENiagaraScriptUsage TargetUsage;
 		
diff --git a/Engine/Plugins/FX/Niagara/Source/NiagaraEditorWidgets/Private/Stack/SNiagaraStackFunctionInputValue.cpp b/Engine/Plugins/FX/Niagara/Source/NiagaraEditorWidgets/Private/Stack/SNiagaraStackFunctionInputValue.cpp
--- a/Engine/Plugins/FX/Niagara/Source/NiagaraEditorWidgets/Private/Stack/SNiagaraStackFunctionInputValue.cpp
+++ b/Engine/Plugins/FX/Niagara/Source/NiagaraEditorWidgets/Private/Stack/SNiagaraStackFunctionInputValue.cpp
@@ -20,10 +20,12 @@
 #include "NiagaraNodeAssignment.h"
 #include "NiagaraNodeFunctionCall.h"
 #include "NiagaraParameterCollection.h"
+#include "NiagaraParameterDefinitions.h"
 #include "NiagaraStackCommandContext.h"
 #include "NiagaraScriptVariable.h"
 #include "NiagaraSettings.h"
 #include "NiagaraSystem.h"
+#include "NiagaraTrace.h"
 #include "PropertyCustomizationHelpers.h"
 #include "PropertyEditorModule.h"
 #include "SDropTarget.h"
@@ -31,6 +33,7 @@
 #include "SNiagaraParameterDropTarget.h"
 #include "SNiagaraParameterEditor.h"
 #include "ScopedTransaction.h"
+#include "Config/NiagaraFavoriteActionsConfig.h"
 #include "Stack/SNiagaraStackIndent.h"
 #include "Styling/AppStyle.h"
 #include "Subsystems/AssetEditorSubsystem.h"
@@ -62,6 +65,7 @@
 const float TextIconSize = 16;
 
 bool SNiagaraStackFunctionInputValue::bLibraryOnly = true;
+FName SNiagaraStackFunctionInputValue::FavoriteActionsProfile = TEXT("StackFunctionInputActionsProfile");
 
 void SNiagaraStackFunctionInputValue::Construct(const FArguments& InArgs, UNiagaraStackFunctionInput* InFunctionInput)
 {
@@ -849,12 +853,16 @@ TSharedRef<SWidget> SNiagaraStackFunctionInputValue::OnGetAvailableHandleMenu()
                 .OnGenerateWidgetForSection(this, &SNiagaraStackFunctionInputValue::OnGenerateWidgetForSection)
                 .OnGenerateWidgetForCategory(this, &SNiagaraStackFunctionInputValue::OnGenerateWidgetForCategory)
                 .OnGenerateWidgetForItem(this, &SNiagaraStackFunctionInputValue::OnGenerateWidgetForItem)
-                .OnGetItemWeight_Static(&FNiagaraEditorUtilities::GetWeightForItem)
+				.OnGetItemWeight_Lambda([](const TSharedPtr<FNiagaraMenuAction_Generic>& Item, const TArray<FString>& FilterTerms)
+				{
+					return FNiagaraEditorUtilities::GetWeightForItem(Item, SNiagaraStackFunctionInputValue::FavoriteActionsProfile, FilterTerms);
+				})
                 .OnItemActivated(this, &SNiagaraStackFunctionInputValue::OnItemActivated)
                 .AllowMultiselect(false)
                 .OnDoesItemPassCustomFilter(this, &SNiagaraStackFunctionInputValue::DoesItemPassCustomFilter)
                 .ClickActivateMode(EItemSelectorClickActivateMode::SingleClick)
                 .ExpandInitially(false)
+				.OnItemRowHoverEvent(this, &SNiagaraStackFunctionInputValue::OnActionRowHoverEvent)
                 .OnGetSectionData_Lambda([](const ENiagaraMenuSections& Section)
                 {
                     if(Section == ENiagaraMenuSections::Suggested)
@@ -1119,10 +1127,14 @@ TArray<TSharedPtr<FNiagaraMenuAction_Generic>> SNiagaraStackFunctionInputValue::
 			Source = TTuple<EScriptSource, FText>(EScriptSource::Niagara, FText::FromString("Scratch Pad"));
 			bIsInLibrary = true;
 		}
-		
+
+		FNiagaraFavoritesActionData FavoritesActionData;
+		FavoritesActionData.ActionIdentifier.Names.Add(FName(GetPathNameSafe(DynamicInputScript)));
+		FavoritesActionData.bFavoriteByDefault = ScriptData->bSuggested;
+			
 		TSharedPtr<FNiagaraMenuAction_Generic> DynamicInputAction(new FNiagaraMenuAction_Generic(
 			FNiagaraMenuAction_Generic::FOnExecuteAction::CreateStatic(&ReassignDynamicInputScript, FunctionInput, DynamicInputScript),
-			DisplayName, ScriptData->bSuggested ? ENiagaraMenuSections::Suggested : ENiagaraMenuSections::General, {CategoryName.ToString()}, Tooltip, ScriptData->Keywords
+			DisplayName, {CategoryName.ToString()}, FavoritesActionData, Tooltip, ScriptData->Keywords
             ));
 		DynamicInputAction->SourceData = FNiagaraActionSourceData(Source.Key, Source.Value, true);
 		DynamicInputAction->bIsInLibrary = bIsInLibrary;
@@ -1173,12 +1185,16 @@ void SNiagaraStackFunctionInputValue::ShowReassignDynamicInputScriptMenu()
                 .OnGenerateWidgetForSection(this, &SNiagaraStackFunctionInputValue::OnGenerateWidgetForSection)
                 .OnGenerateWidgetForCategory(this, &SNiagaraStackFunctionInputValue::OnGenerateWidgetForCategory)
                 .OnGenerateWidgetForItem(this, &SNiagaraStackFunctionInputValue::OnGenerateWidgetForItem)
-                .OnGetItemWeight_Static(&FNiagaraEditorUtilities::GetWeightForItem)
+				.OnGetItemWeight_Lambda([](const TSharedPtr<FNiagaraMenuAction_Generic>& Item, const TArray<FString>& FilterTerms)
+				 {
+					 return FNiagaraEditorUtilities::GetWeightForItem(Item, SNiagaraStackFunctionInputValue::FavoriteActionsProfile, FilterTerms);
+				 })
                 .OnItemActivated(this, &SNiagaraStackFunctionInputValue::OnItemActivated)
                 .AllowMultiselect(false)
                 .OnDoesItemPassCustomFilter(this, &SNiagaraStackFunctionInputValue::DoesItemPassCustomFilter)
                 .ClickActivateMode(EItemSelectorClickActivateMode::SingleClick)
                 .ExpandInitially(false)
+				.OnItemRowHoverEvent(this, &SNiagaraStackFunctionInputValue::OnActionRowHoverEvent)
                 .OnGetSectionData_Lambda([](const ENiagaraMenuSections& Section)
                 {
                     if(Section == ENiagaraMenuSections::Suggested)
@@ -1221,12 +1237,20 @@ FReply SNiagaraStackFunctionInputValue::ScratchButtonPressed() const
 	return FReply::Unhandled();
 }
 
+void SNiagaraStackFunctionInputValue::OnActionRowHoverEvent(const TSharedPtr<FNiagaraMenuAction_Generic>& ActionNode, bool bIsHovered)
+{
+	ActionNode->bIsHovered = bIsHovered;
+}
+
 TArray<TSharedPtr<FNiagaraMenuAction_Generic>> SNiagaraStackFunctionInputValue::CollectActions()
 {
+	TRACE_CPUPROFILER_EVENT_SCOPE_ON_CHANNEL(SNiagaraStackFunctionInputValue::CollectActions, NiagaraChannel);
+	
 	TArray<TSharedPtr<FNiagaraMenuAction_Generic>> OutAllActions;
 	bool bIsDataInterfaceOrObject = FunctionInput->GetInputType().IsDataInterface() || FunctionInput->GetInputType().IsUObject();
 
 	FNiagaraActionSourceData NiagaraSourceData(EScriptSource::Niagara, FText::FromString(TEXT("Niagara")), true);
+	FNiagaraFavoriteActionsProfile& ActionsProfile = UNiagaraFavoriteActionsConfig::Get()->GetActionsProfile(FavoriteActionsProfile);
 	
 	// Set a local value
 	{
@@ -1239,7 +1263,7 @@ TArray<TSharedPtr<FNiagaraMenuAction_Generic>> SNiagaraStackFunctionInputValue::
 		TSharedPtr<FNiagaraMenuAction_Generic> SetLocalValueAction(new FNiagaraMenuAction_Generic(
 			FNiagaraMenuAction_Generic::FOnExecuteAction::CreateSP(this, &SNiagaraStackFunctionInputValue::SetToLocalValue),
 			FNiagaraMenuAction_Generic::FCanExecuteAction::CreateLambda([=]() { return bCanSetLocalValue; }),
-            DisplayName, ENiagaraMenuSections::General, {}, Tooltip, FText()));
+            DisplayName, {}, {}, Tooltip, FText()));
 		SetLocalValueAction->SourceData = NiagaraSourceData;
 		OutAllActions.Add(SetLocalValueAction);
 	}
@@ -1273,13 +1297,16 @@ TArray<TSharedPtr<FNiagaraMenuAction_Generic>> SNiagaraStackFunctionInputValue::
 				bIsInLibrary = true;
 			}
 			
+			// We construct an ActionIdentifier to check for it in the favorites list
+			FNiagaraFavoritesActionData FavoritesActionData;
+			FavoritesActionData.ActionIdentifier = FNiagaraActionIdentifier({FName(GetPathNameSafe(DynamicInputScript))}, {}); 
+			FavoritesActionData.bFavoriteByDefault = ScriptData->bSuggested;
+			
 			TSharedPtr<FNiagaraMenuAction_Generic> DynamicInputAction(new FNiagaraMenuAction_Generic(
                 FNiagaraMenuAction_Generic::FOnExecuteAction::CreateSP(this, &SNiagaraStackFunctionInputValue::DynamicInputScriptSelected, DynamicInputScript),
-                DisplayName, ScriptData->bSuggested ? ENiagaraMenuSections::Suggested : ENiagaraMenuSections::General, {CategoryName.ToString()}, Tooltip, ScriptData->Keywords));
-			
-			DynamicInputAction->SourceData = FNiagaraActionSourceData(Source.Key, Source.Value, true);
-
-			
+                DisplayName, {CategoryName.ToString()}, FavoritesActionData, Tooltip, ScriptData->Keywords));
+			DynamicInputAction->FavoritesActionData = FavoritesActionData;
+			DynamicInputAction->SourceData = FNiagaraActionSourceData(Source.Key, Source.Value, true);			
 			
 			DynamicInputAction->bIsExperimental = ScriptData->bExperimental;
 			DynamicInputAction->bIsInLibrary = bIsInLibrary;
@@ -1291,15 +1318,25 @@ TArray<TSharedPtr<FNiagaraMenuAction_Generic>> SNiagaraStackFunctionInputValue::
 	const UNiagaraSettings* Settings = GetDefault<UNiagaraSettings>();
 	bool bAllowConversions = ensure(Settings != nullptr) && Settings->bShowConvertibleInputsInStack;
 
-	TArray<FNiagaraVariableBase> AvailableParameters;
-	TMap<FNiagaraVariableBase, UNiagaraScript*> AvailableParametersWithConversion;
-	FunctionInput->GetAvailableParameters(AvailableParameters, AvailableParametersWithConversion, bAllowConversions);
+	TSet<UNiagaraStackFunctionInput::FNiagaraAvailableParameterInfo> AvailableParameterInfos;
+
+	UNiagaraStackFunctionInput::FGetAvailableParameterArgs GetAvailableParameterArgs;
+	GetAvailableParameterArgs.bIncludeConversionScripts = bAllowConversions;
+	GetAvailableParameterArgs.bIncludeParameterDefinitions = true;
+	FunctionInput->GetAvailableParameters(AvailableParameterInfos, GetAvailableParameterArgs);
 
+	// First, we add the inputs that can be directly linked (without a conversion script)
 	const FString RootCategoryName = FString("Link Inputs");
 	const FText MapInputFormat = LOCTEXT("LinkInputFormat", "Link this input to {0}");
-	for (const FNiagaraVariableBase& AvailableParameter : AvailableParameters)
+	for (const UNiagaraStackFunctionInput::FNiagaraAvailableParameterInfo& AvailableParameterInfo : AvailableParameterInfos)
 	{
-		FNiagaraParameterHandle AvailableHandle = FNiagaraParameterHandle(AvailableParameter.GetName());
+		if(AvailableParameterInfo.ConversionScript != nullptr)
+		{
+			continue;
+		}
+		
+		FNiagaraParameterHandle AvailableHandle = FNiagaraParameterHandle(AvailableParameterInfo.Variable.GetName());
+		
 		TArray<FName> HandleParts = AvailableHandle.GetHandleParts();
 		FNiagaraNamespaceMetadata NamespaceMetadata = GetDefault<UNiagaraEditorSettings>()->GetMetaDataForNamespaces(HandleParts);
 		if (NamespaceMetadata.IsValid())
@@ -1309,10 +1346,14 @@ TArray<TSharedPtr<FNiagaraMenuAction_Generic>> SNiagaraStackFunctionInputValue::
 			const FText Category = NamespaceMetadata.DisplayName;
 			const FText DisplayName = FText::FromName(AvailableHandle.GetParameterHandleString());
 			const FText Tooltip = FText::Format(MapInputFormat, FText::FromName(AvailableHandle.GetParameterHandleString()));
+
+			FNiagaraFavoritesActionData FavoritesActionData;
+			FavoritesActionData.ActionIdentifier.Names.Add(FName(DisplayName.ToString()));
+			FavoritesActionData.bFavoriteByDefault = AvailableParameterInfo.MatchingParameterDefinition ? AvailableParameterInfo.MatchingParameterDefinition->Metadata.bSuggested : false;
 			
 			TSharedPtr<FNiagaraMenuAction_Generic> LinkAction(new FNiagaraMenuAction_Generic(
-				FNiagaraMenuAction_Generic::FOnExecuteAction::CreateSP(this, &SNiagaraStackFunctionInputValue::ParameterSelected, AvailableParameter),
-				DisplayName, ENiagaraMenuSections::General, {RootCategoryName, Category.ToString()}, Tooltip, FText()));
+				FNiagaraMenuAction_Generic::FOnExecuteAction::CreateSP(this, &SNiagaraStackFunctionInputValue::ParameterSelected, AvailableParameterInfo.Variable),
+				DisplayName, {RootCategoryName, Category.ToString()}, FavoritesActionData,  Tooltip, FText()));
 			
 			LinkAction->SetParameterVariable(FNiagaraVariable(FunctionInput->GetInputType(), AvailableHandle.GetParameterHandleString()));
 			LinkAction->SourceData = NiagaraSourceData;
@@ -1321,11 +1362,17 @@ TArray<TSharedPtr<FNiagaraMenuAction_Generic>> SNiagaraStackFunctionInputValue::
 		}
 	}
 
+	// Then we add those that can only get added via conversion script. This will only be valid if the plugin setting allows it to
 	const FText ConvertInputFormat = LOCTEXT("ConvertInputFormat", "Link this input to {0} via a conversion script");
-	for (const auto& Entry : AvailableParametersWithConversion)
+	for (const UNiagaraStackFunctionInput::FNiagaraAvailableParameterInfo& AvailableParameterInfo : AvailableParameterInfos)
 	{
-		FNiagaraVariableBase ParameterVariable = Entry.Key;
-		UNiagaraScript* ConversionScript = Entry.Value;
+		if(AvailableParameterInfo.ConversionScript == nullptr)
+		{
+			continue;
+		}
+		
+		FNiagaraVariableBase ParameterVariable = AvailableParameterInfo.Variable;
+		UNiagaraScript* ConversionScript = AvailableParameterInfo.ConversionScript;
 		const FNiagaraParameterHandle& AvailableHandle(ParameterVariable.GetName());
 		TArray<FName> HandleParts = AvailableHandle.GetHandleParts();
 		FNiagaraNamespaceMetadata NamespaceMetadata = GetDefault<UNiagaraEditorSettings>()->GetMetaDataForNamespaces(HandleParts);
@@ -1339,7 +1386,7 @@ TArray<TSharedPtr<FNiagaraMenuAction_Generic>> SNiagaraStackFunctionInputValue::
 
 			TSharedPtr<FNiagaraMenuAction_Generic> LinkAction(new FNiagaraMenuAction_Generic(
 				FNiagaraMenuAction_Generic::FOnExecuteAction::CreateSP(this, &SNiagaraStackFunctionInputValue::ParameterWithConversionSelected, ParameterVariable, ConversionScript),
-				DisplayName, ENiagaraMenuSections::General, {RootCategoryName, Category.ToString()}, Tooltip, FText()));
+				DisplayName, {RootCategoryName, Category.ToString()}, {}, Tooltip, FText()));
 
 			LinkAction->SetParameterVariable(ParameterVariable);
 
@@ -1369,14 +1416,14 @@ TArray<TSharedPtr<FNiagaraMenuAction_Generic>> SNiagaraStackFunctionInputValue::
 		{
 			FNiagaraParameterHandle HandleToRead(AvailableNamespace, InputName);
 			FNiagaraVariableBase ParameterToRead = FNiagaraVariableBase(FunctionInput->GetInputType(), HandleToRead.GetParameterHandleString());
-			bool bIsContained = AvailableParameters.Contains(ParameterToRead);
+			bool bIsContained = AvailableParameterInfos.Contains(ParameterToRead);
 
 			if(bIsContained)
 			{
 				TSet<FName> ExistingNames;
-				for(const FNiagaraVariableBase& AvailableParameter : AvailableParameters)
+				for(const UNiagaraStackFunctionInput::FNiagaraAvailableParameterInfo& AvailableParameter : AvailableParameterInfos)
 				{
-					FNiagaraParameterHandle AvailableHandle = FNiagaraParameterHandle(AvailableParameter.GetName());
+					FNiagaraParameterHandle AvailableHandle = FNiagaraParameterHandle(AvailableParameter.Variable.GetName());
 					ExistingNames.Add(AvailableHandle.GetName());
 				}
 
@@ -1393,7 +1440,7 @@ TArray<TSharedPtr<FNiagaraMenuAction_Generic>> SNiagaraStackFunctionInputValue::
 
 			TSharedPtr<FNiagaraMenuAction_Generic> MakeAction(new FNiagaraMenuAction_Generic(
 				FNiagaraMenuAction_Generic::FOnExecuteAction::CreateSP(this, &SNiagaraStackFunctionInputValue::ParameterSelected, ParameterToRead),         
-		        DisplayName, ENiagaraMenuSections::General, {CategoryName.ToString()}, Tooltip, FText()));
+		        DisplayName, {CategoryName.ToString()}, {}, Tooltip, FText()));
 
 			MakeAction->SourceData = NiagaraSourceData;
 
@@ -1411,7 +1458,7 @@ TArray<TSharedPtr<FNiagaraMenuAction_Generic>> SNiagaraStackFunctionInputValue::
 
 		TSharedPtr<FNiagaraMenuAction_Generic> ExpressionAction(new FNiagaraMenuAction_Generic(
                 FNiagaraMenuAction_Generic::FOnExecuteAction::CreateSP(this, &SNiagaraStackFunctionInputValue::CustomExpressionSelected),
-                DisplayName, ENiagaraMenuSections::General, {}, Tooltip, FText()));
+                DisplayName, {}, {}, Tooltip, FText()));
 
 		ExpressionAction->SourceData = NiagaraSourceData;
 
@@ -1428,7 +1475,7 @@ TArray<TSharedPtr<FNiagaraMenuAction_Generic>> SNiagaraStackFunctionInputValue::
 
 		TSharedPtr<FNiagaraMenuAction_Generic> CreateScratchAction(new FNiagaraMenuAction_Generic(
            FNiagaraMenuAction_Generic::FOnExecuteAction::CreateSP(this, &SNiagaraStackFunctionInputValue::CreateScratchSelected),
-           DisplayName, ENiagaraMenuSections::General, {}, Tooltip, FText()));
+           DisplayName, {}, {}, Tooltip, FText()));
 
 		CreateScratchAction->SourceData = NiagaraSourceData;
 
@@ -1443,7 +1490,7 @@ TArray<TSharedPtr<FNiagaraMenuAction_Generic>> SNiagaraStackFunctionInputValue::
 		TSharedPtr<FNiagaraMenuAction_Generic> DeleteInputAction(new FNiagaraMenuAction_Generic(
                 FNiagaraMenuAction_Generic::FOnExecuteAction::CreateUObject(FunctionInput, &UNiagaraStackFunctionInput::DeleteInput),
                 FNiagaraMenuAction_Generic::FCanExecuteAction::CreateUObject(FunctionInput, &UNiagaraStackFunctionInput::CanDeleteInput),
-                DisplayName, ENiagaraMenuSections::General, {}, Tooltip, FText()));
+                DisplayName, {}, {}, Tooltip, FText()));
 
 		DeleteInputAction->SourceData = NiagaraSourceData;
 		OutAllActions.Add(DeleteInputAction);
@@ -1461,12 +1508,17 @@ TArray<FString> SNiagaraStackFunctionInputValue::OnGetCategoriesForItem(
 TArray<ENiagaraMenuSections> SNiagaraStackFunctionInputValue::OnGetSectionsForItem(
 	const TSharedPtr<FNiagaraMenuAction_Generic>& Item)
 {
-	if(Item->Section == ENiagaraMenuSections::Suggested)
+	TArray<ENiagaraMenuSections> Sections{ENiagaraMenuSections::General};
+	
+	if(Item->FavoritesActionData.IsSet())
 	{
-		return { ENiagaraMenuSections::General, ENiagaraMenuSections::Suggested };
+		FNiagaraFavoriteActionsProfile& ActionsProfile = UNiagaraFavoriteActionsConfig::Get()->GetActionsProfile(SNiagaraStackFunctionInputValue::FavoriteActionsProfile);
+		if(ActionsProfile.IsFavorite(Item->FavoritesActionData.GetValue()))
+		{
+			Sections.Add(ENiagaraMenuSections::Suggested);
+		}
 	}
-		
-	return {Item->Section};
+	return Sections;
 }
 
 bool SNiagaraStackFunctionInputValue::OnCompareSectionsForEquality(const ENiagaraMenuSections& SectionA,
@@ -1528,7 +1580,9 @@ TSharedRef<SWidget> SNiagaraStackFunctionInputValue::OnGenerateWidgetForItem(
 {
 	FCreateNiagaraWidgetForActionData ActionData(Item);
 	ActionData.HighlightText = TAttribute<FText>::Create(TAttribute<FText>::FGetter::CreateRaw(this, &SNiagaraStackFunctionInputValue::GetFilterText));
-	return SNew(SNiagaraActionWidget, ActionData).bShowTypeIfParameter(false);
+	ActionData.FavoriteActionsProfileName = FavoriteActionsProfile;
+	return SNew(SNiagaraActionWidget, ActionData)
+		.bShowTypeIfParameter(false);
 }
 
 bool SNiagaraStackFunctionInputValue::DoesItemPassCustomFilter(const TSharedPtr<FNiagaraMenuAction_Generic>& Item)
diff --git a/Engine/Plugins/FX/Niagara/Source/NiagaraEditorWidgets/Private/Stack/SNiagaraStackFunctionInputValue.h b/Engine/Plugins/FX/Niagara/Source/NiagaraEditorWidgets/Private/Stack/SNiagaraStackFunctionInputValue.h
--- a/Engine/Plugins/FX/Niagara/Source/NiagaraEditorWidgets/Private/Stack/SNiagaraStackFunctionInputValue.h
+++ b/Engine/Plugins/FX/Niagara/Source/NiagaraEditorWidgets/Private/Stack/SNiagaraStackFunctionInputValue.h
@@ -199,9 +199,11 @@ private:
 
 	TAttribute<EVisibility> CompactActionMenuButtonVisibilityAttribute;
 	TSharedPtr<FNiagaraStackCommandContext> StackCommandContext;
-
+	
 	static bool bLibraryOnly;
 
+	static FName FavoriteActionsProfile;
+
 private:
 	TArray<TSharedPtr<FNiagaraMenuAction_Generic>> CollectActions();
 	TArray<FString> OnGetCategoriesForItem(const TSharedPtr<FNiagaraMenuAction_Generic>& Item);
@@ -217,6 +219,7 @@ private:
 	TSharedRef<SWidget> OnGenerateWidgetForItem(const TSharedPtr<FNiagaraMenuAction_Generic>& Item);
 	bool DoesItemPassCustomFilter(const TSharedPtr<FNiagaraMenuAction_Generic>& Item);
 	void OnItemActivated(const TSharedPtr<FNiagaraMenuAction_Generic>& Item);
+	void OnActionRowHoverEvent(const TSharedPtr<FNiagaraMenuAction_Generic>& ActionNode, bool bIsHovered);
 
 	void TriggerRefresh(const TMap<EScriptSource, bool>& SourceState);
 
diff --git a/Engine/Plugins/FX/Niagara/Source/NiagaraEditorWidgets/Private/Stack/SNiagaraStackItemGroupAddMenu.cpp b/Engine/Plugins/FX/Niagara/Source/NiagaraEditorWidgets/Private/Stack/SNiagaraStackItemGroupAddMenu.cpp
--- a/Engine/Plugins/FX/Niagara/Source/NiagaraEditorWidgets/Private/Stack/SNiagaraStackItemGroupAddMenu.cpp
+++ b/Engine/Plugins/FX/Niagara/Source/NiagaraEditorWidgets/Private/Stack/SNiagaraStackItemGroupAddMenu.cpp
@@ -7,6 +7,7 @@
 #include "ViewModels/Stack/INiagaraStackItemGroupAddUtilities.h"
 #include "SGraphActionMenu.h"
 #include "SNiagaraGraphActionWidget.h"
+#include "Config/NiagaraFavoriteActionsConfig.h"
 #include "Framework/Application/SlateApplication.h"
 #include "Widgets/SItemSelector.h"
 #include "Widgets/Text/SRichTextBlock.h"
@@ -70,12 +71,16 @@ void SNiagaraStackItemGroupAddMenu::Construct(const FArguments& InArgs, UNiagara
 	                .OnGenerateWidgetForSection(this, &SNiagaraStackItemGroupAddMenu::OnGenerateWidgetForSection)
 	                .OnGenerateWidgetForCategory(this, &SNiagaraStackItemGroupAddMenu::OnGenerateWidgetForCategory)
 	                .OnGenerateWidgetForItem(this, &SNiagaraStackItemGroupAddMenu::OnGenerateWidgetForItem)
-	                .OnGetItemWeight_Static(&FNiagaraEditorUtilities::GetWeightForItem)
+	                .OnGetItemWeight_Lambda([](const TSharedPtr<FNiagaraMenuAction_Generic>& Item, const TArray<FString>& FilterTerms)
+	                {
+	                	return FNiagaraEditorUtilities::GetWeightForItem(Item, INiagaraStackItemGroupAddAction::FavoriteActionsProfile, FilterTerms);
+	                })
 	                .OnItemActivated(this, &SNiagaraStackItemGroupAddMenu::OnItemActivated)
 	                .AllowMultiselect(false)
 	                .OnDoesItemPassCustomFilter(this, &SNiagaraStackItemGroupAddMenu::DoesItemPassCustomFilter)
 	                .ClickActivateMode(EItemSelectorClickActivateMode::SingleClick)
 	                .ExpandInitially(AddUtilities->GetAutoExpandAddActions())
+					.OnItemRowHoverEvent(this, &SNiagaraStackItemGroupAddMenu::OnItemRowHoverEvent)
 	                .OnGetSectionData_Lambda([](const ENiagaraMenuSections& Section)
 	                {
 	                    if(Section == ENiagaraMenuSections::Suggested)
@@ -104,17 +109,21 @@ TArray<TSharedPtr<FNiagaraMenuAction_Generic>> SNiagaraStackItemGroupAddMenu::Co
 	AddOptions.bIncludeDeprecated = false;
 	AddOptions.bIncludeNonLibrary = true;
 
+	// We translate group add actions into the generic menu actions here
+	
 	TArray<TSharedRef<INiagaraStackItemGroupAddAction>> AddActions;
 	AddUtilities->GenerateAddActions(AddActions, AddOptions);
 
+	// For each group action, we parse its data and match data formats to the generic action
 	for (TSharedRef<INiagaraStackItemGroupAddAction> AddAction : AddActions)
-	{
+	{		
 		TSharedPtr<FNiagaraMenuAction_Generic> NewNodeAction(
             new FNiagaraMenuAction_Generic(
             FNiagaraMenuAction_Generic::FOnExecuteAction::CreateRaw(AddUtilities, &INiagaraStackItemGroupAddUtilities::ExecuteAddAction, AddAction, InsertIndex),
-            AddAction->GetDisplayName(), AddAction->GetSuggested() ? ENiagaraMenuSections::Suggested : ENiagaraMenuSections::General, AddAction->GetCategories(), AddAction->GetDescription(), AddAction->GetKeywords()));
+            AddAction->GetDisplayName(), AddAction->GetCategories(), AddAction->GetFavoritesData(), AddAction->GetDescription(), AddAction->GetKeywords()));
 		NewNodeAction->SourceData = AddAction->GetSourceData();
 		NewNodeAction->bIsInLibrary = AddAction->IsInLibrary();
+		
 		OutActions.Add(NewNodeAction);
 	}
 
@@ -137,6 +146,11 @@ void SNiagaraStackItemGroupAddMenu::OnItemActivated(const TSharedPtr<FNiagaraMen
 	}	
 }
 
+void SNiagaraStackItemGroupAddMenu::OnItemRowHoverEvent(const TSharedPtr<FNiagaraMenuAction_Generic>& Item, bool bHover) const
+{
+	Item->bIsHovered = bHover;
+}
+
 void SNiagaraStackItemGroupAddMenu::TriggerRefresh(const TMap<EScriptSource, bool>& SourceState)
 {
 	ActionSelector->RefreshAllCurrentItems();
@@ -176,13 +190,18 @@ TArray<FString> SNiagaraStackItemGroupAddMenu::OnGetCategoriesForItem(
 TArray<ENiagaraMenuSections> SNiagaraStackItemGroupAddMenu::OnGetSectionsForItem(
 	const TSharedPtr<FNiagaraMenuAction_Generic>& Item)
 {
+	TArray<ENiagaraMenuSections> Sections{ENiagaraMenuSections::General};
 	
-	if(Item->Section == ENiagaraMenuSections::Suggested)
+	if(Item->FavoritesActionData.IsSet())
 	{
-		return { ENiagaraMenuSections::General, ENiagaraMenuSections::Suggested };
+		FNiagaraFavoriteActionsProfile& ActionsProfile = UNiagaraFavoriteActionsConfig::Get()->GetActionsProfile(INiagaraStackItemGroupAddAction::FavoriteActionsProfile);
+		if(ActionsProfile.IsFavorite(Item->FavoritesActionData.GetValue()))
+		{
+			Sections.Add(ENiagaraMenuSections::Suggested);
+		}
 	}
 		
-	return {Item->Section};
+	return Sections;
 }
 
 bool SNiagaraStackItemGroupAddMenu::OnCompareSectionsForEquality(const ENiagaraMenuSections& SectionA,
@@ -243,6 +262,7 @@ TSharedRef<SWidget> SNiagaraStackItemGroupAddMenu::OnGenerateWidgetForItem(
 	const TSharedPtr<FNiagaraMenuAction_Generic>& Item)
 {
 	FCreateNiagaraWidgetForActionData ActionData(Item);
+	ActionData.FavoriteActionsProfileName = INiagaraStackItemGroupAddAction::FavoriteActionsProfile;
 	ActionData.HighlightText = TAttribute<FText>::Create(TAttribute<FText>::FGetter::CreateRaw(this, &SNiagaraStackItemGroupAddMenu::GetFilterText));
 	return SNew(SNiagaraActionWidget, ActionData);
 }
diff --git a/Engine/Plugins/FX/Niagara/Source/NiagaraEditorWidgets/Private/Stack/SNiagaraStackItemGroupAddMenu.h b/Engine/Plugins/FX/Niagara/Source/NiagaraEditorWidgets/Private/Stack/SNiagaraStackItemGroupAddMenu.h
--- a/Engine/Plugins/FX/Niagara/Source/NiagaraEditorWidgets/Private/Stack/SNiagaraStackItemGroupAddMenu.h
+++ b/Engine/Plugins/FX/Niagara/Source/NiagaraEditorWidgets/Private/Stack/SNiagaraStackItemGroupAddMenu.h
@@ -57,6 +57,7 @@ private:
 	TSharedRef<SWidget> OnGenerateWidgetForItem(const TSharedPtr<FNiagaraMenuAction_Generic>& Item);
 	bool DoesItemPassCustomFilter(const TSharedPtr<FNiagaraMenuAction_Generic>& Item);
 	void OnItemActivated(const TSharedPtr<FNiagaraMenuAction_Generic>& Item);
+	void OnItemRowHoverEvent(const TSharedPtr<FNiagaraMenuAction_Generic>& Item, bool bHover) const;
 
 	void TriggerRefresh(const TMap<EScriptSource, bool>& SourceState);
 
diff --git a/Engine/Source/Runtime/SlateCore/Private/Styling/StarshipCoreStyle.cpp b/Engine/Source/Runtime/SlateCore/Private/Styling/StarshipCoreStyle.cpp
--- a/Engine/Source/Runtime/SlateCore/Private/Styling/StarshipCoreStyle.cpp
+++ b/Engine/Source/Runtime/SlateCore/Private/Styling/StarshipCoreStyle.cpp
@@ -352,6 +352,8 @@ TSharedRef<ISlateStyle> FStarshipCoreStyle::Create()
 		Style->Set("Icons.Filter", new IMAGE_BRUSH_SVG("Starship/Common/filter", Icon16x16));
 		Style->Set("Icons.AutoFilter", new IMAGE_BRUSH_SVG("Starship/Common/FilterAuto", Icon16x16));
 
+		Style->Set("Icons.Tag", new IMAGE_BRUSH_SVG("Starship/Common/TagSimple_16", Icon16x16));
+		
 		Style->Set("Icons.Lock", new IMAGE_BRUSH_SVG("Starship/Common/lock", Icon16x16));
 		Style->Set("Icons.Unlock", new IMAGE_BRUSH_SVG("Starship/Common/lock-unlocked", Icon16x16));
 
@@ -394,6 +396,7 @@ TSharedRef<ISlateStyle> FStarshipCoreStyle::Create()
 		Style->Set("Icons.Refresh", new IMAGE_BRUSH_SVG("Starship/Common/refresh", Icon16x16));
 
 		Style->Set("Icons.Star", new IMAGE_BRUSH_SVG("Starship/Common/Favorite", Icon16x16));
+		Style->Set("Icons.Star.Outline", new IMAGE_BRUSH_SVG("Starship/Common/FavoriteOutline", Icon16x16));
 		Style->Set("Icons.Link", new IMAGE_BRUSH_SVG("Starship/Common/Linked", Icon16x16));
 		Style->Set("Icons.Unlink", new IMAGE_BRUSH_SVG("Starship/Common/Unlinked", Icon16x16));
 
