Changed raw graph ptr to weakobjectptr for safety reasons

diff --git a/Engine/Plugins/FX/Niagara/Source/NiagaraEditor/Private/ViewModels/NiagaraOverviewGraphViewModel.cpp b/Engine/Plugins/FX/Niagara/Source/NiagaraEditor/Private/ViewModels/NiagaraOverviewGraphViewModel.cpp
--- a/Engine/Plugins/FX/Niagara/Source/NiagaraEditor/Private/ViewModels/NiagaraOverviewGraphViewModel.cpp
+++ b/Engine/Plugins/FX/Niagara/Source/NiagaraEditor/Private/ViewModels/NiagaraOverviewGraphViewModel.cpp
@@ -571,7 +571,7 @@ void FNiagaraOverviewGraphViewModel::PostUndo(bool bSuccess)
 {
 	// This is neccessary to have the graph editor respond correctly to data changes due to undo.
 	// Validate the OverviewGraph exists before notifying as we may be entering PostUndo from a transaction that predates the OverviewGraph.
-	if (OverviewGraph)
+	if (OverviewGraph.IsValid())
 	{
 		OverviewGraph->NotifyGraphChanged();
 	}
diff --git a/Engine/Plugins/FX/Niagara/Source/NiagaraEditor/Public/ViewModels/NiagaraOverviewGraphViewModel.h b/Engine/Plugins/FX/Niagara/Source/NiagaraEditor/Public/ViewModels/NiagaraOverviewGraphViewModel.h
--- a/Engine/Plugins/FX/Niagara/Source/NiagaraEditor/Public/ViewModels/NiagaraOverviewGraphViewModel.h
+++ b/Engine/Plugins/FX/Niagara/Source/NiagaraEditor/Public/ViewModels/NiagaraOverviewGraphViewModel.h
@@ -83,7 +83,7 @@ private:
 	/** The view model to interface with the system being viewed and edited by this view model. */
 	TWeakPtr<FNiagaraSystemViewModel> SystemViewModel;
 
-	UEdGraph* OverviewGraph;
+	TWeakObjectPtr<UEdGraph> OverviewGraph;
 
 	/** The display name for the overview graph. */
 	mutable TOptional<FText> DisplayNameCache;
