- Only mark packages dirty when running a commandlet if fx.Niagara.MarkPackageDirtyWhenNotSynchronized is set to true (defaults to false)

#rnx

diff --git a/Engine/Plugins/FX/Niagara/Source/Niagara/Private/NiagaraSystem.cpp b/Engine/Plugins/FX/Niagara/Source/Niagara/Private/NiagaraSystem.cpp
--- a/Engine/Plugins/FX/Niagara/Source/Niagara/Private/NiagaraSystem.cpp
+++ b/Engine/Plugins/FX/Niagara/Source/Niagara/Private/NiagaraSystem.cpp
@@ -96,6 +96,16 @@ namespace NiagaraSystemPrivate
 	static const FName NAME_FixedBoundsSize("FixedBoundsSize");
 	static const FName NAME_NumEmitters("NumEmitters");
 
+#if WITH_EDITORONLY_DATA
+	bool GMarkPackageDirtyWhenNotSynchronized = false;
+	static FAutoConsoleVariableRef CVarMarkPackageDirtyWhenNotSynchronized(
+		TEXT("fx.Niagara.MarkPackageDirtyWhenNotSynchronized"),
+		GMarkPackageDirtyWhenNotSynchronized,
+		TEXT("When running a commandlet will mark the package dirty if the system or emitter scripts are not synchronized."),
+		ECVF_Default
+	);
+#endif
+
 #if WITH_EDITOR
 	constexpr int32 CookDependenciesArgsVersion = 1;
 	void HashDependenciesForCook(FCbFieldViewIterator Args, UE::Cook::FCookDependencyContext& Context);
@@ -791,7 +801,7 @@ void UNiagaraSystem::UpdateSystemAfterLoad()
 		{
 			RequestCompileStatus = ERequestCompileStatus::RequestPendingOnDemand;
 
-			if (IsRunningCommandlet())
+			if (NiagaraSystemPrivate::GMarkPackageDirtyWhenNotSynchronized && IsRunningCommandlet())
 			{
 				// Call modify here so that the system will resave the compile ids and script vm when running the resave
 				// commandlet. We don't need it for normal post-loading.
