[Backout] - CL41562964
#rnx
Original CL Desc
-----------------------------------------------------------------
Added favorites feature to Niagara Menus.
Works on input dropdowns (dynamic inputs, linked attributes), graph actions and modules.
You can specify favorite by default actions, but are free to favorite or unfavorite any action.
Favorited/suggested actions get a search scorre bonus.

- Also added missing FavoriteOutline icon to starship
- Also added TagSimple icon to starship
- got rid of section specification for niagara actions. The sections were only used for the Suggested and the General section, but the action should not determine whether it is suggested or not. It has been replaced by FNiagaraFavoritesActionData, which includes an action identifier and a "favorite by default" member.
Depending on whether an action is 'favorite by default' or not, toggling favorites writes into one of two locations. Either the explicitly favorited actions, or the explicitly unfavorited actions. This architecture allows us to freely favorite and unfavorite any action while still allowing default favorites.
- Added "OnItemRowHoverEvent" delegate to SItemSelector. This works specifically and only on items and is used to surface information about hover status of an item row.
This is used to display the favorite button only on items that are currently hovered on.
- Added 'Suggested' feature for parameters in Parameter Definition assets, works like Suggested for script assets
- added Particles.Position, Particles.ID and Particles.Color to suggested parameters.

#virtualized

diff --git a/Engine/Plugins/FX/Niagara/Source/Niagara/Public/NiagaraVariableMetaData.h b/Engine/Plugins/FX/Niagara/Source/Niagara/Public/NiagaraVariableMetaData.h
--- a/Engine/Plugins/FX/Niagara/Source/Niagara/Public/NiagaraVariableMetaData.h
+++ b/Engine/Plugins/FX/Niagara/Source/Niagara/Public/NiagaraVariableMetaData.h
@@ -274,10 +274,6 @@ struct FNiagaraVariableMetaData
 	UPROPERTY(EditAnywhere, Category = "Variable", meta = (SkipForCompileHash = "true"))
 	FNiagaraInputConditionMetadata VisibleCondition;
 
-	/** If this variable is suggested by default, it will get prioritized in menus. */
-	UPROPERTY(EditAnywhere, Category = "Variable", meta = (SkipForCompileHash = "true"))
-	bool bSuggested = false;
-	
 	UPROPERTY(EditAnywhere, Category = "Variable", DisplayName = "Property Metadata", meta = (ToolTip = "Property Metadata", SkipForCompileHash = "true"))
 	TMap<FName, FString> PropertyMetaData;
 
diff --git a/Engine/Plugins/FX/Niagara/Source/NiagaraEditor/Private/Config/NiagaraFavoriteActionsConfig.cpp b/Engine/Plugins/FX/Niagara/Source/NiagaraEditor/Private/Config/NiagaraFavoriteActionsConfig.cpp
deleted file mode 100644
diff --git a/Engine/Plugins/FX/Niagara/Source/NiagaraEditor/Private/Customizations/NiagaraScriptVariableCustomization.cpp b/Engine/Plugins/FX/Niagara/Source/NiagaraEditor/Private/Customizations/NiagaraScriptVariableCustomization.cpp
--- a/Engine/Plugins/FX/Niagara/Source/NiagaraEditor/Private/Customizations/NiagaraScriptVariableCustomization.cpp
+++ b/Engine/Plugins/FX/Niagara/Source/NiagaraEditor/Private/Customizations/NiagaraScriptVariableCustomization.cpp
@@ -114,9 +114,6 @@ void FNiagaraScriptVariableDetails::CustomizeDetails(IDetailLayoutBuilder& Detai
 		return;
 	}
 
-	// We hide bSuggested by default unless it is part of the a parameter definition
-	DetailBuilder.HideProperty(GET_MEMBER_NAME_CHECKED(UNiagaraScriptVariable, Metadata.bSuggested));
-	
 	const bool bInLibrary = Variable->GetOuter()->IsA<UNiagaraParameterDefinitions>();
 	if (Variable->GetIsStaticSwitch())
 	{
@@ -227,7 +224,6 @@ void FNiagaraScriptVariableDetails::CustomizeDetailsParameterDefinitionsSynchron
 	UObject* VariableOuter = Variable->GetOuter();
 	const bool bInLibraryAsset = VariableOuter->IsA<UNiagaraParameterDefinitions>();
 
-	TSharedPtr<IPropertyHandle> SuggestedPropertyHandle = DetailBuilder.GetProperty(GET_MEMBER_NAME_CHECKED(UNiagaraScriptVariable, Metadata.bSuggested));
 	/**
 	 *  Default mode is handled as a special case for UNiagaraScriptVariables synchronized with Parameter Libraries. 
 	 * 	Hide the actual default mode and instead edit a library default mode in place which maps to the actual default mode and whether or not to override the library default.
@@ -235,9 +231,7 @@ void FNiagaraScriptVariableDetails::CustomizeDetailsParameterDefinitionsSynchron
 	const TSharedPtr<IPropertyHandle> DefaultModeHandle = DetailBuilder.GetProperty(GET_MEMBER_NAME_CHECKED(UNiagaraScriptVariable, DefaultMode), UNiagaraScriptVariable::StaticClass());
 	DetailBuilder.HideProperty(DefaultModeHandle);
 	if (bInLibraryAsset)
-	{
-		DetailBuilder.AddPropertyToCategory(SuggestedPropertyHandle);
-		
+	{ 
 		LibraryDefaultModeValue = GetLibrarySourcedDefaultModeInitialValue();
 		FDetailWidgetRow& DefaultModeWidget = CategoryBuilder.AddCustomRow(LOCTEXT("DefaultModeFilterText", "Default Mode"));
 		DefaultModeWidget.NameContent()
diff --git a/Engine/Plugins/FX/Niagara/Source/NiagaraEditor/Private/EdGraphSchema_Niagara.cpp b/Engine/Plugins/FX/Niagara/Source/NiagaraEditor/Private/EdGraphSchema_Niagara.cpp
--- a/Engine/Plugins/FX/Niagara/Source/NiagaraEditor/Private/EdGraphSchema_Niagara.cpp
+++ b/Engine/Plugins/FX/Niagara/Source/NiagaraEditor/Private/EdGraphSchema_Niagara.cpp
@@ -48,7 +48,6 @@
 #include "Textures/SlateIcon.h"
 #include "UObject/LinkerLoad.h"
 #include "Algo/Sort.h"
-#include "Config/NiagaraFavoriteActionsConfig.h"
 
 #include UE_INLINE_GENERATED_CPP_BY_NAME(EdGraphSchema_Niagara)
 
@@ -222,15 +221,15 @@ TSharedPtr<FNiagaraAction_NewNode> AddNewNodeMenuAction(
 	TArray<TSharedPtr<FNiagaraAction_NewNode>>& NewActions,
 	UEdGraphNode* InNodeTemplate,
 	const FText& DisplayName,
+	ENiagaraMenuSections Section,
 	TArray<FString> NestedCategories,
-	TOptional<FNiagaraFavoritesActionData> InFavoritesActionData,
 	const FText& Tooltip,
 	FText Keywords = FText(),
 	FNiagaraActionSourceData SourceData = FNiagaraActionSourceData(
 		EScriptSource::Niagara, FText::FromString(TEXT("Niagara")), true))
 {
 	const UNiagaraEditorSettings* NiagaraEditorSettings = GetDefault<UNiagaraEditorSettings>();
-	TSharedPtr<FNiagaraAction_NewNode> NewAction = MakeShared<FNiagaraAction_NewNode>(DisplayName, NestedCategories, InFavoritesActionData, Tooltip, Keywords);
+	TSharedPtr<FNiagaraAction_NewNode> NewAction = MakeShared<FNiagaraAction_NewNode>(DisplayName, Section, NestedCategories, Tooltip, Keywords);
 	if (ensureMsgf(InNodeTemplate == nullptr || NiagaraEditorSettings->IsVisibleClass(InNodeTemplate->GetClass()),
 		TEXT("Can not create a menu action for a node of class %s."), *InNodeTemplate->GetClass()->GetName()))
 	{
@@ -351,7 +350,7 @@ void AddParametersForGraph(TArray<TSharedPtr<FNiagaraAction_NewNode>>& NewAction
 			}
 
 			UNiagaraNodeInput* InputNodeTemplate = NewObject<UNiagaraNodeInput>(OwnerOfTemporaries);
-			AddNewNodeMenuAction(NewActions, InputNodeTemplate, DisplayName, {GraphParameterCategory.ToString()}, {}, FText::GetEmpty());
+			AddNewNodeMenuAction(NewActions, InputNodeTemplate, DisplayName, ENiagaraMenuSections::General, {GraphParameterCategory.ToString()}, FText::GetEmpty());
 			InputNodeTemplate->Input = InputNode->Input;
 			InputNodeTemplate->Usage = InputNode->Usage;
 			InputNodeTemplate->ExposureOptions = InputNode->ExposureOptions;
@@ -486,13 +485,10 @@ TArray<TSharedPtr<FNiagaraAction_NewNode>> UEdGraphSchema_Niagara::GetGraphActio
 
 			if(bAddOp)
 			{
-				FNiagaraFavoritesActionData FavoritesActionData;
-				FavoritesActionData.ActionIdentifier.Names.Add(OpInfo.Name);
-				FavoritesActionData.bFavoriteByDefault = false;
-				
+				// todo suggestion info per op?
 				UNiagaraNodeOp* OpNode = NewObject<UNiagaraNodeOp>(OwnerOfTemporaries);
 				OpNode->OpName = OpInfo.Name;
-				TSharedPtr<FNiagaraAction_NewNode> MathOpAction = AddNewNodeMenuAction(NewActions, OpNode, OpInfo.FriendlyName, {OpInfo.Category.ToString()}, FavoritesActionData, OpInfo.Description,  OpInfo.Keywords);
+				TSharedPtr<FNiagaraAction_NewNode> MathOpAction = AddNewNodeMenuAction(NewActions, OpNode, OpInfo.FriendlyName, ENiagaraMenuSections::General, {OpInfo.Category.ToString()}, OpInfo.Description,  OpInfo.Keywords);
 				MathOpAction->AlternateSearchName = OpInfo.AlternateSearchName.IsSet() ? FText::FromName(OpInfo.AlternateSearchName.GetValue()) : TOptional<FText>();
 			}
 		}
@@ -506,7 +502,7 @@ TArray<TSharedPtr<FNiagaraAction_NewNode>> UEdGraphSchema_Niagara::GetGraphActio
 		
 		UNiagaraNodeCustomHlsl* CustomHlslNode = NewObject<UNiagaraNodeCustomHlsl>(OwnerOfTemporaries);
 		CustomHlslNode->SetCustomHlsl(TEXT("// Insert the body of the function here and add any inputs\r\n// and outputs by name using the add pins above.\r\n// Currently, complicated branches, for loops, switches, etc are not advised."));
-		AddNewNodeMenuAction(NewActions, CustomHlslNode, DisplayName, {LOCTEXT("Function Menu Title", "Functions").ToString()}, {}, TooltipDesc,  FText::GetEmpty());
+		AddNewNodeMenuAction(NewActions, CustomHlslNode, DisplayName, ENiagaraMenuSections::General, {LOCTEXT("Function Menu Title", "Functions").ToString()}, TooltipDesc,  FText::GetEmpty());
 	}
 
 	auto AddScriptFunctionAction = [&NewActions, OwnerOfTemporaries](const TArray<FString>& Categories, const FAssetData& ScriptAsset)
@@ -525,14 +521,12 @@ TArray<TSharedPtr<FNiagaraAction_NewNode>> UEdGraphSchema_Niagara::GetGraphActio
 		const TTuple<EScriptSource, FText> Source = FNiagaraEditorUtilities::GetScriptSource(ScriptAsset);
 		FNiagaraActionSourceData SourceData(Source.Key, Source.Value, true);
 		
-		FNiagaraFavoritesActionData FavoritesActionData;
-		FavoritesActionData.ActionIdentifier.Names.Add(FName(ScriptAsset.GetObjectPathString()));
-		FavoritesActionData.bFavoriteByDefault = bSuggested;
-		
+		const ENiagaraMenuSections Section = bSuggested ? ENiagaraMenuSections::Suggested: ENiagaraMenuSections::General;
+
 		UNiagaraNodeFunctionCall* FunctionCallNode = NewObject<UNiagaraNodeFunctionCall>(OwnerOfTemporaries);
 		FunctionCallNode->FunctionScriptAssetObjectPath = FName(ScriptAsset.GetSoftObjectPath().ToString());
 
-		TSharedPtr<FNiagaraAction_NewNode> Action = AddNewNodeMenuAction(NewActions, FunctionCallNode, DisplayName, Categories, FavoritesActionData, TooltipDesc, Keywords, SourceData);
+		TSharedPtr<FNiagaraAction_NewNode> Action = AddNewNodeMenuAction(NewActions, FunctionCallNode, DisplayName, Section, Categories, TooltipDesc, Keywords, SourceData);
 		Action->bIsInLibrary = bIsInLibrary;
 
 		return Action;
@@ -562,7 +556,7 @@ TArray<TSharedPtr<FNiagaraAction_NewNode>> UEdGraphSchema_Niagara::GetGraphActio
 			FunctionCallNode->Signature = Sig;
 
 			const FText MenuDesc = FText::FromString(FName::NameToDisplayString(Sig.Name.ToString(), false));
-			TSharedPtr<FNiagaraAction_NewNode> FunctionCallAction = AddNewNodeMenuAction(NewActions, FunctionCallNode, MenuDesc, TArray<FString>(), {},  MenuDesc);
+			TSharedPtr<FNiagaraAction_NewNode> FunctionCallAction = AddNewNodeMenuAction(NewActions, FunctionCallNode, MenuDesc, ENiagaraMenuSections::General, TArray<FString>(), MenuDesc);
 			FunctionCallAction->WeakNodeTemplate = FunctionCallNode;
 		}
 	}
@@ -613,7 +607,7 @@ TArray<TSharedPtr<FNiagaraAction_NewNode>> UEdGraphSchema_Niagara::GetGraphActio
 					UNiagaraNodeReadDataSet* EventReadNode = NewObject<UNiagaraNodeReadDataSet>(OwnerOfTemporaries);
 					EventReadNode->InitializeFromStruct(PayloadType.GetStruct());
 
-					AddNewNodeMenuAction(NewActions, EventReadNode, DisplayName, ReadCategories, {}, FText::GetEmpty(), FText::GetEmpty());
+					AddNewNodeMenuAction(NewActions, EventReadNode, DisplayName, ENiagaraMenuSections::General, ReadCategories, FText::GetEmpty(), FText::GetEmpty());
 				}
 				{
 					const FText DisplayNameFormat = LOCTEXT("AddEventWriteFmt", "Add {0} Event Write");
@@ -624,7 +618,7 @@ TArray<TSharedPtr<FNiagaraAction_NewNode>> UEdGraphSchema_Niagara::GetGraphActio
 					UNiagaraNodeWriteDataSet* EventWriteNode = NewObject<UNiagaraNodeWriteDataSet>(OwnerOfTemporaries);
 					EventWriteNode->InitializeFromStruct(PayloadType.GetStruct());
 					
-					AddNewNodeMenuAction(NewActions, EventWriteNode, DisplayName,  WriteCategories, {}, FText::GetEmpty(), FText::GetEmpty());
+					AddNewNodeMenuAction(NewActions, EventWriteNode, DisplayName, ENiagaraMenuSections::General, WriteCategories, FText::GetEmpty(), FText::GetEmpty());
 				}
 			}
 		}
@@ -655,7 +649,7 @@ TArray<TSharedPtr<FNiagaraAction_NewNode>> UEdGraphSchema_Niagara::GetGraphActio
 
 			FText DisplayName = FText::Format(LOCTEXT("AddOutput", "Add {0}"), OutputNode->GetNodeTitle(ENodeTitleType::FullTitle));
 
-			AddNewNodeMenuAction(NewActions, OutputNode, DisplayName,  {MenuCat.ToString()}, {}, FText::GetEmpty(), FText::GetEmpty());
+			AddNewNodeMenuAction(NewActions, OutputNode, DisplayName, ENiagaraMenuSections::General, {MenuCat.ToString()}, FText::GetEmpty(), FText::GetEmpty());
 
 			UNiagaraNodeOutput* UpdateOutputNode = NiagaraGraph->FindOutputNode(ENiagaraScriptUsage::ParticleUpdateScript);
 			if (UpdateOutputNode)
@@ -693,7 +687,7 @@ TArray<TSharedPtr<FNiagaraAction_NewNode>> UEdGraphSchema_Niagara::GetGraphActio
 				FText Desc = FText::Format(DescFmt, DisplayName);
 				
 				UNiagaraNodeConvert* ConvertNode = NewObject<UNiagaraNodeConvert>(OwnerOfTemporaries);
-				TSharedPtr<FNiagaraAction_NewNode> Action = AddNewNodeMenuAction(NewActions, ConvertNode, DisplayName,  {bMake ? MakeCat.ToString() : BreakCat.ToString()}, {}, FText::GetEmpty(), FText::GetEmpty());
+				TSharedPtr<FNiagaraAction_NewNode> Action = AddNewNodeMenuAction(NewActions, ConvertNode, DisplayName, ENiagaraMenuSections::General, {bMake ? MakeCat.ToString() : BreakCat.ToString()}, FText::GetEmpty(), FText::GetEmpty());
 				if (bMake)
 				{
 					ConvertNode->InitAsMake(Type);
@@ -798,7 +792,7 @@ TArray<TSharedPtr<FNiagaraAction_NewNode>> UEdGraphSchema_Niagara::GetGraphActio
 			FText Desc = LOCTEXT("NiagaraConvert", "Convert");
 			
 			UNiagaraNodeConvert* ConvertNode = NewObject<UNiagaraNodeConvert>(OwnerOfTemporaries);
-			AddNewNodeMenuAction(NewActions, ConvertNode, Desc,  {}, {},FText::GetEmpty(), FText::GetEmpty());
+			AddNewNodeMenuAction(NewActions, ConvertNode, Desc, ENiagaraMenuSections::General, {}, FText::GetEmpty(), FText::GetEmpty());
 		}
 	}
 
@@ -816,7 +810,7 @@ TArray<TSharedPtr<FNiagaraAction_NewNode>> UEdGraphSchema_Niagara::GetGraphActio
 					continue;
 
 				UNiagaraNodeFunctionCall* FuncNode = NewObject<UNiagaraNodeFunctionCall>(OwnerOfTemporaries);
-				AddNewNodeMenuAction(NewActions, FuncNode, FText::FromString(FName::NameToDisplayString(Sig.GetNameString(), false)),  {MenuCat.ToString()}, {}, Sig.GetDescription(), FText::GetEmpty());
+				AddNewNodeMenuAction(NewActions, FuncNode, FText::FromString(FName::NameToDisplayString(Sig.GetNameString(), false)), ENiagaraMenuSections::General, {MenuCat.ToString()}, Sig.GetDescription(), FText::GetEmpty());
 				FuncNode->Signature = Sig;
 			}
 		}
@@ -853,7 +847,7 @@ TArray<TSharedPtr<FNiagaraAction_NewNode>> UEdGraphSchema_Niagara::GetGraphActio
 					const FText Category = LOCTEXT("NiagaraSwizzles", "Swizzles");
 					
 					UNiagaraNodeConvert* ConvertNode = NewObject<UNiagaraNodeConvert>(OwnerOfTemporaries);
-					AddNewNodeMenuAction(NewActions, ConvertNode, FText::FromString(Swiz),  {Category.ToString()}, {}, FText::GetEmpty(), FText::GetEmpty());
+					AddNewNodeMenuAction(NewActions, ConvertNode, FText::FromString(Swiz), ENiagaraMenuSections::General, {Category.ToString()}, FText::GetEmpty(), FText::GetEmpty());
 
 					ConvertNode->InitAsSwizzle(Swiz);
 				}
@@ -868,42 +862,32 @@ TArray<TSharedPtr<FNiagaraAction_NewNode>> UEdGraphSchema_Niagara::GetGraphActio
 		if (NiagaraEditorSettings->IsVisibleClass(UNiagaraNodeParameterMapGet::StaticClass()))
 		{
 			FString Name = TEXT("Parameter Map Get");
-			
-			FNiagaraFavoritesActionData FavoritesActionData;
-			FavoritesActionData.ActionIdentifier.Names.Add(FName(Name));
-			FavoritesActionData.bFavoriteByDefault = true;
-			
 			UNiagaraNodeParameterMapGet* BaseNode = NewObject<UNiagaraNodeParameterMapGet>(OwnerOfTemporaries);
-			AddNewNodeMenuAction(NewActions, BaseNode, FText::FromString(Name), {MenuCat.ToString()}, FavoritesActionData, FText::GetEmpty(), FText::GetEmpty());
+			AddNewNodeMenuAction(NewActions, BaseNode, FText::FromString(Name), ENiagaraMenuSections::Suggested, {MenuCat.ToString()}, FText::GetEmpty(), FText::GetEmpty());
 		}
 		if (NiagaraEditorSettings->IsVisibleClass(UNiagaraNodeParameterMapSet::StaticClass()))
 		{
 			FString Name = TEXT("Parameter Map Set");
-			
-			FNiagaraFavoritesActionData FavoritesActionData;
-			FavoritesActionData.ActionIdentifier.Names.Add(FName(Name));
-			FavoritesActionData.bFavoriteByDefault = true;
-			
 			UNiagaraNodeParameterMapSet* BaseNode = NewObject<UNiagaraNodeParameterMapSet>(OwnerOfTemporaries);
-			AddNewNodeMenuAction(NewActions, BaseNode, FText::FromString(Name), {MenuCat.ToString()}, FavoritesActionData, FText::GetEmpty(), FText::GetEmpty());
+			AddNewNodeMenuAction(NewActions, BaseNode, FText::FromString(Name), ENiagaraMenuSections::Suggested, {MenuCat.ToString()}, FText::GetEmpty(), FText::GetEmpty());
 		}
 		if (NiagaraEditorSettings->IsVisibleClass(UNiagaraNodeParameterMapFor::StaticClass()))
 		{
 			FString Name = TEXT("Parameter Map For");
 			UNiagaraNodeParameterMapFor* BaseNode = NewObject<UNiagaraNodeParameterMapFor>(OwnerOfTemporaries);
-			AddNewNodeMenuAction(NewActions, BaseNode, FText::FromString(Name),  {MenuCat.ToString()}, {}, FText::GetEmpty(), FText::GetEmpty());
+			AddNewNodeMenuAction(NewActions, BaseNode, FText::FromString(Name), ENiagaraMenuSections::General, {MenuCat.ToString()}, FText::GetEmpty(), FText::GetEmpty());
 		}
 		if (NiagaraEditorSettings->IsVisibleClass(UNiagaraNodeParameterMapForWithContinue::StaticClass()))
 		{
 			FString Name = TEXT("Parameter Map For With Continue");
 			UNiagaraNodeParameterMapForWithContinue* BaseNode = NewObject<UNiagaraNodeParameterMapForWithContinue>(OwnerOfTemporaries);
-			AddNewNodeMenuAction(NewActions, BaseNode, FText::FromString(Name),  { MenuCat.ToString() }, {}, FText::GetEmpty(), FText::GetEmpty());
+			AddNewNodeMenuAction(NewActions, BaseNode, FText::FromString(Name), ENiagaraMenuSections::General, { MenuCat.ToString() }, FText::GetEmpty(), FText::GetEmpty());
 		}
 		if (NiagaraEditorSettings->IsVisibleClass(UNiagaraNodeParameterMapForIndex::StaticClass()))
 		{
 			FString Name = TEXT("Parameter Map For Current Index");
 			UNiagaraNodeParameterMapForIndex* BaseNode = NewObject<UNiagaraNodeParameterMapForIndex>(OwnerOfTemporaries);
-			AddNewNodeMenuAction(NewActions, BaseNode, FText::FromString(Name),  { MenuCat.ToString() }, {}, FText::GetEmpty(), FText::GetEmpty());
+			AddNewNodeMenuAction(NewActions, BaseNode, FText::FromString(Name), ENiagaraMenuSections::General, { MenuCat.ToString() }, FText::GetEmpty(), FText::GetEmpty());
 		}
 	}
 
@@ -915,7 +899,7 @@ TArray<TSharedPtr<FNiagaraAction_NewNode>> UEdGraphSchema_Niagara::GetGraphActio
 		{
 			FString Name = TEXT("Add Comment");
 			UEdGraphNode_Comment* BaseNode = NewObject<UEdGraphNode_Comment>(OwnerOfTemporaries);
-			AddNewNodeMenuAction(NewActions, BaseNode, FText::FromString(Name),  {MenuCat.ToString()}, {}, FText::GetEmpty(), FText::GetEmpty());
+			AddNewNodeMenuAction(NewActions, BaseNode, FText::FromString(Name), ENiagaraMenuSections::General, {MenuCat.ToString()}, FText::GetEmpty(), FText::GetEmpty());
 		}		
 	}
 
@@ -927,7 +911,7 @@ TArray<TSharedPtr<FNiagaraAction_NewNode>> UEdGraphSchema_Niagara::GetGraphActio
 		{
 			FString Name = TEXT("Add Compiler Output Tag");
 			UNiagaraNodeOutputTag* BaseNode = NewObject<UNiagaraNodeOutputTag>(OwnerOfTemporaries);
-			AddNewNodeMenuAction(NewActions, BaseNode, FText::FromString(Name),  {MenuCat.ToString()}, {}, FText::GetEmpty(), FText::GetEmpty());
+			AddNewNodeMenuAction(NewActions, BaseNode, FText::FromString(Name), ENiagaraMenuSections::General, {MenuCat.ToString()}, FText::GetEmpty(), FText::GetEmpty());
 		}
 	}
 	
@@ -948,7 +932,7 @@ TArray<TSharedPtr<FNiagaraAction_NewNode>> UEdGraphSchema_Niagara::GetGraphActio
 				const FText DisplayName = FText::Format(LOCTEXT("GetSystemConstant", "Get {Constant}"), Args);
 				
 				UNiagaraNodeInput* InputNode = NewObject<UNiagaraNodeInput>(OwnerOfTemporaries);
-				AddNewNodeMenuAction(NewActions, InputNode, DisplayName,  {LOCTEXT("System Parameters Menu Title", "System Parameters").ToString()}, {}, FText::GetEmpty(), FText::GetEmpty());
+				AddNewNodeMenuAction(NewActions, InputNode, DisplayName, ENiagaraMenuSections::General, {LOCTEXT("System Parameters Menu Title", "System Parameters").ToString()}, FText::GetEmpty(), FText::GetEmpty());
 
 				InputNode->Usage = ENiagaraInputNodeUsage::SystemConstant;
 				InputNode->Input = SysConst;
@@ -964,7 +948,7 @@ TArray<TSharedPtr<FNiagaraAction_NewNode>> UEdGraphSchema_Niagara::GetGraphActio
 			const FText DisplayName = FText::Format(LOCTEXT("GetTranslatorConstant", "{Constant}"), Args);
 			
 			UNiagaraNodeInput* InputNode = NewObject<UNiagaraNodeInput>(OwnerOfTemporaries);
-			AddNewNodeMenuAction(NewActions, InputNode, DisplayName,  {LOCTEXT("Translator Parameters Menu Title", "Special Purpose Parameters").ToString()}, {}, FText::GetEmpty(), FText::GetEmpty());
+			AddNewNodeMenuAction(NewActions, InputNode, DisplayName, ENiagaraMenuSections::General, {LOCTEXT("Translator Parameters Menu Title", "Special Purpose Parameters").ToString()}, FText::GetEmpty(), FText::GetEmpty());
 
 			InputNode->Usage = ENiagaraInputNodeUsage::TranslatorConstant;
 			InputNode->ExposureOptions.bCanAutoBind = true;
@@ -1021,7 +1005,7 @@ TArray<TSharedPtr<FNiagaraAction_NewNode>> UEdGraphSchema_Niagara::GetGraphActio
 					const FText DisplayName = FText::Format(MenuDescFmt, Type.GetNameText());
 					
 					UNiagaraNodeInput* InputNode = NewObject<UNiagaraNodeInput>(OwnerOfTemporaries);
-					AddNewNodeMenuAction(NewActions, InputNode, DisplayName,  Categories, {}, FText::GetEmpty(), FText::GetEmpty());
+					AddNewNodeMenuAction(NewActions, InputNode, DisplayName, ENiagaraMenuSections::General, Categories, FText::GetEmpty(), FText::GetEmpty());
 					FNiagaraEditorUtilities::InitializeParameterInputNode(*InputNode, Type, NiagaraGraph);
 				}
 
@@ -1032,7 +1016,7 @@ TArray<TSharedPtr<FNiagaraAction_NewNode>> UEdGraphSchema_Niagara::GetGraphActio
 					const FText DisplayName = FText::Format(MenuDescFmt, PinType.GetNameText());
 					
 					UNiagaraNodeInput* InputNode = NewObject<UNiagaraNodeInput>(OwnerOfTemporaries);
-					AddNewNodeMenuAction(NewActions, InputNode, DisplayName,  {}, {}, FText::GetEmpty(), FText::GetEmpty());
+					AddNewNodeMenuAction(NewActions, InputNode, DisplayName, ENiagaraMenuSections::General, {}, FText::GetEmpty(), FText::GetEmpty());
 					FNiagaraEditorUtilities::InitializeParameterInputNode(*InputNode, PinType, NiagaraGraph);
 				}
 			}
@@ -1047,7 +1031,7 @@ TArray<TSharedPtr<FNiagaraAction_NewNode>> UEdGraphSchema_Niagara::GetGraphActio
 		const FText RerouteMenuDesc = LOCTEXT("NiagaraRerouteMenuDesc", "Reroute");
 		
 		UNiagaraNodeReroute* RerouteNode = NewObject<UNiagaraNodeReroute>(OwnerOfTemporaries);
-		AddNewNodeMenuAction(NewActions, RerouteNode, RerouteMenuDesc,  {UtilMenuCat.ToString()}, {}, FText::GetEmpty(), FText::GetEmpty());
+		AddNewNodeMenuAction(NewActions, RerouteNode, RerouteMenuDesc, ENiagaraMenuSections::General, {UtilMenuCat.ToString()}, FText::GetEmpty(), FText::GetEmpty());
 	}
 	
 	// Add select node
@@ -1055,13 +1039,9 @@ TArray<TSharedPtr<FNiagaraAction_NewNode>> UEdGraphSchema_Niagara::GetGraphActio
 	if (bAllowSelectNodes && NiagaraEditorSettings->IsVisibleClass(UNiagaraNodeSelect::StaticClass()))
 	{
 		const FText SelectMenuDesc = LOCTEXT("NiagaraSelectMenuDesc", "Select / If");
-
-		FNiagaraFavoritesActionData FavoritesActionData;
-		FavoritesActionData.ActionIdentifier.Names.Add(FName(FText::AsCultureInvariant(SelectMenuDesc).ToString()));
-		FavoritesActionData.bFavoriteByDefault = true;
 		
 		UNiagaraNodeSelect* Node = NewObject<UNiagaraNodeSelect>(OwnerOfTemporaries);
-		AddNewNodeMenuAction(NewActions, Node, SelectMenuDesc,  {UtilMenuCat.ToString()}, FavoritesActionData, FText::GetEmpty(), FText::FromString(TEXT("If Branch Bool Select")));
+		AddNewNodeMenuAction(NewActions, Node, SelectMenuDesc, ENiagaraMenuSections::Suggested, {UtilMenuCat.ToString()}, FText::GetEmpty(), FText::FromString(TEXT("If Branch Bool Select")));
 	}
 
 	// Add static switch node
@@ -1069,15 +1049,11 @@ TArray<TSharedPtr<FNiagaraAction_NewNode>> UEdGraphSchema_Niagara::GetGraphActio
 	if (bAllowStaticSwitchNodes && NiagaraEditorSettings->IsVisibleClass(UNiagaraNodeStaticSwitch::StaticClass()))
 	{
 		const FText UsageSelectorMenuDesc = LOCTEXT("NiagaraStaticSwitchMenuDesc", "Static Switch");
-
-		FNiagaraFavoritesActionData FavoritesActionData;
-		FavoritesActionData.ActionIdentifier.Names.Add(FName(FText::AsCultureInvariant(UsageSelectorMenuDesc).ToString()));
-		FavoritesActionData.bFavoriteByDefault = true;
 		
 		UNiagaraNodeStaticSwitch* Node = NewObject<UNiagaraNodeStaticSwitch>(OwnerOfTemporaries);
 		// new nodes should auto refresh
 		Node->SwitchTypeData.bAutoRefreshEnabled = true;
-		AddNewNodeMenuAction(NewActions, Node, UsageSelectorMenuDesc,  {UtilMenuCat.ToString()}, FavoritesActionData, FText::GetEmpty(), FText::FromString(TEXT("")));
+		AddNewNodeMenuAction(NewActions, Node, UsageSelectorMenuDesc, ENiagaraMenuSections::Suggested, {UtilMenuCat.ToString()}, FText::GetEmpty(), FText::FromString(TEXT("")));
 	}
 
 	return NewActions;
diff --git a/Engine/Plugins/FX/Niagara/Source/NiagaraEditor/Private/NiagaraActions.cpp b/Engine/Plugins/FX/Niagara/Source/NiagaraEditor/Private/NiagaraActions.cpp
--- a/Engine/Plugins/FX/Niagara/Source/NiagaraEditor/Private/NiagaraActions.cpp
+++ b/Engine/Plugins/FX/Niagara/Source/NiagaraEditor/Private/NiagaraActions.cpp
@@ -24,7 +24,6 @@
 #include "Widgets/SToolTip.h"
 
 #include "NiagaraDataChannel.h"
-#include "Config/NiagaraFavoriteActionsConfig.h"
 
 #include UE_INLINE_GENERATED_CPP_BY_NAME(NiagaraActions)
 
@@ -103,23 +102,16 @@ void FNiagaraMenuActionCollector::AddAllActionsTo(FGraphActionListBuilderBase& A
 	}
 }
 
-FNiagaraMenuAction_Base::FNiagaraMenuAction_Base(FText InDisplayName, TArray<FString> InNodeCategories, TOptional<FNiagaraFavoritesActionData> InFavoritesActionData, FText InToolTip, FText InKeywords, float InIntrinsicWeightMultiplier)
+FNiagaraMenuAction_Base::FNiagaraMenuAction_Base(FText InDisplayName, ENiagaraMenuSections InSection,
+                                                 TArray<FString> InNodeCategories, FText InToolTip, FText InKeywords, float InIntrinsicWeightMultiplier)
 {
 	DisplayName = InDisplayName;
+	Section = InSection;
 	Categories = InNodeCategories;
-	FavoritesActionData = InFavoritesActionData;
 	ToolTip = InToolTip;
 	Keywords = InKeywords;
 	SearchWeightMultiplier = InIntrinsicWeightMultiplier;
 
-	if(FavoritesActionData.IsSet() == false)
-	{
-		FNiagaraFavoritesActionData DefaultFavoritesActionData;
-		DefaultFavoritesActionData.ActionIdentifier.Names.Add(FName(DisplayName.ToString()));
-		DefaultFavoritesActionData.bFavoriteByDefault = false;
-		FavoritesActionData = DefaultFavoritesActionData;
-	}
-	
 	UpdateFullSearchText();
 }
 
diff --git a/Engine/Plugins/FX/Niagara/Source/NiagaraEditor/Private/NiagaraEditorModule.cpp b/Engine/Plugins/FX/Niagara/Source/NiagaraEditor/Private/NiagaraEditorModule.cpp
--- a/Engine/Plugins/FX/Niagara/Source/NiagaraEditor/Private/NiagaraEditorModule.cpp
+++ b/Engine/Plugins/FX/Niagara/Source/NiagaraEditor/Private/NiagaraEditorModule.cpp
@@ -1100,12 +1100,6 @@ void FNiagaraEditorModule::StartupModule()
 
 	UNiagaraAssetBrowserConfig::Initialize();
 
-	// We have to call shutdown before the exit actually happens to be able to save to disk
-	FEditorDelegates::OnEditorPreExit.AddLambda([]()
-	{
-		UNiagaraFavoriteActionsConfig::Shutdown();
-	});
-	
 	FNiagaraAssetDetailDatabase::Init();
 
 	TempPackage = NewObject<UPackage>(nullptr, TEXT("/Temp/NiagaraEditor"), RF_Transient);
@@ -1662,7 +1656,7 @@ void FNiagaraEditorModule::ShutdownModule()
 	FNiagaraEditorStyle::Unregister();
 
 	UnregisterSettings();
-	
+
 	if (UObjectInitialized())
 	{
 		FComponentAssetBrokerage::UnregisterBroker(NiagaraComponentBroker);
diff --git a/Engine/Plugins/FX/Niagara/Source/NiagaraEditor/Private/NiagaraEditorUtilities.cpp b/Engine/Plugins/FX/Niagara/Source/NiagaraEditor/Private/NiagaraEditorUtilities.cpp
--- a/Engine/Plugins/FX/Niagara/Source/NiagaraEditor/Private/NiagaraEditorUtilities.cpp
+++ b/Engine/Plugins/FX/Niagara/Source/NiagaraEditor/Private/NiagaraEditorUtilities.cpp
@@ -1655,7 +1655,7 @@ bool FNiagaraEditorUtilities::IsEnginePluginAsset(const FTopLevelAssetPath& InTo
 		FPaths::IsUnderDirectory(PackagePathLocal, FPaths::EnginePluginsDir() + "FX/Niagara");
 }
 
-int32 FNiagaraEditorUtilities::GetWeightForItem(const TSharedPtr<FNiagaraMenuAction_Generic>& InCurrentAction, FName FavoriteActionsProfile, const TArray<FString>& InFilterTerms)
+int32 FNiagaraEditorUtilities::GetWeightForItem(const TSharedPtr<FNiagaraMenuAction_Generic>& InCurrentAction, const TArray<FString>& InFilterTerms)
 {
 	// The overall 'weight'
 	int32 TotalWeight = 0;
@@ -1881,14 +1881,11 @@ int32 FNiagaraEditorUtilities::GetWeightForItem(const TSharedPtr<FNiagaraMenuAct
 		}
 	}
 
-	// favorite actions get favored massively
-	if(InCurrentAction->FavoritesActionData.IsSet())
+	// suggested actions get favored massively
+	if(InCurrentAction->Section == ENiagaraMenuSections::Suggested)
 	{
-		FNiagaraFavoritesActionData ActionData = InCurrentAction->FavoritesActionData.GetValue();
-		if(UNiagaraFavoriteActionsConfig::Get()->GetActionsProfile(FavoriteActionsProfile).IsFavorite(ActionData))
-		{
-			TotalWeight += (int32)FMath::Pow(10.0f, 4.0f);
-		}
+		//TotalWeight *= TotalWeight;
+		TotalWeight += (int32)FMath::Pow(10.0f, 4.0f);
 	}
 	
 	return TotalWeight > 0.f ? TotalWeight : INDEX_NONE;
@@ -3131,19 +3128,6 @@ void FNiagaraEditorUtilities::GetAvailableParameterCollections(TArray<UNiagaraPa
 	}
 }
 
-void FNiagaraEditorUtilities::GetAvailableParameterDefinitions(TArray<UNiagaraParameterDefinitions*>& OutParameterDefinitions)
-{
-	const TArray<TWeakObjectPtr<UNiagaraParameterDefinitions>>& CachedParameterDefinitionAssets = FNiagaraEditorModule::Get().GetCachedParameterDefinitionsAssets();
-
-	for (const TWeakObjectPtr<UNiagaraParameterDefinitions>& CachedParameterDefinitionAsset : CachedParameterDefinitionAssets)
-	{
-		if (UNiagaraParameterDefinitions* ParameterDefinition = CachedParameterDefinitionAsset.Get())
-		{
-			OutParameterDefinitions.Add(ParameterDefinition);
-		}
-	}
-}
-
 TSharedPtr<INiagaraParameterDefinitionsSubscriberViewModel> FNiagaraEditorUtilities::GetOwningLibrarySubscriberViewModelForGraph(const UNiagaraGraph* Graph)
 {
 	if (UNiagaraScript* BaseScript = Graph->GetTypedOuter<UNiagaraScript>())
diff --git a/Engine/Plugins/FX/Niagara/Source/NiagaraEditor/Private/ViewModels/Stack/NiagaraStackFunctionInput.cpp b/Engine/Plugins/FX/Niagara/Source/NiagaraEditor/Private/ViewModels/Stack/NiagaraStackFunctionInput.cpp
--- a/Engine/Plugins/FX/Niagara/Source/NiagaraEditor/Private/ViewModels/Stack/NiagaraStackFunctionInput.cpp
+++ b/Engine/Plugins/FX/Niagara/Source/NiagaraEditor/Private/ViewModels/Stack/NiagaraStackFunctionInput.cpp
@@ -31,7 +31,6 @@
 #include "NiagaraScriptMergeManager.h"
 #include "NiagaraScriptSource.h"
 #include "NiagaraSettings.h"
-#include "NiagaraTrace.h"
 #include "ViewModels/NiagaraSystemScriptViewModel.h"
 #include "ScopedTransaction.h"
 #include "Algo/RemoveIf.h"
@@ -61,17 +60,6 @@
 
 static FText TooManyConversionScripts = LOCTEXT("TooManyConversionScripts", "There is more than one dynamic input script available auto-convert the dragged parameter. Please fix this by disabling conversion for all but one of them:\n{0}");
 
-UNiagaraStackFunctionInput::FNiagaraAvailableParameterInfo::FNiagaraAvailableParameterInfo(FNiagaraVariableBase InVariableBase)
-{
-	Variable = InVariableBase;
-}
-
-UNiagaraStackFunctionInput::FNiagaraAvailableParameterInfo::FNiagaraAvailableParameterInfo(FNiagaraVariableBase InVariableBase,	UNiagaraScript* InConversionScript)
-{
-	Variable = InVariableBase;
-	ConversionScript = InConversionScript;
-}
-
 UNiagaraStackFunctionInput::UNiagaraStackFunctionInput()
 	: OwningModuleNode(nullptr)
 	, OwningFunctionCallNode(nullptr)
@@ -1820,10 +1808,8 @@ UNiagaraScript* UNiagaraStackFunctionInput::FindConversionScript(const FNiagaraT
 	return ConversionScriptCache[FromType];
 }
 
-void UNiagaraStackFunctionInput::GetAvailableParameters(TSet<FNiagaraAvailableParameterInfo>& OutAvailableParameters, FGetAvailableParameterArgs Arguments) const
+void UNiagaraStackFunctionInput::GetAvailableParameters(TArray<FNiagaraVariableBase>& OutAvailableParameters, TMap<FNiagaraVariableBase, UNiagaraScript*>& OutAvailableParametersWithConversions, bool bIncludeConversionScripts) const
 {
-	TRACE_CPUPROFILER_EVENT_SCOPE_ON_CHANNEL(UNiagaraStackFunctionInput::GetAvailableParameters, NiagaraChannel);
-	
 	TMap<FNiagaraTypeDefinition, UNiagaraScript*> ConversionScriptCache;
 	
 	// Engine Handles.
@@ -1833,9 +1819,9 @@ void UNiagaraStackFunctionInput::GetAvailableParameters(TSet<FNiagaraAvailablePa
 		{
 			OutAvailableParameters.Add(SystemVariable);
 		}
-		else if (UNiagaraScript* ConversionScript = FindConversionScript(SystemVariable.GetType(), ConversionScriptCache, Arguments.bIncludeConversionScripts))
+		else if (UNiagaraScript* ConversionScript = FindConversionScript(SystemVariable.GetType(), ConversionScriptCache, bIncludeConversionScripts))
 		{
-			OutAvailableParameters.Add({SystemVariable, ConversionScript});
+			OutAvailableParametersWithConversions.Add(SystemVariable, ConversionScript);
 		}
 	}
 
@@ -1848,9 +1834,9 @@ void UNiagaraStackFunctionInput::GetAvailableParameters(TSet<FNiagaraAvailablePa
 		{
 			OutAvailableParameters.Add(ExposedVar);
 		}
-		else if (UNiagaraScript* ConversionScript = FindConversionScript(ExposedVar.GetType(), ConversionScriptCache, Arguments.bIncludeConversionScripts))
+		else if (UNiagaraScript* ConversionScript = FindConversionScript(ExposedVar.GetType(), ConversionScriptCache, bIncludeConversionScripts))
 		{
-			OutAvailableParameters.Add({ExposedVar, ConversionScript});
+			OutAvailableParametersWithConversions.Add(ExposedVar, ConversionScript);
 		}
 	}
 
@@ -1948,7 +1934,7 @@ void UNiagaraStackFunctionInput::GetAvailableParameters(TSet<FNiagaraAvailablePa
 
 						if (FNiagaraEditorUtilities::AreTypesAssignable(HistoryVariable.GetType(), InputType))
 						{
-							OutAvailableParameters.Add(HistoryVariable);
+							OutAvailableParameters.AddUnique(HistoryVariable);
 							if (AvailableHandle.GetNamespace() == OutputNamespace)
 							{
 								AvailableParametersForThisOutput.Add(TPair<FNiagaraVariableBase, FNiagaraParameterHandle>(HistoryVariable, AvailableHandle));
@@ -1963,13 +1949,13 @@ void UNiagaraStackFunctionInput::GetAvailableParameters(TSet<FNiagaraAvailablePa
 									FString NewName = HistoryVariable.GetName().ToString().Replace(*StackContextRoots[StackRootIdx].ToString(), *FNiagaraConstants::StackContextNamespace.ToString());
 
 									FNiagaraVariableBase AvailableAliasedParameter = FNiagaraVariableBase(HistoryVariable.GetType(), *NewName);
-									OutAvailableParameters.Add(AvailableAliasedParameter);
+									OutAvailableParameters.AddUnique(AvailableAliasedParameter);
 								}
 							}
 						}
-						else if (UNiagaraScript* ConversionScript = FindConversionScript(HistoryVariable.GetType(), ConversionScriptCache, Arguments.bIncludeConversionScripts))
+						else if (UNiagaraScript* ConversionScript = FindConversionScript(HistoryVariable.GetType(), ConversionScriptCache, bIncludeConversionScripts))
 						{
-							OutAvailableParameters.Add({HistoryVariable, ConversionScript});
+							OutAvailableParametersWithConversions.Add(HistoryVariable, ConversionScript);
 						}
 					}
 				}
@@ -1984,7 +1970,7 @@ void UNiagaraStackFunctionInput::GetAvailableParameters(TSet<FNiagaraAvailablePa
 					if (AvailableParameter.GetType().IsStatic() == false)
 					{
 						FNiagaraParameterHandle InitialParameterHandle = FNiagaraParameterHandle::CreateInitialParameterHandle(AvailableHandle);
-						OutAvailableParameters.Add(FNiagaraVariableBase(AvailableParameter.GetType(), InitialParameterHandle.GetParameterHandleString()));
+						OutAvailableParameters.AddUnique(FNiagaraVariableBase(AvailableParameter.GetType(), InitialParameterHandle.GetParameterHandleString()));
 					}
 				}
 			}
@@ -2000,39 +1986,14 @@ void UNiagaraStackFunctionInput::GetAvailableParameters(TSet<FNiagaraAvailablePa
 		{
 			if (FNiagaraEditorUtilities::AreTypesAssignable(CollectionParam.GetType(), InputType))
 			{
-				OutAvailableParameters.Add(CollectionParam);
+				OutAvailableParameters.AddUnique(CollectionParam);
 			}
-			else if (UNiagaraScript* ConversionScript = FindConversionScript(CollectionParam.GetType(), ConversionScriptCache, Arguments.bIncludeConversionScripts))
+			else if (UNiagaraScript* ConversionScript = FindConversionScript(CollectionParam.GetType(), ConversionScriptCache, bIncludeConversionScripts))
 			{
-				OutAvailableParameters.Add({CollectionParam, ConversionScript});
+				OutAvailableParametersWithConversions.Add(CollectionParam, ConversionScript);
 			}
 		}
 	}
-
-	if(Arguments.bIncludeParameterDefinitions)
-	{
-		TRACE_CPUPROFILER_EVENT_SCOPE_ON_CHANNEL(UNiagaraStackFunctionInput::GetAvailableParameters::IncludeParameterDefinitions, NiagaraChannel);
-		// At last, we patch up all the found parameters with available parameter definitions
-		TArray<UNiagaraParameterDefinitions*> AvailableParameterDefinitions;
-		FNiagaraEditorUtilities::GetAvailableParameterDefinitions(AvailableParameterDefinitions);
-		TMap<FNiagaraVariableBase, UNiagaraScriptVariable*> ParameterDefinitionsCache;
-
-		for(UNiagaraParameterDefinitions* DefinitionsAsset : AvailableParameterDefinitions)
-		{
-			for(UNiagaraScriptVariable* ScriptVariable : DefinitionsAsset->GetParametersConst())
-			{
-				ParameterDefinitionsCache.Add(ScriptVariable->Variable, ScriptVariable);
-			}
-		}
-		
-		for(FNiagaraAvailableParameterInfo& ParameterInfo : OutAvailableParameters)
-		{
-			if(ParameterDefinitionsCache.Contains(ParameterInfo.Variable))
-			{
-				ParameterInfo.MatchingParameterDefinition = ParameterDefinitionsCache[ParameterInfo.Variable];
-			}		
-		}
-	}
 }
 
 UNiagaraNodeFunctionCall* UNiagaraStackFunctionInput::GetDefaultFunctionNode() const
@@ -3587,14 +3548,13 @@ void UNiagaraStackFunctionInput::GetDefaultLinkedParameterOrLinkedFunctionFromDe
 	else if(LinkedValues.Num() > 1)
 	{
 		// If there are a chain of linked values use the first one that's available, otherwise just use the last one.
-		TSet<FNiagaraAvailableParameterInfo> AvailableParameterInfos;
-		FGetAvailableParameterArgs Args;
-		Args.bIncludeConversionScripts = false;
-		GetAvailableParameters(AvailableParameterInfos, Args);
+		TArray<FNiagaraVariableBase> AvailableParameters;
+		TMap<FNiagaraVariableBase, UNiagaraScript*> ConversionHandles;
+		GetAvailableParameters(AvailableParameters, ConversionHandles, false);
 		for (FLinkedParameterOrFunctionNode& LinkedValue : LinkedValues)
 		{
 			if (LinkedValue.LinkedFunctionCallNode.IsValid() ||
-				AvailableParameterInfos.Contains(LinkedValue.LinkedParameter.GetValue()))
+				AvailableParameters.Contains(LinkedValue.LinkedParameter.GetValue()))
 			{
 				ValueToUse = &LinkedValue;
 				break;
diff --git a/Engine/Plugins/FX/Niagara/Source/NiagaraEditor/Private/ViewModels/Stack/NiagaraStackGraphUtilities.cpp b/Engine/Plugins/FX/Niagara/Source/NiagaraEditor/Private/ViewModels/Stack/NiagaraStackGraphUtilities.cpp
--- a/Engine/Plugins/FX/Niagara/Source/NiagaraEditor/Private/ViewModels/Stack/NiagaraStackGraphUtilities.cpp
+++ b/Engine/Plugins/FX/Niagara/Source/NiagaraEditor/Private/ViewModels/Stack/NiagaraStackGraphUtilities.cpp
@@ -49,7 +49,6 @@
 #include "NiagaraSimulationStageBase.h"
 #include "NiagaraEmitter.h"
 #include "DataHierarchyViewModelBase.h"
-#include "NiagaraTrace.h"
 #include "ViewModels/HierarchyEditor/NiagaraSummaryViewViewModel.h"
 #include "TraversalCache/StaticContextBuilder.h"
 #include "TraversalCache/TraversalCache.h"
@@ -548,8 +547,6 @@ TArray<FName> FNiagaraStackGraphUtilities::StackContextResolution(FVersionedNiag
 
 void FNiagaraStackGraphUtilities::BuildParameterMapHistoryWithStackContextResolution(FVersionedNiagaraEmitter OwningEmitter, UNiagaraNodeOutput* OutputNodeInChain, UNiagaraNode* NodeToVisit, TArray<FNiagaraParameterMapHistory>& OutHistories, bool bRecursive /*= true*/, bool bFilterForCompilation /*= true*/)
 {
-	TRACE_CPUPROFILER_EVENT_SCOPE_ON_CHANNEL(FNiagaraStackGraphUtilities::BuildParameterMapHistoryWithStackContextResolution, NiagaraChannel);
-	
 	FNiagaraParameterMapHistoryBuilder Builder;
 
 	bool bSetUsage = false;
diff --git a/Engine/Plugins/FX/Niagara/Source/NiagaraEditor/Private/ViewModels/Stack/NiagaraStackRenderItemGroup.cpp b/Engine/Plugins/FX/Niagara/Source/NiagaraEditor/Private/ViewModels/Stack/NiagaraStackRenderItemGroup.cpp
--- a/Engine/Plugins/FX/Niagara/Source/NiagaraEditor/Private/ViewModels/Stack/NiagaraStackRenderItemGroup.cpp
+++ b/Engine/Plugins/FX/Niagara/Source/NiagaraEditor/Private/ViewModels/Stack/NiagaraStackRenderItemGroup.cpp
@@ -54,21 +54,6 @@ public:
 		return FText();
 	}
 
-	virtual TOptional<FNiagaraFavoritesActionData> GetFavoritesData() const override
-	{
-		return TOptional<FNiagaraFavoritesActionData>();
-	}
-	
-	virtual bool IsInLibrary() const override
-	{
-		return true;
-	}
-	
-	virtual FNiagaraActionSourceData GetSourceData() const override
-	{
-		return FNiagaraActionSourceData();
-	}
-	
 	const FNiagaraRendererCreationInfo& GetRendererCreationInfo() { return RendererCreationInfo; }
 
 private:
diff --git a/Engine/Plugins/FX/Niagara/Source/NiagaraEditor/Private/ViewModels/Stack/NiagaraStackScriptItemGroup.cpp b/Engine/Plugins/FX/Niagara/Source/NiagaraEditor/Private/ViewModels/Stack/NiagaraStackScriptItemGroup.cpp
--- a/Engine/Plugins/FX/Niagara/Source/NiagaraEditor/Private/ViewModels/Stack/NiagaraStackScriptItemGroup.cpp
+++ b/Engine/Plugins/FX/Niagara/Source/NiagaraEditor/Private/ViewModels/Stack/NiagaraStackScriptItemGroup.cpp
@@ -15,7 +15,6 @@
 #include "ViewModels/NiagaraScriptGraphViewModel.h"
 #include "NiagaraStackEditorData.h"
 #include "ScopedTransaction.h"
-#include "Config/NiagaraFavoriteActionsConfig.h"
 #include "DragAndDrop/AssetDragDropOp.h"
 #include "Framework/Notifications/NotificationManager.h"
 #include "Interfaces/IMainFrameModule.h"
@@ -43,8 +42,6 @@
 
 using namespace UE::Niagara::Wizard;
 
-const FName INiagaraStackItemGroupAddAction::FavoriteActionsProfile = FName("StackGroupAdd");
-
 class FScriptGroupAddAction : public INiagaraStackItemGroupAddAction
 {
 public:
@@ -69,14 +66,9 @@ public:
 		AssetData.GetTagValue(GET_MEMBER_NAME_CHECKED(FVersionedNiagaraScriptData, Keywords), Keywords);
 
 		bool bSuggested = AssetData.GetTagValueRef<bool>(GET_MEMBER_NAME_CHECKED(FVersionedNiagaraScriptData, bSuggested));
-
-		FNiagaraFavoritesActionData FavoritesActionData;
-		FavoritesActionData.ActionIdentifier.Names.Add(FName(AssetData.GetObjectPathString()));
-		FavoritesActionData.bFavoriteByDefault = bSuggested;
-		
 		TTuple<EScriptSource, FText> Source = FNiagaraEditorUtilities::GetScriptSource(AssetData);
 		FNiagaraActionSourceData SourceData(Source.Key, Source.Value, true);
-		return MakeShareable(new FScriptGroupAddAction(DisplayName, {Category.ToString()}, Description, Keywords, FavoritesActionData, bIsInLibrary, FNiagaraVariable(), false, AssetData, nullptr, false, false, SourceData));
+		return MakeShareable(new FScriptGroupAddAction(DisplayName, {Category.ToString()}, Description, Keywords, bSuggested, bIsInLibrary, FNiagaraVariable(), false, AssetData, nullptr, false, false, SourceData));
 	}
 
 	static TSharedRef<FScriptGroupAddAction> CreateModuleActionFromScratchPadScript(UNiagaraScript* ModuleScript)
@@ -93,14 +85,11 @@ public:
 		FText DisplayName = FNiagaraEditorUtilities::FormatScriptName(ModuleScript->GetFName(), bIsInLibrary);
 		FText Description = FNiagaraEditorUtilities::FormatScriptDescription(ScriptData->Description, FSoftObjectPath(ModuleScript), bIsInLibrary);
 		FText Keywords = ScriptData->Keywords;
+		bool bSuggested = ScriptData->bSuggested;
 
-		FNiagaraFavoritesActionData FavoritesActionData;
-		FavoritesActionData.ActionIdentifier.Names.Add(FName(GetPathNameSafe(ModuleScript)));
-		FavoritesActionData.bFavoriteByDefault = ScriptData->bSuggested;
-		
 		FAssetData Data = FAssetData(ModuleScript);
 		FNiagaraActionSourceData SourceData(EScriptSource::Niagara, FText::FromString("Scratch Pad"), true);
-		return MakeShareable(new FScriptGroupAddAction(DisplayName, {Category.ToString()}, Description, Keywords, FavoritesActionData, bIsInLibrary, FNiagaraVariable(), false, FAssetData(), ModuleScript, false, false, SourceData));
+		return MakeShareable(new FScriptGroupAddAction(DisplayName, {Category.ToString()}, Description, Keywords, bSuggested, bIsInLibrary, FNiagaraVariable(), false, FAssetData(), ModuleScript, false, false, SourceData));
 	}
 
 	static TSharedRef<FScriptGroupAddAction> CreateExistingParameterModuleAction(FNiagaraVariable ParameterVariable)
@@ -111,7 +100,7 @@ public:
 		FText AttributeDescription = FNiagaraConstants::GetAttributeDescription(ParameterVariable);
 		FText Description = FText::Format(LOCTEXT("ExistingParameterModuleDescriptoinFormat", "Description: Set the parameter {0}. {1}"), FText::FromName(ParameterVariable.GetName()), AttributeDescription);
 
-		return MakeShareable(new FScriptGroupAddAction(DisplayName, {Category.ToString()}, Description, FText(), {}, true, ParameterVariable, false, FAssetData(), nullptr, false, false));
+		return MakeShareable(new FScriptGroupAddAction(DisplayName, {Category.ToString()}, Description, FText(), false, true, ParameterVariable, false, FAssetData(), nullptr, false, false));
 	}
 
 	static TSharedRef<FScriptGroupAddAction> CreateNewParameterModuleAction(FName NewParameterNamespace, FNiagaraTypeDefinition NewParameterType)
@@ -123,40 +112,26 @@ public:
 		FNiagaraParameterHandle NewParameterHandle(NewParameterNamespace, *(TEXT("New") + NewParameterType.GetName()));
 		FNiagaraVariable NewParameter(NewParameterType, NewParameterHandle.GetParameterHandleString());
 
-		return MakeShareable(new FScriptGroupAddAction( DisplayName, {Category.ToString()}, Description, FText(), {}, true, NewParameter, true, FAssetData(), nullptr, false, false));
+		return MakeShareable(new FScriptGroupAddAction( DisplayName, {Category.ToString()}, Description, FText(), false, true, NewParameter, true, FAssetData(), nullptr, false, false));
 	}
 
 	static TSharedRef<FScriptGroupAddAction> CreateNewScratchModuleAction()
 	{
 		FText DisplayName = LOCTEXT("NewScratchModuleName", "New Scratch Pad Module");
 		FText Description = LOCTEXT("NewScratchModuleDescription", "Description: Create a new scratch pad module.");
-		
-		FNiagaraFavoritesActionData FavoritesActionData;
-		FavoritesActionData.ActionIdentifier.Names.Add(FName(DisplayName.ToString()));
-		FavoritesActionData.bFavoriteByDefault = true;
-		
-		return MakeShareable(new FScriptGroupAddAction( DisplayName, {}, Description, FText(), FavoritesActionData, true, FNiagaraVariable(), false, FAssetData(), nullptr, true, false));
+		return MakeShareable(new FScriptGroupAddAction( DisplayName, {}, Description, FText(), true, true, FNiagaraVariable(), false, FAssetData(), nullptr, true, false));
 	}
 
 	static TSharedRef<FScriptGroupAddAction> CreateNewSetSpecificModuleAction()
 	{
 		FText DisplayName = LOCTEXT("NewSetSpecificModuleName", "Set new or existing parameter directly");
 		FText Description = LOCTEXT("NewSetSpecificModuleDescription", "Description: Create a new module that can set new or existing parameters directly.");
-		
-		FNiagaraFavoritesActionData FavoritesActionData;
-		FavoritesActionData.ActionIdentifier.Names.Add(FName(DisplayName.ToString()));
-		FavoritesActionData.bFavoriteByDefault = true;
-		
-		return MakeShareable(new FScriptGroupAddAction( DisplayName, {}, Description, FText(), FavoritesActionData, true, FNiagaraVariable(), false, FAssetData(), nullptr, false, true));
+		return MakeShareable(new FScriptGroupAddAction( DisplayName, {}, Description, FText(), true, true, FNiagaraVariable(), false, FAssetData(), nullptr, false, true));
 	}
 
 	static TSharedRef<FScriptGroupAddAction> CreateWizardAction(const FModuleWizardGenerator::FAction& Action)
 	{
-		FNiagaraFavoritesActionData FavoritesActionData;
-		FavoritesActionData.ActionIdentifier.Names.Add(Action.WizardModel->GetIdentifier());
-		FavoritesActionData.bFavoriteByDefault = Action.bSuggestedAction;
-		
-		TSharedRef<FScriptGroupAddAction> AddAction = MakeShareable(new FScriptGroupAddAction( Action.DisplayName, {}, Action.Description, Action.Keywords, FavoritesActionData, true, FNiagaraVariable(), false, FAssetData(), nullptr, false, true));
+		TSharedRef<FScriptGroupAddAction> AddAction = MakeShareable(new FScriptGroupAddAction( Action.DisplayName, {}, Action.Description, Action.Keywords, Action.bSuggestedAction, true, FNiagaraVariable(), false, FAssetData(), nullptr, false, true));
 		AddAction->ModuleWizardModel = Action.WizardModel;
 		return AddAction;
 	}
@@ -181,9 +156,9 @@ public:
 		return Keywords;
 	}
 
-	virtual TOptional<FNiagaraFavoritesActionData> GetFavoritesData() const override
+	virtual bool GetSuggested() const override
 	{
-		return FavoritesData;
+		return bSuggested;
 	}
 
 	virtual bool IsInLibrary() const override
@@ -238,7 +213,7 @@ public:
 	}
 
 private:
-	FScriptGroupAddAction(FText InDisplayName, TArray<FString> InCategories, FText InDescription, FText InKeywords, TOptional<FNiagaraFavoritesActionData> InFavoritesData, bool bInIsInLibrary,
+	FScriptGroupAddAction(FText InDisplayName, TArray<FString> InCategories, FText InDescription, FText InKeywords, bool bInSuggested, bool bInIsInLibrary,
 		FNiagaraVariable InModuleParameterVariable, bool bInRenameParameterOnAdd,
 		FAssetData InModuleAssetData, UNiagaraScript* InModuleScript,
 		bool bInIsNewScratchModuleAction, bool bInIsNewSetSpecificModuleAction, FNiagaraActionSourceData InSourceData = FNiagaraActionSourceData(EScriptSource::Niagara, FText::FromString(TEXT("Niagara")), true))
@@ -247,7 +222,7 @@ private:
 		, Categories(InCategories)
 		, Description(InDescription)
 		, Keywords(InKeywords)
-		, FavoritesData(InFavoritesData)
+		, bSuggested(bInSuggested)
 		, bIsInLibrary(bInIsInLibrary)
 		, ModuleParameterVariable(InModuleParameterVariable)
 		, bRenameParameterOnAdd(bInRenameParameterOnAdd)
@@ -263,7 +238,7 @@ private:
 	TArray<FString> Categories;
 	FText Description;
 	FText Keywords;
-	TOptional<FNiagaraFavoritesActionData> FavoritesData;
+	bool bSuggested;
 	bool bIsInLibrary;
 	FNiagaraVariable ModuleParameterVariable;
 	bool bRenameParameterOnAdd;
diff --git a/Engine/Plugins/FX/Niagara/Source/NiagaraEditor/Private/ViewModels/Stack/NiagaraStackStatelessEmitterSimulateGroup.cpp b/Engine/Plugins/FX/Niagara/Source/NiagaraEditor/Private/ViewModels/Stack/NiagaraStackStatelessEmitterSimulateGroup.cpp
--- a/Engine/Plugins/FX/Niagara/Source/NiagaraEditor/Private/ViewModels/Stack/NiagaraStackStatelessEmitterSimulateGroup.cpp
+++ b/Engine/Plugins/FX/Niagara/Source/NiagaraEditor/Private/ViewModels/Stack/NiagaraStackStatelessEmitterSimulateGroup.cpp
@@ -95,9 +95,6 @@ public:
 	virtual FText GetDisplayName() const override { return DisplayName; }
 	virtual FText GetDescription() const override { return FText(); }
 	virtual FText GetKeywords() const override { return FText(); }
-	virtual TOptional<FNiagaraFavoritesActionData> GetFavoritesData() const override { return TOptional<FNiagaraFavoritesActionData>(); }
-	virtual bool IsInLibrary() const override { return true; }
-	virtual FNiagaraActionSourceData GetSourceData() const override { return FNiagaraActionSourceData(); }
 
 private:
 	TWeakObjectPtr<UNiagaraStatelessModule> StatelessModuleWeak;
diff --git a/Engine/Plugins/FX/Niagara/Source/NiagaraEditor/Private/ViewModels/Stack/NiagaraStackStatelessEmitterSpawnGroup.cpp b/Engine/Plugins/FX/Niagara/Source/NiagaraEditor/Private/ViewModels/Stack/NiagaraStackStatelessEmitterSpawnGroup.cpp
--- a/Engine/Plugins/FX/Niagara/Source/NiagaraEditor/Private/ViewModels/Stack/NiagaraStackStatelessEmitterSpawnGroup.cpp
+++ b/Engine/Plugins/FX/Niagara/Source/NiagaraEditor/Private/ViewModels/Stack/NiagaraStackStatelessEmitterSpawnGroup.cpp
@@ -106,9 +106,6 @@ public:
 	virtual FText GetDisplayName() const override { return DisplayName; }
 	virtual FText GetDescription() const override { return FText(); }
 	virtual FText GetKeywords() const override { return FText(); }
-	virtual TOptional<FNiagaraFavoritesActionData> GetFavoritesData() const override { return TOptional<FNiagaraFavoritesActionData>(); }
-	virtual bool IsInLibrary() const override { return true; }
-	virtual FNiagaraActionSourceData GetSourceData() const override { return FNiagaraActionSourceData(); }
 
 private:
 	ENiagaraStatelessSpawnInfoType SpawnInfoType;
diff --git a/Engine/Plugins/FX/Niagara/Source/NiagaraEditor/Private/ViewModels/Stack/NiagaraStackSystemSettingsGroup.cpp b/Engine/Plugins/FX/Niagara/Source/NiagaraEditor/Private/ViewModels/Stack/NiagaraStackSystemSettingsGroup.cpp
--- a/Engine/Plugins/FX/Niagara/Source/NiagaraEditor/Private/ViewModels/Stack/NiagaraStackSystemSettingsGroup.cpp
+++ b/Engine/Plugins/FX/Niagara/Source/NiagaraEditor/Private/ViewModels/Stack/NiagaraStackSystemSettingsGroup.cpp
@@ -51,21 +51,6 @@ public:
 		return FText();
 	}
 
-	virtual TOptional<FNiagaraFavoritesActionData> GetFavoritesData() const override
-	{
-		return TOptional<FNiagaraFavoritesActionData>();
-	}
-	
-	virtual bool IsInLibrary() const override
-	{
-		return true;
-	}
-	
-	virtual FNiagaraActionSourceData GetSourceData() const override
-	{
-		return FNiagaraActionSourceData();
-	}
-
 private:
 	FNiagaraVariable NewParameterVariable;
 };
diff --git a/Engine/Plugins/FX/Niagara/Source/NiagaraEditor/Private/Widgets/DataChannel/NiagaraDataChannelWizard.cpp b/Engine/Plugins/FX/Niagara/Source/NiagaraEditor/Private/Widgets/DataChannel/NiagaraDataChannelWizard.cpp
--- a/Engine/Plugins/FX/Niagara/Source/NiagaraEditor/Private/Widgets/DataChannel/NiagaraDataChannelWizard.cpp
+++ b/Engine/Plugins/FX/Niagara/Source/NiagaraEditor/Private/Widgets/DataChannel/NiagaraDataChannelWizard.cpp
@@ -542,11 +542,6 @@ TSharedRef<FModuleWizardModel> DataChannel::CreateReadNDCModuleWizardModel()
 			Pages.Add(VariablesPage.ToSharedRef());
 		}
 		virtual ~FReadNDCModel() override = default;
-
-		virtual FName GetIdentifier() const override
-		{
-			return "ReadNDCWizardModel";
-		}
 		
 		virtual void GenerateNewModuleContent(TSharedPtr<FNiagaraScratchPadScriptViewModel> ScratchPadScriptViewModel, const TArray<const UNiagaraNodeFunctionCall*>& PreviousModules) override
 		{
@@ -706,11 +701,6 @@ TSharedRef<FModuleWizardModel> DataChannel::CreateWriteNDCModuleWizardModel()
 			Pages.Add(VariablesPage.ToSharedRef());
 		}
 		virtual ~FWriteNDCModel() override = default;
-
-		virtual FName GetIdentifier() const override
-		{
-			return "WriteNDCWizardModel";
-		}
 		
 		virtual void GenerateNewModuleContent(TSharedPtr<FNiagaraScratchPadScriptViewModel> ScratchPadScriptViewModel, const TArray<const UNiagaraNodeFunctionCall*>& PreviousModules) override
 		{
@@ -886,11 +876,6 @@ TSharedRef<FModuleWizardModel> DataChannel::CreateSpawnNDCModuleWizardModel()
 		}
 		virtual ~FSpawnNDCModel() override = default;
 
-		virtual FName GetIdentifier() const override
-		{
-			return "SpawnNDCWizardModel";
-		}
-		
 		virtual TArray<FModuleCreationEntry> GetModulesToCreate(UNiagaraNodeOutput* ProvidedOutputNode, int32 ProvidedTargetIndex, TSharedPtr<FNiagaraSystemViewModel> SystemModel, TSharedPtr<FNiagaraEmitterViewModel> EmitterViewModel) override
 		{
 			TArray<FModuleCreationEntry> Result;
diff --git a/Engine/Plugins/FX/Niagara/Source/NiagaraEditor/Private/Widgets/SNiagaraGraphActionMenu.cpp b/Engine/Plugins/FX/Niagara/Source/NiagaraEditor/Private/Widgets/SNiagaraGraphActionMenu.cpp
--- a/Engine/Plugins/FX/Niagara/Source/NiagaraEditor/Private/Widgets/SNiagaraGraphActionMenu.cpp
+++ b/Engine/Plugins/FX/Niagara/Source/NiagaraEditor/Private/Widgets/SNiagaraGraphActionMenu.cpp
@@ -11,7 +11,6 @@
 #include "NiagaraEditorStyle.h"
 #include "NiagaraEditorUtilities.h"
 #include "SNiagaraGraphActionWidget.h"
-#include "Config/NiagaraFavoriteActionsConfig.h"
 #include "Widgets/Text/SRichTextBlock.h"
 #include "Widgets/Text/STextBlock.h"
 #include "Widgets/SItemSelector.h"
@@ -19,7 +18,6 @@
 #define LOCTEXT_NAMESPACE "NiagaraGraphActionMenu"
 
 bool SNiagaraGraphActionMenu::bLibraryOnly = true;
-const FName SNiagaraGraphActionMenu::FavoriteActionsProfile = FName("NiagaraGraphActions");
 
 SNiagaraGraphActionMenu::~SNiagaraGraphActionMenu()
 {
@@ -76,7 +74,7 @@ void SNiagaraGraphActionMenu::Construct( const FArguments& InArgs )
 					.OnGenerateWidgetForItem(this, &SNiagaraGraphActionMenu::OnGenerateWidgetForItem)
 					.OnGetItemWeight_Lambda([](const TSharedPtr<FNiagaraAction_NewNode>& InCurrentAction, const TArray<FString>& InFilterTerms)
 					{
-						return FNiagaraEditorUtilities::GetWeightForItem(InCurrentAction, SNiagaraGraphActionMenu::FavoriteActionsProfile, InFilterTerms);
+						return FNiagaraEditorUtilities::GetWeightForItem(InCurrentAction, InFilterTerms);
 					})
 					.OnItemActivated(this, &SNiagaraGraphActionMenu::OnItemActivated)
 					.AllowMultiselect(false)
@@ -84,10 +82,6 @@ void SNiagaraGraphActionMenu::Construct( const FArguments& InArgs )
 					.OnDoesSectionPassCustomFilter(this, &SNiagaraGraphActionMenu::DoesSectionPassCustomFilter)
 					.ClickActivateMode(EItemSelectorClickActivateMode::SingleClick)
 					.ExpandInitially(false)
-					.OnItemRowHoverEvent_Lambda([](const TSharedPtr<FNiagaraAction_NewNode>& Item, bool bHover)
-					{
-						Item->bIsHovered = bHover;
-					})
 					.OnGetSectionData_Lambda([](const ENiagaraMenuSections& Section)
 					{
 						if(Section == ENiagaraMenuSections::Suggested)
@@ -126,19 +120,12 @@ TArray<FString> SNiagaraGraphActionMenu::OnGetCategoriesForItem(const TSharedPtr
 
 TArray<ENiagaraMenuSections> SNiagaraGraphActionMenu::OnGetSectionsForItem(const TSharedPtr<FNiagaraAction_NewNode>& Item)
 {	
-	TArray<ENiagaraMenuSections> Sections{ENiagaraMenuSections::General};
-
-	if(Item->FavoritesActionData.IsSet())
+	if(Item->Section == ENiagaraMenuSections::Suggested)
 	{
-		FNiagaraFavoriteActionsProfile& ActionsProfile = UNiagaraFavoriteActionsConfig::Get()->GetActionsProfile(SNiagaraGraphActionMenu::FavoriteActionsProfile);
-		TOptional<bool> bIsFavorite = ActionsProfile.IsFavorite(Item->FavoritesActionData.GetValue());
-		if(bIsFavorite.IsSet() && bIsFavorite.GetValue())
-		{
-			Sections.Add(ENiagaraMenuSections::Suggested);
-		}
+		return { ENiagaraMenuSections::General, ENiagaraMenuSections::Suggested };
 	}
 		
-	return Sections;
+	return {Item->Section};
 }
 
 bool SNiagaraGraphActionMenu::OnCompareSectionsForEquality(const ENiagaraMenuSections& SectionA, const ENiagaraMenuSections& SectionB)
@@ -196,7 +183,6 @@ TSharedRef<SWidget> SNiagaraGraphActionMenu::OnGenerateWidgetForItem(const TShar
 {
 	FCreateNiagaraWidgetForActionData ActionData(Item);
 	ActionData.HighlightText = TAttribute<FText>::Create(TAttribute<FText>::FGetter::CreateRaw(this, &SNiagaraGraphActionMenu::GetFilterText));
-	ActionData.FavoriteActionsProfileName = SNiagaraGraphActionMenu::FavoriteActionsProfile;
 	return SNew(SNiagaraActionWidget, ActionData);
 }
 
diff --git a/Engine/Plugins/FX/Niagara/Source/NiagaraEditor/Private/Widgets/SNiagaraGraphActionWidget.cpp b/Engine/Plugins/FX/Niagara/Source/NiagaraEditor/Private/Widgets/SNiagaraGraphActionWidget.cpp
--- a/Engine/Plugins/FX/Niagara/Source/NiagaraEditor/Private/Widgets/SNiagaraGraphActionWidget.cpp
+++ b/Engine/Plugins/FX/Niagara/Source/NiagaraEditor/Private/Widgets/SNiagaraGraphActionWidget.cpp
@@ -10,7 +10,6 @@
 #include "Widgets/SNiagaraParameterName.h"
 #include "NiagaraEditorStyle.h"
 #include "NiagaraEditorUtilities.h"
-#include "Config/NiagaraFavoriteActionsConfig.h"
 #include "Widgets/Layout/SSpacer.h"
 #include "Widgets/Layout/SSeparator.h"
 
@@ -83,7 +82,6 @@ FReply SNiagaraGraphActionWidget::OnMouseButtonDown(const FGeometry& MyGeometry,
 void SNiagaraActionWidget::Construct(const FArguments& InArgs, const FCreateNiagaraWidgetForActionData& InCreateData)
 {
 	ActionPtr = InCreateData.Action;
-	FavoriteProfileName = InCreateData.FavoriteActionsProfileName;
 
 	TSharedPtr<SWidget> NameWidget;
 	if (ActionPtr->GetParameterVariable().IsSet())
@@ -125,11 +123,11 @@ void SNiagaraActionWidget::Construct(const FArguments& InArgs, const FCreateNiag
 			NameWidget.ToSharedRef()
 		]
 		+ SHorizontalBox::Slot()
-		.FillWidth(1.f)
-		[
+       .FillWidth(1.f)
+       [
            SNew(SSpacer)
-		]
-		+ SHorizontalBox::Slot()
+       ]
+       + SHorizontalBox::Slot()
         .AutoWidth()
         .VAlign(VAlign_Center)
         .HAlign(HAlign_Right)
@@ -139,23 +137,15 @@ void SNiagaraActionWidget::Construct(const FArguments& InArgs, const FCreateNiag
             .Visibility(ActionPtr->bIsExperimental ? EVisibility::Visible : EVisibility::Collapsed)
             .ToolTipText(ActionPtr->bIsExperimental ? LOCTEXT("ScriptExperimentalToolTip", "This script is experimental, use with care!") : FText::GetEmpty())
         ]
-		+ SHorizontalBox::Slot()
-		.HAlign(HAlign_Right)
-		.VAlign(VAlign_Center)
-		.AutoWidth()
-		[
-			CreateFavoriteActionWidget()
-		]
-		+ SHorizontalBox::Slot()
-		.HAlign(HAlign_Right)
-		.VAlign(VAlign_Fill)
-		.AutoWidth()
-		[
+       + SHorizontalBox::Slot()
+       .HAlign(HAlign_Right)
+       .VAlign(VAlign_Fill)
+       [
            SNew(SSeparator)
            .SeparatorImage(FAppStyle::Get().GetBrush("Separator"))
            .Orientation(EOrientation::Orient_Vertical)
-			.Visibility(ActionPtr->SourceData.bDisplaySource ? EVisibility::Visible : EVisibility::Collapsed)
-		]
+           .Visibility(ActionPtr->SourceData.bDisplaySource ? EVisibility::Visible : EVisibility::Collapsed)
+       ]
        + SHorizontalBox::Slot()
        .HAlign(HAlign_Fill)
        .VAlign(VAlign_Center)
@@ -175,91 +165,4 @@ void SNiagaraActionWidget::Construct(const FArguments& InArgs, const FCreateNiag
     ];
 }
 
-TSharedRef<SWidget> SNiagaraActionWidget::CreateFavoriteActionWidget()
-{	
-	return
-		SNew(SBox)
-		.WidthOverride(16.f)
-		.HeightOverride(16.f)
-		[
-			SNew(SButton)
-			.ContentPadding(FMargin(0.f))
-			.ButtonStyle(&FAppStyle::Get().GetWidgetStyle<FButtonStyle>("HoverHintOnly"))
-			.OnClicked(this, &SNiagaraActionWidget::OnFavoriteStateChanged)
-			.Visibility(this, &SNiagaraActionWidget::OnShouldShowFavoriteButton)
-			.ToolTipText(this, &SNiagaraActionWidget::GetFavoriteButtonTooltipText)
-			[
-				SNew(SImage)
-				.Image(this, &SNiagaraActionWidget::GetFavoriteBrush)
-			]
-		];
-}
-
-FReply SNiagaraActionWidget::OnFavoriteStateChanged()
-{
-	// We still want the button to consume the input, but generally this will never happen as the button isn't displayed
-	if(FavoriteProfileName.IsSet() == false || ActionPtr->FavoritesActionData.IsSet() == false)
-	{
-		return FReply::Handled();
-	}
-	
-	FNiagaraFavoriteActionsProfile& ActionsProfile = UNiagaraFavoriteActionsConfig::Get()->GetActionsProfile(FavoriteProfileName.GetValue());
-	ActionsProfile.ToggleFavoriteAction(ActionPtr->FavoritesActionData.GetValue());
-
-	return FReply::Handled();
-}
-
-ECheckBoxState SNiagaraActionWidget::OnGetIsFavorite() const
-{
-	if(FavoriteProfileName.IsSet() == false || ActionPtr->FavoritesActionData.IsSet() == false)
-	{
-		return ECheckBoxState::Undetermined;
-	}
-		
-	FNiagaraFavoriteActionsProfile& ActionsProfile = UNiagaraFavoriteActionsConfig::Get()->GetActionsProfile(FavoriteProfileName.GetValue());
-	return ActionsProfile.IsFavorite(ActionPtr->FavoritesActionData.GetValue()) ? ECheckBoxState::Checked : ECheckBoxState::Unchecked;
-}
-
-EVisibility SNiagaraActionWidget::OnShouldShowFavoriteButton() const
-{
-	if(FavoriteProfileName.IsSet() == false || ActionPtr->FavoritesActionData.IsSet() == false)
-	{
-		return EVisibility::Collapsed;
-	}
-
-	return ActionPtr->bIsHovered ? EVisibility::Visible : EVisibility::Collapsed;
-}
-
-const FSlateBrush* SNiagaraActionWidget::GetFavoriteBrush() const
-{
-	if(FavoriteProfileName.IsSet() == false || ActionPtr->FavoritesActionData.IsSet() == false)
-	{
-		return FAppStyle::GetNoBrush();
-	}
-	
-	FNiagaraFavoriteActionsProfile& ActionsProfile = UNiagaraFavoriteActionsConfig::Get()->GetActionsProfile(FavoriteProfileName.GetValue());	
-	if(ActionsProfile.IsFavorite(ActionPtr->FavoritesActionData.GetValue()))
-	{
-		return FAppStyle::GetBrush("Icons.Star");
-	}
-	
-	return FAppStyle::GetBrush("Icons.Star.Outline");
-}
-
-FText SNiagaraActionWidget::GetFavoriteButtonTooltipText() const
-{
-	if(FavoriteProfileName.IsSet() == false || ActionPtr->FavoritesActionData.IsSet() == false)
-	{
-		return FText::GetEmpty();
-	}
-
-	FNiagaraFavoriteActionsProfile& ActionsProfile = UNiagaraFavoriteActionsConfig::Get()->GetActionsProfile(FavoriteProfileName.GetValue());
-	if(ActionsProfile.IsFavorite(ActionPtr->FavoritesActionData.GetValue()))
-	{
-		return LOCTEXT("UnfavoriteButtonTooltip", "Remove this action from favorites");
-	}
-
-	return LOCTEXT("FavoriteButtonTooltip", "Add this action to favorites");
-}
-
 #undef LOCTEXT_NAMESPACE
diff --git a/Engine/Plugins/FX/Niagara/Source/NiagaraEditor/Public/Config/NiagaraFavoriteActionsConfig.h b/Engine/Plugins/FX/Niagara/Source/NiagaraEditor/Public/Config/NiagaraFavoriteActionsConfig.h
deleted file mode 100644
diff --git a/Engine/Plugins/FX/Niagara/Source/NiagaraEditor/Public/NiagaraActions.h b/Engine/Plugins/FX/Niagara/Source/NiagaraEditor/Public/NiagaraActions.h
--- a/Engine/Plugins/FX/Niagara/Source/NiagaraEditor/Public/NiagaraActions.h
+++ b/Engine/Plugins/FX/Niagara/Source/NiagaraEditor/Public/NiagaraActions.h
@@ -9,9 +9,9 @@
 #include "DragAndDrop/DecoratedDragDropOp.h"
 #include "NiagaraEditorCommon.h"
 #include "ViewModels/Stack/NiagaraParameterHandle.h"
-#include "Config/NiagaraFavoriteActionsConfig.h"
 #include "NiagaraActions.generated.h"
 
+
 class UToolMenu;
 class UGraphNodeContextMenuContext;
 
@@ -112,18 +112,21 @@ USTRUCT()
 struct FNiagaraMenuAction_Base
 {
 	GENERATED_USTRUCT_BODY();
-	
+
 	DECLARE_DELEGATE(FOnExecuteAction);
 	DECLARE_DELEGATE_RetVal(bool, FCanExecuteAction);
 
 	FNiagaraMenuAction_Base() {}
-	NIAGARAEDITOR_API FNiagaraMenuAction_Base(FText DisplayName, TArray<FString> InNodeCategories, TOptional<FNiagaraFavoritesActionData> FavoritesActionData, FText InToolTip, FText InKeywords, float InIntrinsicWeightMultiplier = 1.f);
+	NIAGARAEDITOR_API FNiagaraMenuAction_Base(FText DisplayName, ENiagaraMenuSections Section, TArray<FString> InNodeCategories, FText InToolTip, FText InKeywords, float InIntrinsicWeightMultiplier = 1.f);
 
 	NIAGARAEDITOR_API void UpdateFullSearchText();
 
 	bool bIsExperimental = false;
 
 	bool bIsInLibrary = true;
+
+	/** Top level section this action belongs to. */
+	ENiagaraMenuSections Section;
 	
 	/** Nested categories below a top level section. Can be empty */
 	TArray<FString> Categories;
@@ -148,10 +151,6 @@ struct FNiagaraMenuAction_Base
 
 	/** A multiplier intrinsic to the action. Can be used to tweak search relevance depending on context (like module actions in a module script) */
 	float SearchWeightMultiplier = 1.f;
-	
-	TOptional<FNiagaraFavoritesActionData> FavoritesActionData;
-	
-	mutable bool bIsHovered = false;
 };
 
 USTRUCT()
@@ -162,16 +161,16 @@ struct FNiagaraMenuAction_Generic : public FNiagaraMenuAction_Base
 	FNiagaraMenuAction_Generic() {}
 
 	FNiagaraMenuAction_Generic(FOnExecuteAction ExecuteAction, FCanExecuteAction InCanExecuteAction,
-		FText InDisplayName, TArray<FString> InNodeCategories, TOptional<FNiagaraFavoritesActionData> FavoritesActionData,  FText InToolTip, FText InKeywords, float InIntrinsicWeightMultiplier = 1.f)
-    : FNiagaraMenuAction_Base(InDisplayName, InNodeCategories, FavoritesActionData, InToolTip, InKeywords, InIntrinsicWeightMultiplier)
+		FText InDisplayName, ENiagaraMenuSections Section, TArray<FString> InNodeCategories, FText InToolTip, FText InKeywords, float InIntrinsicWeightMultiplier = 1.f)
+    : FNiagaraMenuAction_Base(InDisplayName, Section, InNodeCategories, InToolTip, InKeywords, InIntrinsicWeightMultiplier)
 	{
 		Action = ExecuteAction;
 		CanExecuteAction = InCanExecuteAction;
 	}
 
 	FNiagaraMenuAction_Generic(FOnExecuteAction ExecuteAction,
-        FText InDisplayName, TArray<FString> InNodeCategories, TOptional<FNiagaraFavoritesActionData> FavoritesActionData, FText InToolTip, FText InKeywords, float InIntrinsicWeightMultiplier = 1.f)
-    : FNiagaraMenuAction_Base(InDisplayName, InNodeCategories, FavoritesActionData, InToolTip, InKeywords, InIntrinsicWeightMultiplier)
+        FText InDisplayName, ENiagaraMenuSections Section, TArray<FString> InNodeCategories, FText InToolTip, FText InKeywords, float InIntrinsicWeightMultiplier = 1.f)
+    : FNiagaraMenuAction_Base(InDisplayName, Section, InNodeCategories, InToolTip, InKeywords, InIntrinsicWeightMultiplier)
 	{
 		Action = ExecuteAction;
 	}
@@ -207,8 +206,8 @@ struct FNiagaraAction_NewNode : public FNiagaraMenuAction_Generic
 
 	FNiagaraAction_NewNode() {}
 	FNiagaraAction_NewNode(
-		FText InDisplayName, TArray<FString> InNodeCategories, TOptional<FNiagaraFavoritesActionData> FavoritesActionData, FText InToolTip, FText InKeywords, float InIntrinsicWeightMultiplier = 1.f)
-		: FNiagaraMenuAction_Generic(FOnExecuteAction(), InDisplayName, InNodeCategories, FavoritesActionData, InToolTip, InKeywords, InIntrinsicWeightMultiplier)
+		FText InDisplayName, ENiagaraMenuSections Section, TArray<FString> InNodeCategories, FText InToolTip, FText InKeywords, float InIntrinsicWeightMultiplier = 1.f)
+		: FNiagaraMenuAction_Generic(FOnExecuteAction(), InDisplayName, Section, InNodeCategories, InToolTip, InKeywords, InIntrinsicWeightMultiplier)
 	{		
 	}
 
diff --git a/Engine/Plugins/FX/Niagara/Source/NiagaraEditor/Public/NiagaraEditorUtilities.h b/Engine/Plugins/FX/Niagara/Source/NiagaraEditor/Public/NiagaraEditorUtilities.h
--- a/Engine/Plugins/FX/Niagara/Source/NiagaraEditor/Public/NiagaraEditorUtilities.h
+++ b/Engine/Plugins/FX/Niagara/Source/NiagaraEditor/Public/NiagaraEditorUtilities.h
@@ -253,7 +253,7 @@ namespace FNiagaraEditorUtilities
 
 	NIAGARAEDITOR_API bool IsEnginePluginAsset(const FTopLevelAssetPath& InTopLevelAssetPath);
 
-	NIAGARAEDITOR_API int32 GetWeightForItem(const TSharedPtr<FNiagaraMenuAction_Generic>& Item, FName FavoriteActionsProfile, const TArray<FString>& FilterTerms);
+	NIAGARAEDITOR_API int32 GetWeightForItem(const TSharedPtr<FNiagaraMenuAction_Generic>& Item, const TArray<FString>& FilterTerms);
 
 	NIAGARAEDITOR_API bool DoesItemMatchFilterText(const FText& FilterText, const TSharedPtr<FNiagaraMenuAction_Generic>& Item);
 	
@@ -356,8 +356,6 @@ namespace FNiagaraEditorUtilities
 
 	NIAGARAEDITOR_API void GetAvailableParameterCollections(TArray<UNiagaraParameterCollection*>& OutParameterCollections);
 
-	NIAGARAEDITOR_API void GetAvailableParameterDefinitions(TArray<UNiagaraParameterDefinitions*>& OutParameterDefinitions);
-
 	TSharedPtr<INiagaraParameterDefinitionsSubscriberViewModel> GetOwningLibrarySubscriberViewModelForGraph(const UNiagaraGraph* Graph);
 
 	TArray<UNiagaraParameterDefinitions*> DowncastParameterDefinitionsBaseArray(const TArray<UNiagaraParameterDefinitionsBase*> BaseArray);
diff --git a/Engine/Plugins/FX/Niagara/Source/NiagaraEditor/Public/NiagaraParameterDefinitions.h b/Engine/Plugins/FX/Niagara/Source/NiagaraEditor/Public/NiagaraParameterDefinitions.h
--- a/Engine/Plugins/FX/Niagara/Source/NiagaraEditor/Public/NiagaraParameterDefinitions.h
+++ b/Engine/Plugins/FX/Niagara/Source/NiagaraEditor/Public/NiagaraParameterDefinitions.h
@@ -73,7 +73,7 @@ public:
 	/** Getters for script variables. 
 	 *  NOTE: These public getters are implemented to support details views. Editing UNiagaraScriptVariable definitions publicly should be done via the add/remove/rename interface otherwise.
 	 */
-	NIAGARAEDITOR_API UNiagaraScriptVariable* GetScriptVariable(const FNiagaraVariable& Var);
+	UNiagaraScriptVariable* GetScriptVariable(const FNiagaraVariable& Var);
 	UNiagaraScriptVariable* GetScriptVariable(const FGuid& ScriptVarId);
 
 	//~ Begin Subscribed ParameterDefinitions methods
diff --git a/Engine/Plugins/FX/Niagara/Source/NiagaraEditor/Public/SNiagaraGraphActionWidget.h b/Engine/Plugins/FX/Niagara/Source/NiagaraEditor/Public/SNiagaraGraphActionWidget.h
--- a/Engine/Plugins/FX/Niagara/Source/NiagaraEditor/Public/SNiagaraGraphActionWidget.h
+++ b/Engine/Plugins/FX/Niagara/Source/NiagaraEditor/Public/SNiagaraGraphActionWidget.h
@@ -29,8 +29,6 @@ struct FCreateNiagaraWidgetForActionData
 	TAttribute<FText> HighlightText;
 	/** True if the widget should be read only - no renaming allowed */
 	bool bIsReadOnly = true;
-	/** If specified, adds a 'favorite' toggle to the widget */
-	TOptional<FName> FavoriteActionsProfileName;
 };
 
 /** Custom widget for GraphActionMenu */
@@ -39,7 +37,7 @@ class SNiagaraGraphActionWidget : public SCompoundWidget
 	public:
 
 	SLATE_BEGIN_ARGS( SNiagaraGraphActionWidget ) {}
-		SLATE_ATTRIBUTE(FText, HighlightText)
+	SLATE_ATTRIBUTE(FText, HighlightText)
 SLATE_END_ARGS()
 
 NIAGARAEDITOR_API void Construct(const FArguments& InArgs, const FCreateWidgetForActionData* InCreateData);
@@ -59,25 +57,17 @@ public:
 
 	SLATE_BEGIN_ARGS( SNiagaraActionWidget )
 		: _bShowTypeIfParameter(true)
-		{} 
+	{} 
 		SLATE_ARGUMENT(bool, bShowTypeIfParameter)
 	SLATE_END_ARGS()
 
 	NIAGARAEDITOR_API void Construct(const FArguments& InArgs, const FCreateNiagaraWidgetForActionData& InCreateData);
 	//virtual FReply OnMouseButtonDown( const FGeometry& MyGeometry, const FPointerEvent& MouseEvent ) override;
 
-private:
-	TSharedRef<SWidget> CreateFavoriteActionWidget();
-	FReply OnFavoriteStateChanged();
-	ECheckBoxState OnGetIsFavorite() const;
-	EVisibility OnShouldShowFavoriteButton() const;
-	const FSlateBrush* GetFavoriteBrush() const;
-	FText GetFavoriteButtonTooltipText() const;
 private:
 	/** The item that we want to display with this widget */
 	TSharedPtr<struct FNiagaraMenuAction_Generic> ActionPtr;
 	/** Delegate executed when mouse button goes down */
 	FOnActionClicked MouseButtonDownDelegate;
 	FNiagaraActionSourceData SourceData;
-	TOptional<FName> FavoriteProfileName;
 };
diff --git a/Engine/Plugins/FX/Niagara/Source/NiagaraEditor/Public/ViewModels/Stack/INiagaraStackItemGroupAddUtilities.h b/Engine/Plugins/FX/Niagara/Source/NiagaraEditor/Public/ViewModels/Stack/INiagaraStackItemGroupAddUtilities.h
--- a/Engine/Plugins/FX/Niagara/Source/NiagaraEditor/Public/ViewModels/Stack/INiagaraStackItemGroupAddUtilities.h
+++ b/Engine/Plugins/FX/Niagara/Source/NiagaraEditor/Public/ViewModels/Stack/INiagaraStackItemGroupAddUtilities.h
@@ -7,8 +7,6 @@
 class INiagaraStackItemGroupAddAction
 {
 public:
-	NIAGARAEDITOR_API static const FName FavoriteActionsProfile;
-	
 	/** Gets the category for this action. */
 	virtual TArray<FString> GetCategories() const = 0;
 
@@ -20,15 +18,24 @@ public:
 
 	/** Gets a space separated string of keywords which expose additional search terms for this action. */
 	virtual FText GetKeywords() const = 0;
-	
-	/** Required for favorites feature. */
-	virtual TOptional<FNiagaraFavoritesActionData> GetFavoritesData() const = 0;
+
+	/** Gets a bool that indicates whether this action is suggested or not */
+	virtual bool GetSuggested() const
+	{
+		return false;
+	}
 
 	/** Indicates if this action represents a library action */
-	virtual bool IsInLibrary() const = 0;
+	virtual bool IsInLibrary() const
+	{
+		return true;
+	}
 
 	/** Gets a source data struct that is useful to display additional information about an action */
-	virtual FNiagaraActionSourceData GetSourceData() const = 0;
+	virtual FNiagaraActionSourceData GetSourceData() const
+	{
+		return FNiagaraActionSourceData();
+	}
 
 	virtual ~INiagaraStackItemGroupAddAction() { }
 };
diff --git a/Engine/Plugins/FX/Niagara/Source/NiagaraEditor/Public/ViewModels/Stack/NiagaraStackEmitterPropertiesGroup.h b/Engine/Plugins/FX/Niagara/Source/NiagaraEditor/Public/ViewModels/Stack/NiagaraStackEmitterPropertiesGroup.h
--- a/Engine/Plugins/FX/Niagara/Source/NiagaraEditor/Public/ViewModels/Stack/NiagaraStackEmitterPropertiesGroup.h
+++ b/Engine/Plugins/FX/Niagara/Source/NiagaraEditor/Public/ViewModels/Stack/NiagaraStackEmitterPropertiesGroup.h
@@ -34,9 +34,6 @@ private:
 		virtual FText GetDisplayName() const override {	return DisplayName;	}
 		virtual FText GetDescription() const override {	return Description;	}
 		virtual FText GetKeywords() const override { return Keywords; }
-		virtual TOptional<FNiagaraFavoritesActionData> GetFavoritesData() const override { return TOptional<FNiagaraFavoritesActionData>(); }
-		virtual bool IsInLibrary() const override { return true; }
-		virtual FNiagaraActionSourceData GetSourceData() const override { return FNiagaraActionSourceData(); }
 
 		EStageAddMode AddMode;
 		UClass* SimulationStageClass;
diff --git a/Engine/Plugins/FX/Niagara/Source/NiagaraEditor/Public/ViewModels/Stack/NiagaraStackFunctionInput.h b/Engine/Plugins/FX/Niagara/Source/NiagaraEditor/Public/ViewModels/Stack/NiagaraStackFunctionInput.h
--- a/Engine/Plugins/FX/Niagara/Source/NiagaraEditor/Public/ViewModels/Stack/NiagaraStackFunctionInput.h
+++ b/Engine/Plugins/FX/Niagara/Source/NiagaraEditor/Public/ViewModels/Stack/NiagaraStackFunctionInput.h
@@ -59,25 +59,6 @@ public:
 		None
 	};
 
-	struct NIAGARAEDITOR_API FNiagaraAvailableParameterInfo
-	{
-		FNiagaraAvailableParameterInfo(FNiagaraVariableBase InVariableBase);
-		FNiagaraAvailableParameterInfo(FNiagaraVariableBase InVariableBase, UNiagaraScript* InConversionScript);
-
-		/** The variable this info struct is about. */
-		FNiagaraVariableBase Variable;
-		/** If this is set, the implication is that the variable can only be applied to this stack function input through conversion. */
-		UNiagaraScript* ConversionScript = nullptr;
-		/** If this variable is defined in a parameter definition asset, the data is included here for further processing. */
-		UNiagaraScriptVariable* MatchingParameterDefinition = nullptr;
-
-		/** We use the variable as identifier so we can use array operations such as Contains or Find based on just the variable. */
-		bool operator==(const FNiagaraAvailableParameterInfo& Other) const
-		{
-			return Variable == Other.Variable;
-		}
-	};
-	
 	DECLARE_MULTICAST_DELEGATE(FOnValueChanged);
 
 public:
@@ -173,13 +154,7 @@ public:
 	NIAGARAEDITOR_API void SetLinkedParameterValueViaConversionScript(const FNiagaraVariableBase& InParameter, UNiagaraScript& ConversionScript);
 
 	/** Gets the current set of available parameter handles which can be assigned to this input. Optionally returns possible conversion scripts. */
-	struct NIAGARAEDITOR_API FGetAvailableParameterArgs
-	{
-		bool bIncludeConversionScripts = true;
-		bool bIncludeParameterDefinitions = false;
-	};
-	
-	NIAGARAEDITOR_API void GetAvailableParameters(TSet<FNiagaraAvailableParameterInfo>& OutAvailableParameters, FGetAvailableParameterArgs Arguments) const;
+	NIAGARAEDITOR_API void GetAvailableParameters(TArray<FNiagaraVariableBase>& OutAvailableParameters, TMap<FNiagaraVariableBase, UNiagaraScript*>& OutAvailableParametersWithConversions, bool bIncludeConversionScripts = true) const;
 
 	/** Gets the function node form the script graph if the current value mode is DefaultFunction. */
 	NIAGARAEDITOR_API UNiagaraNodeFunctionCall* GetDefaultFunctionNode() const;
@@ -589,8 +564,3 @@ private:
 	// The value of the change id for the called graph of the function call which owns this input as of the last refresh.
 	TOptional<FGuid> LastFunctionGraphChangeId;
 };
-
-FORCEINLINE uint32 GetTypeHash(const UNiagaraStackFunctionInput::FNiagaraAvailableParameterInfo& ParameterInfo)
-{
-	return GetTypeHash(ParameterInfo.Variable);
-}
diff --git a/Engine/Plugins/FX/Niagara/Source/NiagaraEditor/Public/Widgets/SItemSelector.h b/Engine/Plugins/FX/Niagara/Source/NiagaraEditor/Public/Widgets/SItemSelector.h
--- a/Engine/Plugins/FX/Niagara/Source/NiagaraEditor/Public/Widgets/SItemSelector.h
+++ b/Engine/Plugins/FX/Niagara/Source/NiagaraEditor/Public/Widgets/SItemSelector.h
@@ -75,7 +75,6 @@ public:
 	DECLARE_DELEGATE_RetVal_OneParam(const CategoryKeyType&, FOnGetKeyForCategory, const CategoryType& /* Category */);
 	DECLARE_DELEGATE_RetVal_OneParam(const SectionKeyType&, FOnGetKeyForSection, const SectionType& /* Section */);
 	DECLARE_DELEGATE_OneParam(FOnSuggestionUpdated, int32 /* Suggestion Index*/)
-	DECLARE_DELEGATE_TwoParams(FOnItemRowHoverEvent, const ItemType&, bool /* bHover */)
 
 	DECLARE_DELEGATE_RetVal_OneParam(bool, FOnCategoryPassesFilter, const CategoryType& /* CategoryA */);
 
@@ -225,8 +224,6 @@ public:
 		/** Whether we want to expand the tree initially or not. */
 		SLATE_ARGUMENT(bool, ExpandInitially)
 
-		SLATE_EVENT(FOnItemRowHoverEvent, OnItemRowHoverEvent)
-		
 		/** An optional delegate called when initially creating this and ExpandInitially is false. */
 
 		/** Slot for additional widget content to go adjacent to right of the search box. */
@@ -1362,44 +1359,21 @@ private:
 			SLATE_ARGUMENT(bool, ShowSelection)
 			SLATE_ARGUMENT(FMargin, Padding)
 			SLATE_EVENT(FOnDragDetected, OnDragDetected)
-			SLATE_EVENT(FOnItemRowHoverEvent, OnItemRowHoverEventDelegate)
-			SLATE_ARGUMENT(TSharedPtr<FItemSelectorItemContainerViewModel>, ItemViewModel)
 		SLATE_END_ARGS();
 
 		void Construct(const FArguments& InArgs, const TSharedRef<STableViewBase>& OwnerTree)
 		{
-			ItemViewModelWeak = InArgs._ItemViewModel;
-			
 			typename SItemSelectorTableRow::FArguments Arguments;
 			Arguments = Arguments[InArgs._Content.Widget];
 			Arguments._Style = InArgs._Style;
 			Arguments._ShowSelection = InArgs._ShowSelection;
 			Arguments._Padding = InArgs._Padding;
 			Arguments._OnDragDetected = InArgs._OnDragDetected;
-
-			OnItemRowHoverEventDelegate = InArgs._OnItemRowHoverEventDelegate;
 			SItemSelectorTableRow::Construct(Arguments, OwnerTree);
 		}
-
-		virtual void OnMouseEnter(const FGeometry& InGeometry, const FPointerEvent& PointerEvent) override
-		{
-			SItemSelectorTableRow::OnMouseEnter(InGeometry, PointerEvent);
-
-			OnItemRowHoverEventDelegate.ExecuteIfBound(ItemViewModelWeak.Pin()->GetItem(), true);
-		}
-
-		virtual void OnMouseLeave(const FPointerEvent& PointerEvent) override
-		{
-			SItemSelectorTableRow::OnMouseLeave(PointerEvent);
-
-			OnItemRowHoverEventDelegate.ExecuteIfBound(ItemViewModelWeak.Pin()->GetItem(), false);
-		}
-
-		TWeakPtr<FItemSelectorItemContainerViewModel> ItemViewModelWeak;
-		FOnItemRowHoverEvent OnItemRowHoverEventDelegate;
 	};
 
-	class SItemSelectorItemCategoryTableRow : public SItemSelectorTableRow
+	class SItemSelectorItemCategoryTableRow : public SItemSelectorItemTableRow
 	{
 	public:
 		SLATE_BEGIN_ARGS(SItemSelectorItemCategoryTableRow)
@@ -1414,26 +1388,24 @@ private:
 			SLATE_ARGUMENT(FLinearColor, BorderBackgroundColor)
 			SLATE_ARGUMENT(FMargin, ChildSlotPadding)
 			SLATE_ARGUMENT(FMargin, BorderBackgroundPadding)
-			SLATE_ARGUMENT(TSharedPtr<FItemSelectorItemCategoryViewModel>, CategoryViewModel)
 			SLATE_EVENT(FOnGetCategoryBackgroundImage, OnGetCategoryBackgroundImage)
 		SLATE_END_ARGS();
 
 		void Construct(const FArguments& InArgs, const TSharedRef<STableViewBase>& OwnerTree)
 		{
-			CategoryViewModel = InArgs._CategoryViewModel;
 			OnGetCategoryBackgroundImage = InArgs._OnGetCategoryBackgroundImage;
 
-			typename SItemSelectorTableRow::FArguments Arguments;
+			typename SItemSelectorItemTableRow::FArguments Arguments;
 			Arguments = Arguments[InArgs._Content.Widget];
 			Arguments._Style = InArgs._Style;
 			Arguments._ShowSelection = InArgs._ShowSelection;
 			Arguments._Padding = InArgs._Padding;
-			SItemSelectorTableRow::Construct(Arguments, OwnerTree);
+			SItemSelectorItemTableRow::Construct(Arguments, OwnerTree);
 
 			// If category background images are used, construct a border to display them.
 			if (OnGetCategoryBackgroundImage.IsBound())
 			{
-				SItemSelectorTableRow::ChildSlot
+				SItemSelectorItemTableRow::ChildSlot
 				.Padding(InArgs._ChildSlotPadding)
 				[
 					SAssignNew(ContentBorder, SBorder)
@@ -1472,40 +1444,8 @@ private:
 	private:
 		FOnGetCategoryBackgroundImage OnGetCategoryBackgroundImage;
 		TSharedPtr<SBorder> ContentBorder;
-		TWeakPtr<FItemSelectorItemCategoryViewModel> CategoryViewModel;
 	};
 
-	class SItemSelectorItemSectionTableRow : public SItemSelectorTableRow
-	{
-		SLATE_BEGIN_ARGS(SItemSelectorItemSectionTableRow)
-			: _Style(&FCoreStyle::Get().GetWidgetStyle<FTableRowStyle>("TableView.Row"))
-			, _ShowSelection(false)
-			{}
-			SLATE_DEFAULT_SLOT(typename SItemSelectorItemSectionTableRow::FArguments, Content)
-			SLATE_STYLE_ARGUMENT(FTableRowStyle, Style)
-			SLATE_ARGUMENT(bool, ShowSelection)
-			SLATE_ARGUMENT(FMargin, Padding)
-			SLATE_EVENT(FOnDragDetected, OnDragDetected)
-			SLATE_EVENT(FOnItemRowHoverEvent, OnItemRowHoverEventDelegate)
-			SLATE_ARGUMENT(TSharedPtr<FSectionViewModel>, SectionViewModel)
-		SLATE_END_ARGS()
-
-		void Construct(const FArguments& InArgs, const TSharedRef<STableViewBase>& OwnerTree)
-		{
-			SectionViewModelWeak = InArgs._SectionViewModel;
-			
-			typename SItemSelectorTableRow::FArguments Arguments;
-			Arguments = Arguments[InArgs._Content.Widget];
-			Arguments._Style = InArgs._Style;
-			Arguments._ShowSelection = InArgs._ShowSelection;
-			Arguments._Padding = InArgs._Padding;
-			Arguments._OnDragDetected = InArgs._OnDragDetected;
-			
-			SItemSelectorTableRow::Construct(Arguments, OwnerTree);
-		}
-
-		TWeakPtr<FSectionViewModel> SectionViewModelWeak;
-	};
 public:
 	void Construct(const FArguments& InArgs)
 	{
@@ -1554,7 +1494,6 @@ public:
 		OnGetKeyForCategory = InArgs._OnGetKeyForCategory;
 		OnGetKeyForSection = InArgs._OnGetKeyForSection;
 		SearchBoxAdjacentContentWidget = InArgs._SearchBoxAdjacentContent.Widget;
-		OnItemRowHoverEventDelegate = InArgs._OnItemRowHoverEvent;
 
 		OnSuggestionUpdated = FOnSuggestionUpdated::CreateSP(this, &SItemSelector::OnSuggestionChanged);
 
@@ -2057,7 +1996,6 @@ private:
 				.ChildSlotPadding(CategoryChildSlotPadding)
 				.BorderBackgroundPadding(CategoryBorderBackgroundPadding)
 				.OnGetCategoryBackgroundImage(OnGetCategoryBackgroundImage)
-				.CategoryViewModel(ItemCategoryViewModel)
 				[
 					OnGenerateWidgetForCategory.Execute(ItemCategoryViewModel->GetCategory())
 				];
@@ -2070,8 +2008,6 @@ private:
 				.ShowSelection(true)
 				.Padding(Margin)
 				.OnDragDetected(this, &SItemSelector::OnItemDragDetected)
-				.OnItemRowHoverEventDelegate(OnItemRowHoverEventDelegate)
-				.ItemViewModel(ItemViewModel)
 				[
 					OnGenerateWidgetForItem.Execute(ItemViewModel->GetItem())
 				];
@@ -2079,11 +2015,10 @@ private:
 		case EItemSelectorItemViewModelType::Section:
 		{
 			TSharedRef<FSectionViewModel> SectionViewModel = StaticCastSharedRef<FSectionViewModel>(Item);
-			return SNew(SItemSelectorItemSectionTableRow, OwnerTable)
+			return SNew(SItemSelectorItemTableRow, OwnerTable)
 				.Style(SectionRowStyle)
 				.ShowSelection(false)
 				.Padding(Margin)
-				.SectionViewModel(SectionViewModel)
 				[
 					OnGenerateWidgetForSection.Execute(SectionViewModel->GetSection())
 				];
@@ -2335,8 +2270,7 @@ private:
 	FOnGetKeyForItem OnGetKeyForItem;
 	FOnGetKeyForCategory OnGetKeyForCategory;
 	FOnGetKeyForSection OnGetKeyForSection;
-	FOnItemRowHoverEvent OnItemRowHoverEventDelegate;
-	
+
 	TSharedPtr<FItemSelectorViewModel> ViewModelUtilities;
 	TSharedPtr<STreeView<TSharedRef<FItemSelectorItemViewModel>>> ItemTree;
 	FOnSuggestionUpdated OnSuggestionUpdated;
diff --git a/Engine/Plugins/FX/Niagara/Source/NiagaraEditor/Public/Widgets/SNiagaraGraphActionMenu.h b/Engine/Plugins/FX/Niagara/Source/NiagaraEditor/Public/Widgets/SNiagaraGraphActionMenu.h
--- a/Engine/Plugins/FX/Niagara/Source/NiagaraEditor/Public/Widgets/SNiagaraGraphActionMenu.h
+++ b/Engine/Plugins/FX/Niagara/Source/NiagaraEditor/Public/Widgets/SNiagaraGraphActionMenu.h
@@ -35,10 +35,6 @@ struct FNiagaraGraphActionMenuFilter
 class SNiagaraGraphActionMenu : public SBorder
 {
 public:
-	/** The name of the profile used to save favorite action data. Used with UNiagaraFavoriteActionsConfig.
-	 * Each action gets an identifier, and the user-authored favorites data is saved per profile. */
-	static const FName FavoriteActionsProfile;
-	
 	SLATE_BEGIN_ARGS( SNiagaraGraphActionMenu )
 		: _GraphObj( static_cast<UEdGraph*>(NULL) )
 		, _NewNodePosition( FVector2D::ZeroVector )
diff --git a/Engine/Plugins/FX/Niagara/Source/NiagaraEditor/Public/Widgets/Wizard/SNiagaraModuleWizard.h b/Engine/Plugins/FX/Niagara/Source/NiagaraEditor/Public/Widgets/Wizard/SNiagaraModuleWizard.h
--- a/Engine/Plugins/FX/Niagara/Source/NiagaraEditor/Public/Widgets/Wizard/SNiagaraModuleWizard.h
+++ b/Engine/Plugins/FX/Niagara/Source/NiagaraEditor/Public/Widgets/Wizard/SNiagaraModuleWizard.h
@@ -57,9 +57,6 @@ namespace UE::Niagara::Wizard
 		// called after the generated module is added to the stack - returning true will apply the modified clipboard content to the module inputs
 		virtual bool UpdateModuleInputs(UNiagaraClipboardContent* NewModule, const TArray<const UNiagaraNodeFunctionCall*>& PreviouslyCreatedModules);
 
-		// a unique name to identify this wizard with. Used in menus.
-		virtual FName GetIdentifier() const { return "ModuleWizardModel"; }
-		
 		// target usage of the generated module
 		ENiagaraScriptUsage TargetUsage;
 		
diff --git a/Engine/Plugins/FX/Niagara/Source/NiagaraEditorWidgets/Private/Stack/SNiagaraStackFunctionInputValue.cpp b/Engine/Plugins/FX/Niagara/Source/NiagaraEditorWidgets/Private/Stack/SNiagaraStackFunctionInputValue.cpp
--- a/Engine/Plugins/FX/Niagara/Source/NiagaraEditorWidgets/Private/Stack/SNiagaraStackFunctionInputValue.cpp
+++ b/Engine/Plugins/FX/Niagara/Source/NiagaraEditorWidgets/Private/Stack/SNiagaraStackFunctionInputValue.cpp
@@ -20,12 +20,10 @@
 #include "NiagaraNodeAssignment.h"
 #include "NiagaraNodeFunctionCall.h"
 #include "NiagaraParameterCollection.h"
-#include "NiagaraParameterDefinitions.h"
 #include "NiagaraStackCommandContext.h"
 #include "NiagaraScriptVariable.h"
 #include "NiagaraSettings.h"
 #include "NiagaraSystem.h"
-#include "NiagaraTrace.h"
 #include "PropertyCustomizationHelpers.h"
 #include "PropertyEditorModule.h"
 #include "SDropTarget.h"
@@ -33,7 +31,6 @@
 #include "SNiagaraParameterDropTarget.h"
 #include "SNiagaraParameterEditor.h"
 #include "ScopedTransaction.h"
-#include "Config/NiagaraFavoriteActionsConfig.h"
 #include "Stack/SNiagaraStackIndent.h"
 #include "Styling/AppStyle.h"
 #include "Subsystems/AssetEditorSubsystem.h"
@@ -65,7 +62,6 @@
 const float TextIconSize = 16;
 
 bool SNiagaraStackFunctionInputValue::bLibraryOnly = true;
-FName SNiagaraStackFunctionInputValue::FavoriteActionsProfile = TEXT("StackFunctionInputActionsProfile");
 
 void SNiagaraStackFunctionInputValue::Construct(const FArguments& InArgs, UNiagaraStackFunctionInput* InFunctionInput)
 {
@@ -853,16 +849,12 @@ TSharedRef<SWidget> SNiagaraStackFunctionInputValue::OnGetAvailableHandleMenu()
                 .OnGenerateWidgetForSection(this, &SNiagaraStackFunctionInputValue::OnGenerateWidgetForSection)
                 .OnGenerateWidgetForCategory(this, &SNiagaraStackFunctionInputValue::OnGenerateWidgetForCategory)
                 .OnGenerateWidgetForItem(this, &SNiagaraStackFunctionInputValue::OnGenerateWidgetForItem)
-				.OnGetItemWeight_Lambda([](const TSharedPtr<FNiagaraMenuAction_Generic>& Item, const TArray<FString>& FilterTerms)
-				{
-					return FNiagaraEditorUtilities::GetWeightForItem(Item, SNiagaraStackFunctionInputValue::FavoriteActionsProfile, FilterTerms);
-				})
+                .OnGetItemWeight_Static(&FNiagaraEditorUtilities::GetWeightForItem)
                 .OnItemActivated(this, &SNiagaraStackFunctionInputValue::OnItemActivated)
                 .AllowMultiselect(false)
                 .OnDoesItemPassCustomFilter(this, &SNiagaraStackFunctionInputValue::DoesItemPassCustomFilter)
                 .ClickActivateMode(EItemSelectorClickActivateMode::SingleClick)
                 .ExpandInitially(false)
-				.OnItemRowHoverEvent(this, &SNiagaraStackFunctionInputValue::OnActionRowHoverEvent)
                 .OnGetSectionData_Lambda([](const ENiagaraMenuSections& Section)
                 {
                     if(Section == ENiagaraMenuSections::Suggested)
@@ -1127,14 +1119,10 @@ TArray<TSharedPtr<FNiagaraMenuAction_Generic>> SNiagaraStackFunctionInputValue::
 			Source = TTuple<EScriptSource, FText>(EScriptSource::Niagara, FText::FromString("Scratch Pad"));
 			bIsInLibrary = true;
 		}
-
-		FNiagaraFavoritesActionData FavoritesActionData;
-		FavoritesActionData.ActionIdentifier.Names.Add(FName(GetPathNameSafe(DynamicInputScript)));
-		FavoritesActionData.bFavoriteByDefault = ScriptData->bSuggested;
-			
+		
 		TSharedPtr<FNiagaraMenuAction_Generic> DynamicInputAction(new FNiagaraMenuAction_Generic(
 			FNiagaraMenuAction_Generic::FOnExecuteAction::CreateStatic(&ReassignDynamicInputScript, FunctionInput, DynamicInputScript),
-			DisplayName, {CategoryName.ToString()}, FavoritesActionData, Tooltip, ScriptData->Keywords
+			DisplayName, ScriptData->bSuggested ? ENiagaraMenuSections::Suggested : ENiagaraMenuSections::General, {CategoryName.ToString()}, Tooltip, ScriptData->Keywords
             ));
 		DynamicInputAction->SourceData = FNiagaraActionSourceData(Source.Key, Source.Value, true);
 		DynamicInputAction->bIsInLibrary = bIsInLibrary;
@@ -1185,16 +1173,12 @@ void SNiagaraStackFunctionInputValue::ShowReassignDynamicInputScriptMenu()
                 .OnGenerateWidgetForSection(this, &SNiagaraStackFunctionInputValue::OnGenerateWidgetForSection)
                 .OnGenerateWidgetForCategory(this, &SNiagaraStackFunctionInputValue::OnGenerateWidgetForCategory)
                 .OnGenerateWidgetForItem(this, &SNiagaraStackFunctionInputValue::OnGenerateWidgetForItem)
-				.OnGetItemWeight_Lambda([](const TSharedPtr<FNiagaraMenuAction_Generic>& Item, const TArray<FString>& FilterTerms)
-				 {
-					 return FNiagaraEditorUtilities::GetWeightForItem(Item, SNiagaraStackFunctionInputValue::FavoriteActionsProfile, FilterTerms);
-				 })
+                .OnGetItemWeight_Static(&FNiagaraEditorUtilities::GetWeightForItem)
                 .OnItemActivated(this, &SNiagaraStackFunctionInputValue::OnItemActivated)
                 .AllowMultiselect(false)
                 .OnDoesItemPassCustomFilter(this, &SNiagaraStackFunctionInputValue::DoesItemPassCustomFilter)
                 .ClickActivateMode(EItemSelectorClickActivateMode::SingleClick)
                 .ExpandInitially(false)
-				.OnItemRowHoverEvent(this, &SNiagaraStackFunctionInputValue::OnActionRowHoverEvent)
                 .OnGetSectionData_Lambda([](const ENiagaraMenuSections& Section)
                 {
                     if(Section == ENiagaraMenuSections::Suggested)
@@ -1237,20 +1221,12 @@ FReply SNiagaraStackFunctionInputValue::ScratchButtonPressed() const
 	return FReply::Unhandled();
 }
 
-void SNiagaraStackFunctionInputValue::OnActionRowHoverEvent(const TSharedPtr<FNiagaraMenuAction_Generic>& ActionNode, bool bIsHovered)
-{
-	ActionNode->bIsHovered = bIsHovered;
-}
-
 TArray<TSharedPtr<FNiagaraMenuAction_Generic>> SNiagaraStackFunctionInputValue::CollectActions()
 {
-	TRACE_CPUPROFILER_EVENT_SCOPE_ON_CHANNEL(SNiagaraStackFunctionInputValue::CollectActions, NiagaraChannel);
-	
 	TArray<TSharedPtr<FNiagaraMenuAction_Generic>> OutAllActions;
 	bool bIsDataInterfaceOrObject = FunctionInput->GetInputType().IsDataInterface() || FunctionInput->GetInputType().IsUObject();
 
 	FNiagaraActionSourceData NiagaraSourceData(EScriptSource::Niagara, FText::FromString(TEXT("Niagara")), true);
-	FNiagaraFavoriteActionsProfile& ActionsProfile = UNiagaraFavoriteActionsConfig::Get()->GetActionsProfile(FavoriteActionsProfile);
 	
 	// Set a local value
 	{
@@ -1263,7 +1239,7 @@ TArray<TSharedPtr<FNiagaraMenuAction_Generic>> SNiagaraStackFunctionInputValue::
 		TSharedPtr<FNiagaraMenuAction_Generic> SetLocalValueAction(new FNiagaraMenuAction_Generic(
 			FNiagaraMenuAction_Generic::FOnExecuteAction::CreateSP(this, &SNiagaraStackFunctionInputValue::SetToLocalValue),
 			FNiagaraMenuAction_Generic::FCanExecuteAction::CreateLambda([=]() { return bCanSetLocalValue; }),
-            DisplayName, {}, {}, Tooltip, FText()));
+            DisplayName, ENiagaraMenuSections::General, {}, Tooltip, FText()));
 		SetLocalValueAction->SourceData = NiagaraSourceData;
 		OutAllActions.Add(SetLocalValueAction);
 	}
@@ -1297,16 +1273,13 @@ TArray<TSharedPtr<FNiagaraMenuAction_Generic>> SNiagaraStackFunctionInputValue::
 				bIsInLibrary = true;
 			}
 			
-			// We construct an ActionIdentifier to check for it in the favorites list
-			FNiagaraFavoritesActionData FavoritesActionData;
-			FavoritesActionData.ActionIdentifier = FNiagaraActionIdentifier({FName(GetPathNameSafe(DynamicInputScript))}, {}); 
-			FavoritesActionData.bFavoriteByDefault = ScriptData->bSuggested;
-			
 			TSharedPtr<FNiagaraMenuAction_Generic> DynamicInputAction(new FNiagaraMenuAction_Generic(
                 FNiagaraMenuAction_Generic::FOnExecuteAction::CreateSP(this, &SNiagaraStackFunctionInputValue::DynamicInputScriptSelected, DynamicInputScript),
-                DisplayName, {CategoryName.ToString()}, FavoritesActionData, Tooltip, ScriptData->Keywords));
-			DynamicInputAction->FavoritesActionData = FavoritesActionData;
-			DynamicInputAction->SourceData = FNiagaraActionSourceData(Source.Key, Source.Value, true);			
+                DisplayName, ScriptData->bSuggested ? ENiagaraMenuSections::Suggested : ENiagaraMenuSections::General, {CategoryName.ToString()}, Tooltip, ScriptData->Keywords));
+			
+			DynamicInputAction->SourceData = FNiagaraActionSourceData(Source.Key, Source.Value, true);
+
+			
 			
 			DynamicInputAction->bIsExperimental = ScriptData->bExperimental;
 			DynamicInputAction->bIsInLibrary = bIsInLibrary;
@@ -1318,25 +1291,15 @@ TArray<TSharedPtr<FNiagaraMenuAction_Generic>> SNiagaraStackFunctionInputValue::
 	const UNiagaraSettings* Settings = GetDefault<UNiagaraSettings>();
 	bool bAllowConversions = ensure(Settings != nullptr) && Settings->bShowConvertibleInputsInStack;
 
-	TSet<UNiagaraStackFunctionInput::FNiagaraAvailableParameterInfo> AvailableParameterInfos;
-
-	UNiagaraStackFunctionInput::FGetAvailableParameterArgs GetAvailableParameterArgs;
-	GetAvailableParameterArgs.bIncludeConversionScripts = bAllowConversions;
-	GetAvailableParameterArgs.bIncludeParameterDefinitions = true;
-	FunctionInput->GetAvailableParameters(AvailableParameterInfos, GetAvailableParameterArgs);
+	TArray<FNiagaraVariableBase> AvailableParameters;
+	TMap<FNiagaraVariableBase, UNiagaraScript*> AvailableParametersWithConversion;
+	FunctionInput->GetAvailableParameters(AvailableParameters, AvailableParametersWithConversion, bAllowConversions);
 
-	// First, we add the inputs that can be directly linked (without a conversion script)
 	const FString RootCategoryName = FString("Link Inputs");
 	const FText MapInputFormat = LOCTEXT("LinkInputFormat", "Link this input to {0}");
-	for (const UNiagaraStackFunctionInput::FNiagaraAvailableParameterInfo& AvailableParameterInfo : AvailableParameterInfos)
+	for (const FNiagaraVariableBase& AvailableParameter : AvailableParameters)
 	{
-		if(AvailableParameterInfo.ConversionScript != nullptr)
-		{
-			continue;
-		}
-		
-		FNiagaraParameterHandle AvailableHandle = FNiagaraParameterHandle(AvailableParameterInfo.Variable.GetName());
-		
+		FNiagaraParameterHandle AvailableHandle = FNiagaraParameterHandle(AvailableParameter.GetName());
 		TArray<FName> HandleParts = AvailableHandle.GetHandleParts();
 		FNiagaraNamespaceMetadata NamespaceMetadata = GetDefault<UNiagaraEditorSettings>()->GetMetaDataForNamespaces(HandleParts);
 		if (NamespaceMetadata.IsValid())
@@ -1346,14 +1309,10 @@ TArray<TSharedPtr<FNiagaraMenuAction_Generic>> SNiagaraStackFunctionInputValue::
 			const FText Category = NamespaceMetadata.DisplayName;
 			const FText DisplayName = FText::FromName(AvailableHandle.GetParameterHandleString());
 			const FText Tooltip = FText::Format(MapInputFormat, FText::FromName(AvailableHandle.GetParameterHandleString()));
-
-			FNiagaraFavoritesActionData FavoritesActionData;
-			FavoritesActionData.ActionIdentifier.Names.Add(FName(DisplayName.ToString()));
-			FavoritesActionData.bFavoriteByDefault = AvailableParameterInfo.MatchingParameterDefinition ? AvailableParameterInfo.MatchingParameterDefinition->Metadata.bSuggested : false;
 			
 			TSharedPtr<FNiagaraMenuAction_Generic> LinkAction(new FNiagaraMenuAction_Generic(
-				FNiagaraMenuAction_Generic::FOnExecuteAction::CreateSP(this, &SNiagaraStackFunctionInputValue::ParameterSelected, AvailableParameterInfo.Variable),
-				DisplayName, {RootCategoryName, Category.ToString()}, FavoritesActionData,  Tooltip, FText()));
+				FNiagaraMenuAction_Generic::FOnExecuteAction::CreateSP(this, &SNiagaraStackFunctionInputValue::ParameterSelected, AvailableParameter),
+				DisplayName, ENiagaraMenuSections::General, {RootCategoryName, Category.ToString()}, Tooltip, FText()));
 			
 			LinkAction->SetParameterVariable(FNiagaraVariable(FunctionInput->GetInputType(), AvailableHandle.GetParameterHandleString()));
 			LinkAction->SourceData = NiagaraSourceData;
@@ -1362,17 +1321,11 @@ TArray<TSharedPtr<FNiagaraMenuAction_Generic>> SNiagaraStackFunctionInputValue::
 		}
 	}
 
-	// Then we add those that can only get added via conversion script. This will only be valid if the plugin setting allows it to
 	const FText ConvertInputFormat = LOCTEXT("ConvertInputFormat", "Link this input to {0} via a conversion script");
-	for (const UNiagaraStackFunctionInput::FNiagaraAvailableParameterInfo& AvailableParameterInfo : AvailableParameterInfos)
+	for (const auto& Entry : AvailableParametersWithConversion)
 	{
-		if(AvailableParameterInfo.ConversionScript == nullptr)
-		{
-			continue;
-		}
-		
-		FNiagaraVariableBase ParameterVariable = AvailableParameterInfo.Variable;
-		UNiagaraScript* ConversionScript = AvailableParameterInfo.ConversionScript;
+		FNiagaraVariableBase ParameterVariable = Entry.Key;
+		UNiagaraScript* ConversionScript = Entry.Value;
 		const FNiagaraParameterHandle& AvailableHandle(ParameterVariable.GetName());
 		TArray<FName> HandleParts = AvailableHandle.GetHandleParts();
 		FNiagaraNamespaceMetadata NamespaceMetadata = GetDefault<UNiagaraEditorSettings>()->GetMetaDataForNamespaces(HandleParts);
@@ -1386,7 +1339,7 @@ TArray<TSharedPtr<FNiagaraMenuAction_Generic>> SNiagaraStackFunctionInputValue::
 
 			TSharedPtr<FNiagaraMenuAction_Generic> LinkAction(new FNiagaraMenuAction_Generic(
 				FNiagaraMenuAction_Generic::FOnExecuteAction::CreateSP(this, &SNiagaraStackFunctionInputValue::ParameterWithConversionSelected, ParameterVariable, ConversionScript),
-				DisplayName, {RootCategoryName, Category.ToString()}, {}, Tooltip, FText()));
+				DisplayName, ENiagaraMenuSections::General, {RootCategoryName, Category.ToString()}, Tooltip, FText()));
 
 			LinkAction->SetParameterVariable(ParameterVariable);
 
@@ -1416,14 +1369,14 @@ TArray<TSharedPtr<FNiagaraMenuAction_Generic>> SNiagaraStackFunctionInputValue::
 		{
 			FNiagaraParameterHandle HandleToRead(AvailableNamespace, InputName);
 			FNiagaraVariableBase ParameterToRead = FNiagaraVariableBase(FunctionInput->GetInputType(), HandleToRead.GetParameterHandleString());
-			bool bIsContained = AvailableParameterInfos.Contains(ParameterToRead);
+			bool bIsContained = AvailableParameters.Contains(ParameterToRead);
 
 			if(bIsContained)
 			{
 				TSet<FName> ExistingNames;
-				for(const UNiagaraStackFunctionInput::FNiagaraAvailableParameterInfo& AvailableParameter : AvailableParameterInfos)
+				for(const FNiagaraVariableBase& AvailableParameter : AvailableParameters)
 				{
-					FNiagaraParameterHandle AvailableHandle = FNiagaraParameterHandle(AvailableParameter.Variable.GetName());
+					FNiagaraParameterHandle AvailableHandle = FNiagaraParameterHandle(AvailableParameter.GetName());
 					ExistingNames.Add(AvailableHandle.GetName());
 				}
 
@@ -1440,7 +1393,7 @@ TArray<TSharedPtr<FNiagaraMenuAction_Generic>> SNiagaraStackFunctionInputValue::
 
 			TSharedPtr<FNiagaraMenuAction_Generic> MakeAction(new FNiagaraMenuAction_Generic(
 				FNiagaraMenuAction_Generic::FOnExecuteAction::CreateSP(this, &SNiagaraStackFunctionInputValue::ParameterSelected, ParameterToRead),         
-		        DisplayName, {CategoryName.ToString()}, {}, Tooltip, FText()));
+		        DisplayName, ENiagaraMenuSections::General, {CategoryName.ToString()}, Tooltip, FText()));
 
 			MakeAction->SourceData = NiagaraSourceData;
 
@@ -1458,7 +1411,7 @@ TArray<TSharedPtr<FNiagaraMenuAction_Generic>> SNiagaraStackFunctionInputValue::
 
 		TSharedPtr<FNiagaraMenuAction_Generic> ExpressionAction(new FNiagaraMenuAction_Generic(
                 FNiagaraMenuAction_Generic::FOnExecuteAction::CreateSP(this, &SNiagaraStackFunctionInputValue::CustomExpressionSelected),
-                DisplayName, {}, {}, Tooltip, FText()));
+                DisplayName, ENiagaraMenuSections::General, {}, Tooltip, FText()));
 
 		ExpressionAction->SourceData = NiagaraSourceData;
 
@@ -1475,7 +1428,7 @@ TArray<TSharedPtr<FNiagaraMenuAction_Generic>> SNiagaraStackFunctionInputValue::
 
 		TSharedPtr<FNiagaraMenuAction_Generic> CreateScratchAction(new FNiagaraMenuAction_Generic(
            FNiagaraMenuAction_Generic::FOnExecuteAction::CreateSP(this, &SNiagaraStackFunctionInputValue::CreateScratchSelected),
-           DisplayName, {}, {}, Tooltip, FText()));
+           DisplayName, ENiagaraMenuSections::General, {}, Tooltip, FText()));
 
 		CreateScratchAction->SourceData = NiagaraSourceData;
 
@@ -1490,7 +1443,7 @@ TArray<TSharedPtr<FNiagaraMenuAction_Generic>> SNiagaraStackFunctionInputValue::
 		TSharedPtr<FNiagaraMenuAction_Generic> DeleteInputAction(new FNiagaraMenuAction_Generic(
                 FNiagaraMenuAction_Generic::FOnExecuteAction::CreateUObject(FunctionInput, &UNiagaraStackFunctionInput::DeleteInput),
                 FNiagaraMenuAction_Generic::FCanExecuteAction::CreateUObject(FunctionInput, &UNiagaraStackFunctionInput::CanDeleteInput),
-                DisplayName, {}, {}, Tooltip, FText()));
+                DisplayName, ENiagaraMenuSections::General, {}, Tooltip, FText()));
 
 		DeleteInputAction->SourceData = NiagaraSourceData;
 		OutAllActions.Add(DeleteInputAction);
@@ -1508,17 +1461,12 @@ TArray<FString> SNiagaraStackFunctionInputValue::OnGetCategoriesForItem(
 TArray<ENiagaraMenuSections> SNiagaraStackFunctionInputValue::OnGetSectionsForItem(
 	const TSharedPtr<FNiagaraMenuAction_Generic>& Item)
 {
-	TArray<ENiagaraMenuSections> Sections{ENiagaraMenuSections::General};
-	
-	if(Item->FavoritesActionData.IsSet())
+	if(Item->Section == ENiagaraMenuSections::Suggested)
 	{
-		FNiagaraFavoriteActionsProfile& ActionsProfile = UNiagaraFavoriteActionsConfig::Get()->GetActionsProfile(SNiagaraStackFunctionInputValue::FavoriteActionsProfile);
-		if(ActionsProfile.IsFavorite(Item->FavoritesActionData.GetValue()))
-		{
-			Sections.Add(ENiagaraMenuSections::Suggested);
-		}
+		return { ENiagaraMenuSections::General, ENiagaraMenuSections::Suggested };
 	}
-	return Sections;
+		
+	return {Item->Section};
 }
 
 bool SNiagaraStackFunctionInputValue::OnCompareSectionsForEquality(const ENiagaraMenuSections& SectionA,
@@ -1580,9 +1528,7 @@ TSharedRef<SWidget> SNiagaraStackFunctionInputValue::OnGenerateWidgetForItem(
 {
 	FCreateNiagaraWidgetForActionData ActionData(Item);
 	ActionData.HighlightText = TAttribute<FText>::Create(TAttribute<FText>::FGetter::CreateRaw(this, &SNiagaraStackFunctionInputValue::GetFilterText));
-	ActionData.FavoriteActionsProfileName = FavoriteActionsProfile;
-	return SNew(SNiagaraActionWidget, ActionData)
-		.bShowTypeIfParameter(false);
+	return SNew(SNiagaraActionWidget, ActionData).bShowTypeIfParameter(false);
 }
 
 bool SNiagaraStackFunctionInputValue::DoesItemPassCustomFilter(const TSharedPtr<FNiagaraMenuAction_Generic>& Item)
diff --git a/Engine/Plugins/FX/Niagara/Source/NiagaraEditorWidgets/Private/Stack/SNiagaraStackFunctionInputValue.h b/Engine/Plugins/FX/Niagara/Source/NiagaraEditorWidgets/Private/Stack/SNiagaraStackFunctionInputValue.h
--- a/Engine/Plugins/FX/Niagara/Source/NiagaraEditorWidgets/Private/Stack/SNiagaraStackFunctionInputValue.h
+++ b/Engine/Plugins/FX/Niagara/Source/NiagaraEditorWidgets/Private/Stack/SNiagaraStackFunctionInputValue.h
@@ -199,10 +199,8 @@ private:
 
 	TAttribute<EVisibility> CompactActionMenuButtonVisibilityAttribute;
 	TSharedPtr<FNiagaraStackCommandContext> StackCommandContext;
-	
-	static bool bLibraryOnly;
 
-	static FName FavoriteActionsProfile;
+	static bool bLibraryOnly;
 
 private:
 	TArray<TSharedPtr<FNiagaraMenuAction_Generic>> CollectActions();
@@ -219,7 +217,6 @@ private:
 	TSharedRef<SWidget> OnGenerateWidgetForItem(const TSharedPtr<FNiagaraMenuAction_Generic>& Item);
 	bool DoesItemPassCustomFilter(const TSharedPtr<FNiagaraMenuAction_Generic>& Item);
 	void OnItemActivated(const TSharedPtr<FNiagaraMenuAction_Generic>& Item);
-	void OnActionRowHoverEvent(const TSharedPtr<FNiagaraMenuAction_Generic>& ActionNode, bool bIsHovered);
 
 	void TriggerRefresh(const TMap<EScriptSource, bool>& SourceState);
 
diff --git a/Engine/Plugins/FX/Niagara/Source/NiagaraEditorWidgets/Private/Stack/SNiagaraStackItemGroupAddMenu.cpp b/Engine/Plugins/FX/Niagara/Source/NiagaraEditorWidgets/Private/Stack/SNiagaraStackItemGroupAddMenu.cpp
--- a/Engine/Plugins/FX/Niagara/Source/NiagaraEditorWidgets/Private/Stack/SNiagaraStackItemGroupAddMenu.cpp
+++ b/Engine/Plugins/FX/Niagara/Source/NiagaraEditorWidgets/Private/Stack/SNiagaraStackItemGroupAddMenu.cpp
@@ -7,7 +7,6 @@
 #include "ViewModels/Stack/INiagaraStackItemGroupAddUtilities.h"
 #include "SGraphActionMenu.h"
 #include "SNiagaraGraphActionWidget.h"
-#include "Config/NiagaraFavoriteActionsConfig.h"
 #include "Framework/Application/SlateApplication.h"
 #include "Widgets/SItemSelector.h"
 #include "Widgets/Text/SRichTextBlock.h"
@@ -71,16 +70,12 @@ void SNiagaraStackItemGroupAddMenu::Construct(const FArguments& InArgs, UNiagara
 	                .OnGenerateWidgetForSection(this, &SNiagaraStackItemGroupAddMenu::OnGenerateWidgetForSection)
 	                .OnGenerateWidgetForCategory(this, &SNiagaraStackItemGroupAddMenu::OnGenerateWidgetForCategory)
 	                .OnGenerateWidgetForItem(this, &SNiagaraStackItemGroupAddMenu::OnGenerateWidgetForItem)
-	                .OnGetItemWeight_Lambda([](const TSharedPtr<FNiagaraMenuAction_Generic>& Item, const TArray<FString>& FilterTerms)
-	                {
-	                	return FNiagaraEditorUtilities::GetWeightForItem(Item, INiagaraStackItemGroupAddAction::FavoriteActionsProfile, FilterTerms);
-	                })
+	                .OnGetItemWeight_Static(&FNiagaraEditorUtilities::GetWeightForItem)
 	                .OnItemActivated(this, &SNiagaraStackItemGroupAddMenu::OnItemActivated)
 	                .AllowMultiselect(false)
 	                .OnDoesItemPassCustomFilter(this, &SNiagaraStackItemGroupAddMenu::DoesItemPassCustomFilter)
 	                .ClickActivateMode(EItemSelectorClickActivateMode::SingleClick)
 	                .ExpandInitially(AddUtilities->GetAutoExpandAddActions())
-					.OnItemRowHoverEvent(this, &SNiagaraStackItemGroupAddMenu::OnItemRowHoverEvent)
 	                .OnGetSectionData_Lambda([](const ENiagaraMenuSections& Section)
 	                {
 	                    if(Section == ENiagaraMenuSections::Suggested)
@@ -109,21 +104,17 @@ TArray<TSharedPtr<FNiagaraMenuAction_Generic>> SNiagaraStackItemGroupAddMenu::Co
 	AddOptions.bIncludeDeprecated = false;
 	AddOptions.bIncludeNonLibrary = true;
 
-	// We translate group add actions into the generic menu actions here
-	
 	TArray<TSharedRef<INiagaraStackItemGroupAddAction>> AddActions;
 	AddUtilities->GenerateAddActions(AddActions, AddOptions);
 
-	// For each group action, we parse its data and match data formats to the generic action
 	for (TSharedRef<INiagaraStackItemGroupAddAction> AddAction : AddActions)
-	{		
+	{
 		TSharedPtr<FNiagaraMenuAction_Generic> NewNodeAction(
             new FNiagaraMenuAction_Generic(
             FNiagaraMenuAction_Generic::FOnExecuteAction::CreateRaw(AddUtilities, &INiagaraStackItemGroupAddUtilities::ExecuteAddAction, AddAction, InsertIndex),
-            AddAction->GetDisplayName(), AddAction->GetCategories(), AddAction->GetFavoritesData(), AddAction->GetDescription(), AddAction->GetKeywords()));
+            AddAction->GetDisplayName(), AddAction->GetSuggested() ? ENiagaraMenuSections::Suggested : ENiagaraMenuSections::General, AddAction->GetCategories(), AddAction->GetDescription(), AddAction->GetKeywords()));
 		NewNodeAction->SourceData = AddAction->GetSourceData();
 		NewNodeAction->bIsInLibrary = AddAction->IsInLibrary();
-		
 		OutActions.Add(NewNodeAction);
 	}
 
@@ -146,11 +137,6 @@ void SNiagaraStackItemGroupAddMenu::OnItemActivated(const TSharedPtr<FNiagaraMen
 	}	
 }
 
-void SNiagaraStackItemGroupAddMenu::OnItemRowHoverEvent(const TSharedPtr<FNiagaraMenuAction_Generic>& Item, bool bHover) const
-{
-	Item->bIsHovered = bHover;
-}
-
 void SNiagaraStackItemGroupAddMenu::TriggerRefresh(const TMap<EScriptSource, bool>& SourceState)
 {
 	ActionSelector->RefreshAllCurrentItems();
@@ -190,18 +176,13 @@ TArray<FString> SNiagaraStackItemGroupAddMenu::OnGetCategoriesForItem(
 TArray<ENiagaraMenuSections> SNiagaraStackItemGroupAddMenu::OnGetSectionsForItem(
 	const TSharedPtr<FNiagaraMenuAction_Generic>& Item)
 {
-	TArray<ENiagaraMenuSections> Sections{ENiagaraMenuSections::General};
 	
-	if(Item->FavoritesActionData.IsSet())
+	if(Item->Section == ENiagaraMenuSections::Suggested)
 	{
-		FNiagaraFavoriteActionsProfile& ActionsProfile = UNiagaraFavoriteActionsConfig::Get()->GetActionsProfile(INiagaraStackItemGroupAddAction::FavoriteActionsProfile);
-		if(ActionsProfile.IsFavorite(Item->FavoritesActionData.GetValue()))
-		{
-			Sections.Add(ENiagaraMenuSections::Suggested);
-		}
+		return { ENiagaraMenuSections::General, ENiagaraMenuSections::Suggested };
 	}
 		
-	return Sections;
+	return {Item->Section};
 }
 
 bool SNiagaraStackItemGroupAddMenu::OnCompareSectionsForEquality(const ENiagaraMenuSections& SectionA,
@@ -262,7 +243,6 @@ TSharedRef<SWidget> SNiagaraStackItemGroupAddMenu::OnGenerateWidgetForItem(
 	const TSharedPtr<FNiagaraMenuAction_Generic>& Item)
 {
 	FCreateNiagaraWidgetForActionData ActionData(Item);
-	ActionData.FavoriteActionsProfileName = INiagaraStackItemGroupAddAction::FavoriteActionsProfile;
 	ActionData.HighlightText = TAttribute<FText>::Create(TAttribute<FText>::FGetter::CreateRaw(this, &SNiagaraStackItemGroupAddMenu::GetFilterText));
 	return SNew(SNiagaraActionWidget, ActionData);
 }
diff --git a/Engine/Plugins/FX/Niagara/Source/NiagaraEditorWidgets/Private/Stack/SNiagaraStackItemGroupAddMenu.h b/Engine/Plugins/FX/Niagara/Source/NiagaraEditorWidgets/Private/Stack/SNiagaraStackItemGroupAddMenu.h
--- a/Engine/Plugins/FX/Niagara/Source/NiagaraEditorWidgets/Private/Stack/SNiagaraStackItemGroupAddMenu.h
+++ b/Engine/Plugins/FX/Niagara/Source/NiagaraEditorWidgets/Private/Stack/SNiagaraStackItemGroupAddMenu.h
@@ -57,7 +57,6 @@ private:
 	TSharedRef<SWidget> OnGenerateWidgetForItem(const TSharedPtr<FNiagaraMenuAction_Generic>& Item);
 	bool DoesItemPassCustomFilter(const TSharedPtr<FNiagaraMenuAction_Generic>& Item);
 	void OnItemActivated(const TSharedPtr<FNiagaraMenuAction_Generic>& Item);
-	void OnItemRowHoverEvent(const TSharedPtr<FNiagaraMenuAction_Generic>& Item, bool bHover) const;
 
 	void TriggerRefresh(const TMap<EScriptSource, bool>& SourceState);
 
diff --git a/Engine/Source/Runtime/SlateCore/Private/Styling/StarshipCoreStyle.cpp b/Engine/Source/Runtime/SlateCore/Private/Styling/StarshipCoreStyle.cpp
--- a/Engine/Source/Runtime/SlateCore/Private/Styling/StarshipCoreStyle.cpp
+++ b/Engine/Source/Runtime/SlateCore/Private/Styling/StarshipCoreStyle.cpp
@@ -352,8 +352,6 @@ TSharedRef<ISlateStyle> FStarshipCoreStyle::Create()
 		Style->Set("Icons.Filter", new IMAGE_BRUSH_SVG("Starship/Common/filter", Icon16x16));
 		Style->Set("Icons.AutoFilter", new IMAGE_BRUSH_SVG("Starship/Common/FilterAuto", Icon16x16));
 
-		Style->Set("Icons.Tag", new IMAGE_BRUSH_SVG("Starship/Common/TagSimple_16", Icon16x16));
-		
 		Style->Set("Icons.Lock", new IMAGE_BRUSH_SVG("Starship/Common/lock", Icon16x16));
 		Style->Set("Icons.Unlock", new IMAGE_BRUSH_SVG("Starship/Common/lock-unlocked", Icon16x16));
 
@@ -396,7 +394,6 @@ TSharedRef<ISlateStyle> FStarshipCoreStyle::Create()
 		Style->Set("Icons.Refresh", new IMAGE_BRUSH_SVG("Starship/Common/refresh", Icon16x16));
 
 		Style->Set("Icons.Star", new IMAGE_BRUSH_SVG("Starship/Common/Favorite", Icon16x16));
-		Style->Set("Icons.Star.Outline", new IMAGE_BRUSH_SVG("Starship/Common/FavoriteOutline", Icon16x16));
 		Style->Set("Icons.Link", new IMAGE_BRUSH_SVG("Starship/Common/Linked", Icon16x16));
 		Style->Set("Icons.Unlink", new IMAGE_BRUSH_SVG("Starship/Common/Unlinked", Icon16x16));
 
