- Fix missing nullptr check

#rnx
#tests Create a new NS, add a Minimal emitter, add sub UV module and lnk to the sprite renderer

diff --git a/Engine/Plugins/FX/Niagara/Source/NiagaraEditor/Private/ViewModels/Stack/NiagaraStackPropertyRow.cpp b/Engine/Plugins/FX/Niagara/Source/NiagaraEditor/Private/ViewModels/Stack/NiagaraStackPropertyRow.cpp
--- a/Engine/Plugins/FX/Niagara/Source/NiagaraEditor/Private/ViewModels/Stack/NiagaraStackPropertyRow.cpp
+++ b/Engine/Plugins/FX/Niagara/Source/NiagaraEditor/Private/ViewModels/Stack/NiagaraStackPropertyRow.cpp
@@ -82,6 +82,11 @@ UNiagaraStackEntry::EStackRowStyle UNiagaraStackPropertyRow::GetStackRowStyle()
 
 bool UNiagaraStackPropertyRow::GetShouldShowInStack() const
 {
+	if (DetailTreeNode == nullptr)
+	{
+		return bIsHiddenCategory == false;
+	}
+
 	if(DetailTreeNode->GetNodeType() == EDetailNodeType::Category)
 	{
 		if(bIsHiddenCategory)
