- Check dynamic parameters module enabled state when gathering the parameter mask

#rnx
#tests Make a system with a stateless emitter, set material to debug draw dynamic parameters, observe the material defaults from Param 0 do not show up

diff --git a/Engine/Plugins/FX/Niagara/Source/Niagara/Internal/Stateless/Modules/NiagaraStatelessModule_DynamicMaterialParameters.h b/Engine/Plugins/FX/Niagara/Source/Niagara/Internal/Stateless/Modules/NiagaraStatelessModule_DynamicMaterialParameters.h
--- a/Engine/Plugins/FX/Niagara/Source/Niagara/Internal/Stateless/Modules/NiagaraStatelessModule_DynamicMaterialParameters.h
+++ b/Engine/Plugins/FX/Niagara/Source/Niagara/Internal/Stateless/Modules/NiagaraStatelessModule_DynamicMaterialParameters.h
@@ -154,9 +154,12 @@ public:
 	int32 GetRendererChannelMask() const
 	{
 		uint32 ChannelMask = 0;
-		for (int32 iParameter = 0; iParameter < NumParameters; ++iParameter)
+		if (IsModuleEnabled())
 		{
-			ChannelMask |= GetParameterChannelMask(iParameter) << (iParameter * NumChannelPerParameter);
+			for (int32 iParameter = 0; iParameter < NumParameters; ++iParameter)
+			{
+				ChannelMask |= GetParameterChannelMask(iParameter) << (iParameter * NumChannelPerParameter);
+			}
 		}
 		return ChannelMask;
 	}
