Added alternate search names for linked input parameters.
This way 'Particles.Age' will get an exact match if you type 'Age'
- Also fixed a minor typo

diff --git a/Engine/Plugins/FX/Niagara/Source/NiagaraEditor/Public/ViewModels/Stack/NiagaraStackEntry.h b/Engine/Plugins/FX/Niagara/Source/NiagaraEditor/Public/ViewModels/Stack/NiagaraStackEntry.h
--- a/Engine/Plugins/FX/Niagara/Source/NiagaraEditor/Public/ViewModels/Stack/NiagaraStackEntry.h
+++ b/Engine/Plugins/FX/Niagara/Source/NiagaraEditor/Public/ViewModels/Stack/NiagaraStackEntry.h
@@ -378,10 +378,10 @@ public:
 		GetFilteredChildren(OutFilteredChildren);
 		for (UNiagaraStackEntry* FilteredChild : OutFilteredChildren)
 		{
-			T* UnfilteredChildOfType = Cast<T>(FilteredChild);
-			if (UnfilteredChildOfType != nullptr)
+			T* FilteredChildOfType = Cast<T>(FilteredChild);
+			if (FilteredChildOfType != nullptr)
 			{
-				OutFilteredChildrenOfType.Add(UnfilteredChildOfType);
+				OutFilteredChildrenOfType.Add(FilteredChildOfType);
 			}
 		}
 
diff --git a/Engine/Plugins/FX/Niagara/Source/NiagaraEditorWidgets/Private/Stack/SNiagaraStackFunctionInputValue.cpp b/Engine/Plugins/FX/Niagara/Source/NiagaraEditorWidgets/Private/Stack/SNiagaraStackFunctionInputValue.cpp
--- a/Engine/Plugins/FX/Niagara/Source/NiagaraEditorWidgets/Private/Stack/SNiagaraStackFunctionInputValue.cpp
+++ b/Engine/Plugins/FX/Niagara/Source/NiagaraEditorWidgets/Private/Stack/SNiagaraStackFunctionInputValue.cpp
@@ -1357,6 +1357,7 @@ TArray<TSharedPtr<FNiagaraMenuAction_Generic>> SNiagaraStackFunctionInputValue::
 			
 			LinkAction->SetParameterVariable(FNiagaraVariable(FunctionInput->GetInputType(), AvailableHandle.GetParameterHandleString()));
 			LinkAction->SourceData = NiagaraSourceData;
+			LinkAction->AlternateSearchName = FText::FromName(HandleParts.Last());
 
 			OutAllActions.Add(LinkAction);
 		}
@@ -1393,6 +1394,7 @@ TArray<TSharedPtr<FNiagaraMenuAction_Generic>> SNiagaraStackFunctionInputValue::
 			// set the source data from the script
 			TTuple<EScriptSource, FText> Source = FNiagaraEditorUtilities::GetScriptSource(ConversionScript);
 			LinkAction->SourceData = FNiagaraActionSourceData(Source.Key, Source.Value, true);;
+			LinkAction->AlternateSearchName = FText::FromName(HandleParts.Last());
 
 			OutAllActions.Add(LinkAction);
 		}
