Lifted hierarchy order restrictions in Niagara Scripts (parameters and categories can be sorted freely now)
- Fixed child inputs allowing for layers > 1
- Fixed stack inputs with a Data Interface that have children inputs listing the children inputs before the Data Interface properties

diff --git a/Engine/Plugins/FX/Niagara/Source/NiagaraEditor/Private/ViewModels/HierarchyEditor/NiagaraScriptParametersHierarchyViewModel.cpp b/Engine/Plugins/FX/Niagara/Source/NiagaraEditor/Private/ViewModels/HierarchyEditor/NiagaraScriptParametersHierarchyViewModel.cpp
--- a/Engine/Plugins/FX/Niagara/Source/NiagaraEditor/Private/ViewModels/HierarchyEditor/NiagaraScriptParametersHierarchyViewModel.cpp
+++ b/Engine/Plugins/FX/Niagara/Source/NiagaraEditor/Private/ViewModels/HierarchyEditor/NiagaraScriptParametersHierarchyViewModel.cpp
@@ -131,6 +131,10 @@ TSharedPtr<FHierarchyElementViewModel> UNiagaraScriptParametersHierarchyViewMode
 	{
 		return MakeShared<FNiagaraHierarchyScriptParameterViewModel>(Item, Parent.ToSharedRef(), this);
 	}
+	if(UHierarchyCategory* Category = Cast<UHierarchyCategory>(Element))
+	{
+		return MakeShared<FNiagaraHierarchyScriptCategoryViewModel>(Category, Parent.ToSharedRef(), this);
+	}
 	else if(UHierarchyRoot* Root = Cast<UHierarchyRoot>(Element))
 	{
 		// If the root is the hierarchy root, we know it's for the hierarchy. If not, it's the transient source root
@@ -251,6 +255,25 @@ TSharedRef<SWidget> FNiagaraHierarchyScriptParameterDragDropOp::CreateCustomDeco
 	return SNullWidget::NullWidget;
 }
 
+FHierarchyElementViewModel::FCanPerformActionResults FNiagaraHierarchyScriptCategoryViewModel::CanDropOnInternal(TSharedPtr<FHierarchyElementViewModel> DraggedElement, EItemDropZone ItemDropZone)
+{
+	if(IsEditableByUser().bCanPerform == false)
+	{
+		return false;
+	}
+	
+	FCanPerformActionResults CanPerformActionResults = FHierarchyCategoryViewModel::CanDropOnInternal(DraggedElement, ItemDropZone);
+
+	// We allow dropping parameters above/below categories
+	if(ItemDropZone != EItemDropZone::OntoItem && DraggedElement->GetData()->IsA<UNiagaraHierarchyScriptParameter>())
+	{
+		CanPerformActionResults.CanPerformMessage = FText::GetEmpty();
+		CanPerformActionResults.bCanPerform = true;
+	}
+
+	return CanPerformActionResults;
+}
+
 bool FNiagaraHierarchyScriptParameterViewModel::DoesExternalDataStillExist(const UHierarchyDataRefreshContext* Context) const
 {
 	// GetScriptVariable checks the owning graph if it still exists
@@ -293,9 +316,28 @@ FHierarchyElementViewModel::FCanPerformActionResults FNiagaraHierarchyScriptPara
 			// We support nested inputs
 			if(DraggedItem->GetData()->IsA<UNiagaraHierarchyScriptParameter>() && TargetDropItem->GetData()->IsA<UNiagaraHierarchyScriptParameter>())
 			{
+				// But not if the dragged input already has a child
+				if(DraggedItem->GetData()->DoesOneChildExist<UNiagaraHierarchyScriptParameter>())
+				{
+					FText BaseMessage = LOCTEXT("DraggedItemHasChild_NotAllowed", "Can't nest input {0} under input {1}.\nRemove children of {2} first. ");
+					AllowDrop.CanPerformMessage = FText::FormatOrdered(BaseMessage, DraggedItem->ToStringAsText(), TargetDropItem->ToStringAsText(), DraggedItem->ToStringAsText());
+					AllowDrop.bCanPerform = false;
+					return AllowDrop;
+				}
+				
+				// And only up to 1 layer if we are going to create a new child input
+				if(TargetDropItem->GetParent().Pin()->GetData<UNiagaraHierarchyScriptParameter>())
+				{
+					FText BaseMessage = LOCTEXT("DroppingInputOnInputNestedChildTooDeep", "Can't nest input {0} under input {1}.\nChildren inputs can only have one layer of depth!");
+					AllowDrop.CanPerformMessage = FText::FormatOrdered(BaseMessage, DraggedItem->ToStringAsText(), TargetDropItem->ToStringAsText());
+					AllowDrop.bCanPerform = false;
+					return AllowDrop;
+				}
+				
 				FText BaseMessage = LOCTEXT("DroppingInputOnInputNestedChild", "This will nest input {0} under input {1}");
 				AllowDrop.CanPerformMessage = FText::FormatOrdered(BaseMessage, DraggedItem->ToStringAsText(), TargetDropItem->ToStringAsText());
 				AllowDrop.bCanPerform = true;
+				return AllowDrop;
 			}
 		}
 		else
@@ -305,12 +347,13 @@ FHierarchyElementViewModel::FCanPerformActionResults FNiagaraHierarchyScriptPara
 			{
 				AllowDrop.bCanPerform = true;
 			}
-			else
+			// if the dragged item is a category, we generally allow putting it above/below other parameters, but not above/below child parameters
+			else if(DraggedItem->GetData()->IsA<UHierarchyCategory>())
 			{
-				// we use default logic only if there is no parent input. Nested children are not allowed to contain anything but other inputs.
-				if(TargetDropItem->GetParent().Pin()->GetData<UNiagaraHierarchyScriptParameter>() == nullptr)
+				AllowDrop.bCanPerform = true;
+				if(TargetDropItem->GetData()->IsA(UNiagaraHierarchyScriptParameter::StaticClass()) && TargetDropItem->GetParent().Pin()->GetData()->IsA(UNiagaraHierarchyScriptParameter::StaticClass()))
 				{
-					AllowDrop = FHierarchyItemViewModel::CanDropOnInternal(DraggedItem, ItemDropZone);
+					AllowDrop.bCanPerform = false;
 				}
 			}
 		}
@@ -346,12 +389,4 @@ void FNiagaraHierarchyScriptParameterViewModel::OnDroppedOnInternal(TSharedPtr<F
 	}
 }
 
-void FNiagaraHierarchyScriptRootViewModel::SortChildrenData() const
-{
-	GetDataMutable()->GetChildrenMutable().StableSort([](const UHierarchyElement& ItemA, const UHierarchyElement& ItemB)
-		{
-			return ItemA.IsA<UHierarchyItem>() && ItemB.IsA<UHierarchyCategory>();
-		});
-}
-
 #undef LOCTEXT_NAMESPACE
diff --git a/Engine/Plugins/FX/Niagara/Source/NiagaraEditor/Public/ViewModels/HierarchyEditor/NiagaraScriptParametersHierarchyViewModel.h b/Engine/Plugins/FX/Niagara/Source/NiagaraEditor/Public/ViewModels/HierarchyEditor/NiagaraScriptParametersHierarchyViewModel.h
--- a/Engine/Plugins/FX/Niagara/Source/NiagaraEditor/Public/ViewModels/HierarchyEditor/NiagaraScriptParametersHierarchyViewModel.h
+++ b/Engine/Plugins/FX/Niagara/Source/NiagaraEditor/Public/ViewModels/HierarchyEditor/NiagaraScriptParametersHierarchyViewModel.h
@@ -90,6 +90,17 @@ public:
 	virtual TSharedRef<SWidget> CreateCustomDecorator() const override;
 };
 
+struct FNiagaraHierarchyScriptCategoryViewModel : public FHierarchyCategoryViewModel
+{
+	FNiagaraHierarchyScriptCategoryViewModel(UHierarchyCategory* Category, TSharedRef<FHierarchyElementViewModel> InParent, TWeakObjectPtr<UNiagaraScriptParametersHierarchyViewModel> ViewModel)
+		: FHierarchyCategoryViewModel(Category, InParent, ViewModel) {}
+
+protected:
+	// We don't sort anything to allow free authoring of UI
+	virtual void SortChildrenData() const override {}
+	virtual FCanPerformActionResults CanDropOnInternal(TSharedPtr<FHierarchyElementViewModel> DraggedElement, EItemDropZone ItemDropZone) override;
+};
+
 struct FNiagaraHierarchyScriptParameterViewModel : public FHierarchyItemViewModel
 {
 	FNiagaraHierarchyScriptParameterViewModel(UNiagaraHierarchyScriptParameter* ScriptParameter, TSharedRef<FHierarchyElementViewModel> InParent, TWeakObjectPtr<UNiagaraScriptParametersHierarchyViewModel> ViewModel)
@@ -115,6 +126,6 @@ struct FNiagaraHierarchyScriptRootViewModel : public FHierarchyRootViewModel
 		: FHierarchyRootViewModel(Root, ViewModel, bInIsForHierarchy)	{}
 
 protected:
-	// In the script, loose parameters are always added before categories (reverse of the default implementation)
-	virtual void SortChildrenData() const override;
+	// We don't sort anything to allow free authoring of UI
+	virtual void SortChildrenData() const override { }
 };
