- Color Graident changes for stateless
- Don't default to color gradient if we have keys outside of 0-1 as they won't show and it's confusing
- If they switch to color gradient offer to trim the keys to be inside the valid range

#tests Create a color curve with a key outside of 0-1, before change this wouldn't show, after change you get the option to trim
#rnx

diff --git a/Engine/Plugins/FX/Niagara/Source/NiagaraEditor/Private/Stateless/NiagaraDistributionPropertyCustomization.cpp b/Engine/Plugins/FX/Niagara/Source/NiagaraEditor/Private/Stateless/NiagaraDistributionPropertyCustomization.cpp
--- a/Engine/Plugins/FX/Niagara/Source/NiagaraEditor/Private/Stateless/NiagaraDistributionPropertyCustomization.cpp
+++ b/Engine/Plugins/FX/Niagara/Source/NiagaraEditor/Private/Stateless/NiagaraDistributionPropertyCustomization.cpp
@@ -6,6 +6,7 @@
 #include "Editor.h"
 #include "EditorUndoClient.h"
 #include "Misc/AxisDisplayInfo.h"
+#include "Misc/MessageDialog.h"
 #include "PropertyHandle.h"
 #include "ScopedTransaction.h"
 #include "Stateless/NiagaraStatelessDistribution.h"
@@ -211,7 +212,7 @@ public:
 		}
 	}
 
-	static ENiagaraDistributionEditorMode GetDistributionEditorModeFromSourceMode(int32 InSourceNumChannels, bool bInDisplayAsColor, ENiagaraDistributionMode InSourceMode, TOptional<ENiagaraDistributionEditorMode> CurrentMode = {})
+	ENiagaraDistributionEditorMode GetDistributionEditorModeFromSourceMode(int32 InSourceNumChannels, bool bInDisplayAsColor, ENiagaraDistributionMode InSourceMode, TOptional<ENiagaraDistributionEditorMode> CurrentMode = {})
 	{
 		if (bInDisplayAsColor)
 		{
@@ -226,7 +227,11 @@ public:
 			case ENiagaraDistributionMode::NonUniformRange:
 				return ENiagaraDistributionEditorMode::ColorRange;
 			case ENiagaraDistributionMode::NonUniformCurve:
-				return CurrentMode.IsSet() && CurrentMode.GetValue() == ENiagaraDistributionEditorMode::NonUniformCurve ? ENiagaraDistributionEditorMode::NonUniformCurve : ENiagaraDistributionEditorMode::ColorGradient;
+				if (CurrentMode.IsSet())
+				{
+					return CurrentMode.GetValue() == ENiagaraDistributionEditorMode::NonUniformCurve ? ENiagaraDistributionEditorMode::NonUniformCurve : ENiagaraDistributionEditorMode::ColorGradient;
+				}
+				return AreCurveKeyTimesNormalized() ?  ENiagaraDistributionEditorMode::ColorGradient : ENiagaraDistributionEditorMode::NonUniformCurve;
 			default:
 				return ENiagaraDistributionEditorMode::ColorConstant;
 			}
@@ -321,17 +326,43 @@ public:
 			return;
 		}
 
+		const ENiagaraDistributionMode NewSourceMode = GetDistributionSourceModeFromEditorMode(InEditorMode);
+
+		// If the distribution mode is not changing and we are switching to be a color gradient we need to make sure the keys are all normalized (i.e. 0-1 time)
+		// if they are not we need to perform a destructive action and trim the keys to 0-1 time range or the widget will not display the keys correctly.
+		bool bNeedsUpdateForColorGradient = false;
+		if (NewSourceMode == SourceDistribution->Mode && InEditorMode == ENiagaraDistributionEditorMode::ColorGradient && !AreCurveKeyTimesNormalized())
+		{
+			const EAppReturnType::Type DialogResponse = FMessageDialog::Open(
+				EAppMsgType::OkCancel,
+				EAppReturnType::Ok,
+				LOCTEXT("TrimKeysForColorGradientDesc", "The color gradient editor only allows keys with a time range between 0.0 and 1.0, your existing curve has values outside of this time range.  Would you like to trim the keys to 0.0 to 1.0 time range?"),
+				LOCTEXT("TrimKeysForColorGradientTitle", "Trim key times 0.0 to 1.0?")
+			);
+			if (DialogResponse != EAppReturnType::Ok)
+			{
+				return;
+			}
+			bNeedsUpdateForColorGradient = true;
+		}
+
 		TGuardValue<bool> UpdateGuard(bUpdatingDistributionModeInternally, true);
-		int32 NewEditorNumChannels = FNiagaraDistributionEditorUtilities::IsUniform(InEditorMode) ? 1 : SourceNumChannels;
+		const int32 NewEditorNumChannels = FNiagaraDistributionEditorUtilities::IsUniform(InEditorMode) ? 1 : SourceNumChannels;
 
-		ENiagaraDistributionMode NewSourceMode = GetDistributionSourceModeFromEditorMode(InEditorMode);
-		if (SourceDistribution->Mode != NewSourceMode)
+		if (SourceDistribution->Mode != NewSourceMode || bNeedsUpdateForColorGradient)
 		{
 			const FScopedTransaction Transaction(LOCTEXT("SetDistributionModeTransaction", "Set distribution mode"));
 			OwnerObject->Modify();
 			PropertyHandle->NotifyPreChange();
 
-			MigrateDataFromModeChange(NewSourceMode, NewEditorNumChannels);
+			if (bNeedsUpdateForColorGradient)
+			{
+				MigrateDataToColorGradient();
+			}
+			else
+			{
+				MigrateDataFromModeChange(NewSourceMode, NewEditorNumChannels);
+			}
 			SourceDistribution->Mode = NewSourceMode;
 
 			PropertyHandle->NotifyPostChange(EPropertyChangeType::ValueSet);
@@ -830,6 +861,65 @@ private:
 		}
 	}
 
+	void MigrateDataToColorGradient()
+	{
+		if ( SourceDistribution->Mode != ENiagaraDistributionMode::NonUniformCurve || AreCurveKeyTimesNormalized() )
+		{
+			return;
+		}
+
+		for (FRichCurve& ChannelCurve : SourceDistribution->ChannelCurves)
+		{
+			if (ChannelCurve.Keys.Num() == 0)
+			{
+				continue;
+			}
+
+			// Trim Start
+			if (ChannelCurve.Keys[0].Time < 0.0f)
+			{
+				ChannelCurve.AddKey(0.0f, ChannelCurve.Eval(0.0f, 0.0));
+				while (ChannelCurve.Keys[0].Time < 0.0f)
+				{
+					ChannelCurve.Keys.RemoveAt(0);
+				}
+			}
+
+			// Trim End
+			if (ChannelCurve.Keys.Last().Time > 1.0f)
+			{
+				ChannelCurve.AddKey(1.0f, ChannelCurve.Eval(1.0f, 0.0));
+				while (ChannelCurve.Keys.Last().Time > 1.0f)
+				{
+					ChannelCurve.Keys.RemoveAt(ChannelCurve.Keys.Num() - 1);
+				}
+			}
+		}
+	}
+
+	bool AreCurveKeyTimesNormalized() const
+	{
+		for (FRichCurve& ChannelCurve : SourceDistribution->ChannelCurves)
+		{
+			if (ChannelCurve.Keys.Num() == 0)
+			{
+				continue;
+			}
+
+			if (ChannelCurve.Keys[0].Time < 0.0f)
+			{
+				return false;
+			}
+
+			if (ChannelCurve.Keys.Last().Time > 1.0f)
+			{
+				return false;
+			}
+		}
+		return true;
+	}
+
+
 private:
 	TSharedPtr<IPropertyHandle> PropertyHandle;
 	TWeakObjectPtr<UObject> OwnerObjectWeak;
