- Fix for aliasing issue with Clang v20.1

#tests Ran ET and a few other projects to ensure no regression
#rnx

diff --git a/Engine/Plugins/FX/Niagara/Source/Niagara/Classes/NiagaraDataSet.h b/Engine/Plugins/FX/Niagara/Source/Niagara/Classes/NiagaraDataSet.h
--- a/Engine/Plugins/FX/Niagara/Source/Niagara/Classes/NiagaraDataSet.h
+++ b/Engine/Plugins/FX/Niagara/Source/Niagara/Classes/NiagaraDataSet.h
@@ -591,6 +591,11 @@ namespace NiagaraDataSetPrivate
 	{
 		return DataBuffer->GetNumInstances();
 	}
+
+	inline uint32 GetComponentStride(const FNiagaraDataBuffer* DataBuffer)
+	{
+		return DataBuffer->GetFloatStride() / sizeof(float);
+	}
 }
 
 #if UE_ENABLE_INCLUDE_ORDER_DEPRECATED_IN_5_6
diff --git a/Engine/Plugins/FX/Niagara/Source/Niagara/Classes/NiagaraDataSetAccessor.h b/Engine/Plugins/FX/Niagara/Source/Niagara/Classes/NiagaraDataSetAccessor.h
--- a/Engine/Plugins/FX/Niagara/Source/Niagara/Classes/NiagaraDataSetAccessor.h
+++ b/Engine/Plugins/FX/Niagara/Source/Niagara/Classes/NiagaraDataSetAccessor.h
@@ -24,6 +24,7 @@ namespace NiagaraDataSetPrivate
 	inline const uint8* GetComponentPtrHalf(const FNiagaraDataBuffer* DataBuffer, uint32 ComponentIdx);
 	inline const uint8* GetComponentPtrInt32(const FNiagaraDataBuffer* DataBuffer, uint32 ComponentIdx);
 	inline uint32 GetNumInstances(const FNiagaraDataBuffer* DataBuffer);
+	inline uint32 GetComponentStride(const FNiagaraDataBuffer* DataBuffer);
 }
 
 template<typename TType>
@@ -533,10 +534,6 @@ struct FNiagaraDataSetAccessorTypeInfo<ENiagaraExecutionState>
 // Structure readers
 
 
-#if defined(_MSC_VER) && USING_CODE_ANALYSIS
-#pragma warning(disable:6294) // Ill-defined for-loop:  initial condition does not satisfy test.  Loop body not executed.
-#endif
-
 template<typename TType>
 struct FNiagaraDataSetReaderStruct
 {
@@ -549,16 +546,10 @@ struct FNiagaraDataSetReaderStruct
 		bIsValid &= (FNiagaraDataSetAccessorTypeInfo<TType>::NumInt32Components == 0) || (Int32ComponentIndex != INDEX_NONE);
 		if (bIsValid)
 		{
-			int32 ComponentIndex = 0;
-			for (int i = 0; i < FNiagaraDataSetAccessorTypeInfo<TType>::NumFloatComponents; ++i)
-			{
-				ComponentData[ComponentIndex++] = NiagaraDataSetPrivate::GetComponentPtrFloat(DataBuffer, FloatComponentIndex + i);
-			}
-			for (int i = 0; i < FNiagaraDataSetAccessorTypeInfo<TType>::NumInt32Components; ++i)
-			{
-				ComponentData[ComponentIndex++] = NiagaraDataSetPrivate::GetComponentPtrInt32(DataBuffer, Int32ComponentIndex + i);
-			}
+			FloatComponents = FNiagaraDataSetAccessorTypeInfo<TType>::NumFloatComponents > 0 ? reinterpret_cast<const float*>(NiagaraDataSetPrivate::GetComponentPtrFloat(DataBuffer, FloatComponentIndex)) : nullptr;
+			Int32Components = FNiagaraDataSetAccessorTypeInfo<TType>::NumInt32Components > 0 ? reinterpret_cast<const int32*>(NiagaraDataSetPrivate::GetComponentPtrInt32(DataBuffer, Int32ComponentIndex)) : nullptr;
 			NumInstances = NiagaraDataSetPrivate::GetNumInstances(DataBuffer);
+			ComponentStride = NiagaraDataSetPrivate::GetComponentStride(DataBuffer);
 		}
 	}
 
@@ -572,11 +563,7 @@ struct FNiagaraDataSetReaderStruct
 	FORCEINLINE TType Get(int Index) const
 	{
 		check(bIsValid && Index <= NumInstances);
-		return FNiagaraDataSetAccessorTypeInfo<TType>::ReadStruct(
-			Index,
-			FNiagaraDataSetAccessorTypeInfo<TType>::NumFloatComponents > 0 ? reinterpret_cast<const float*const*>(ComponentData) : nullptr,
-			FNiagaraDataSetAccessorTypeInfo<TType>::NumInt32Components > 0 ? reinterpret_cast<const int32*const*>(ComponentData + FNiagaraDataSetAccessorTypeInfo<TType>::NumFloatComponents) : nullptr
-		);
+		return FNiagaraDataSetAccessorTypeInfo<TType>::ReadStruct(Index, FloatComponents, Int32Components, ComponentStride);
 	}
 
 	FORCEINLINE TType GetSafe(int Index, const TType& Default) const
@@ -593,7 +580,9 @@ struct FNiagaraDataSetReaderStruct
 private:
 	bool bIsValid = false;
 	int32 NumInstances = 0;
-	const uint8* ComponentData[FNiagaraDataSetAccessorTypeInfo<TType>::NumFloatComponents + FNiagaraDataSetAccessorTypeInfo<TType>::NumInt32Components] = {};
+	int32 ComponentStride = 0;
+	const float* FloatComponents = nullptr;
+	const int32* Int32Components = nullptr;
 };
 
 #if defined(_MSC_VER) && USING_CODE_ANALYSIS
@@ -630,9 +619,11 @@ struct FNiagaraDataSetAccessorStruct
 		const int32 VariableIndex = DataSetCompiledData.Variables.IndexOfByPredicate(FNiagaraVariableMatch(FNiagaraDataSetAccessorTypeInfo<TType>::GetStructType(), VariableName));
 		if (VariableIndex != INDEX_NONE)
 		{
-			bIsValid = true;
-			FloatComponentIndex = DataSetCompiledData.VariableLayouts[VariableIndex].GetFloatComponentStart();
-			Int32ComponentIndex = DataSetCompiledData.VariableLayouts[VariableIndex].GetInt32ComponentStart();
+			const FNiagaraVariableLayoutInfo& VariableLayout = DataSetCompiledData.VariableLayouts[VariableIndex];
+			ensure(VariableLayout.GetNumHalfComponents() == 0);	// We don't support reading halfs for structs
+			bIsValid = VariableLayout.GetNumHalfComponents() == 0;
+			FloatComponentIndex = VariableLayout.GetFloatComponentStart();
+			Int32ComponentIndex = VariableLayout.GetInt32ComponentStart();
 		}
 	}
 
@@ -660,13 +651,13 @@ struct FNiagaraDataSetAccessorTypeInfo<FNiagaraSpawnInfo>
 	static constexpr int32 NumInt32Components = 2;
 	static FNiagaraTypeDefinition GetStructType() { return FNiagaraTypeDefinition(FNiagaraSpawnInfo::StaticStruct()); }
 
-	static FORCENOINLINE FNiagaraSpawnInfo ReadStruct(int32 InstanceIndex, const float*const* FloatComponents, const int32*const* Int32Components)
+	static FNiagaraSpawnInfo ReadStruct(int32 InstanceIndex, const float* FloatComponents, const int32* Int32Components, int32 ComponentStride)
 	{
 		FNiagaraSpawnInfo SpawnInfo;
-		SpawnInfo.Count			= Int32Components[0][InstanceIndex];
-		SpawnInfo.InterpStartDt	= FloatComponents[0][InstanceIndex];
-		SpawnInfo.IntervalDt	= FloatComponents[1][InstanceIndex];
-		SpawnInfo.SpawnGroup	= Int32Components[1][InstanceIndex];
+		SpawnInfo.Count			= Int32Components[(ComponentStride * 0) + InstanceIndex];
+		SpawnInfo.InterpStartDt	= FloatComponents[(ComponentStride * 0) + InstanceIndex];
+		SpawnInfo.IntervalDt	= FloatComponents[(ComponentStride * 1) + InstanceIndex];
+		SpawnInfo.SpawnGroup	= Int32Components[(ComponentStride * 1) + InstanceIndex];
 		return SpawnInfo;
 	}
 };
@@ -680,11 +671,11 @@ struct FNiagaraDataSetAccessorTypeInfo<FNiagaraID>
 	static constexpr int32 NumInt32Components = 2;
 	static FNiagaraTypeDefinition GetStructType() { return FNiagaraTypeDefinition(FNiagaraID::StaticStruct()); }
 
-	static FNiagaraID ReadStruct(int32 InstanceIndex, const float*const* FloatComponents, const int32*const* Int32Components)
+	static FNiagaraID ReadStruct(int32 InstanceIndex, const float* FloatComponents, const int32* Int32Components, int32 ComponentStride)
 	{
 		FNiagaraID NiagaraID;
-		NiagaraID.Index			= Int32Components[0][InstanceIndex];
-		NiagaraID.AcquireTag	= Int32Components[1][InstanceIndex];
+		NiagaraID.Index			= Int32Components[(ComponentStride * 0) + InstanceIndex];
+		NiagaraID.AcquireTag	= Int32Components[(ComponentStride * 1) + InstanceIndex];
 		return NiagaraID;
 	}
 };
