- Stateless noise module updates (no longer experimental)
- Uses a LUT to avoid heavy math ops
- Adds a CPU fallback path

#rnx

diff --git a/Engine/Plugins/FX/Niagara/Shaders/Private/Stateless/Modules/NiagaraStatelessModule_SolveVelocitiesAndForces.ush b/Engine/Plugins/FX/Niagara/Shaders/Private/Stateless/Modules/NiagaraStatelessModule_SolveVelocitiesAndForces.ush
--- a/Engine/Plugins/FX/Niagara/Shaders/Private/Stateless/Modules/NiagaraStatelessModule_SolveVelocitiesAndForces.ush
+++ b/Engine/Plugins/FX/Niagara/Shaders/Private/Stateless/Modules/NiagaraStatelessModule_SolveVelocitiesAndForces.ush
@@ -28,16 +28,13 @@ float	SolveVelocitiesAndForces_PointVelocityScale;
 float	SolveVelocitiesAndForces_PointVelocityBias;
 float3	SolveVelocitiesAndForces_PointOrigin;
 
-uint	SolveVelocitiesAndForces_NoiseEnabled;
-float	SolveVelocitiesAndForces_NoiseAmplitude;
-float3	SolveVelocitiesAndForces_NoiseFrequency;
-//Texture3D		SolveVelocitiesAndForces_NoiseTexture;
-//SamplerState	SolveVelocitiesAndForces_NoiseSampler;
-
-uint	SolveVelocitiesAndForces_NoiseMode;
-uint	SolveVelocitiesAndForces_NoiseLUTOffset;
-uint	SolveVelocitiesAndForces_NoiseLUTNumChannel;
-uint	SolveVelocitiesAndForces_NoiseLUTChannelWidth;
+uint				SolveVelocitiesAndForces_NoiseEnabled;
+float				SolveVelocitiesAndForces_NoiseStrength;
+float				SolveVelocitiesAndForces_NoiseFrequency;
+float				SolveVelocitiesAndForces_NoiseLUTNormalize;
+uint				SolveVelocitiesAndForces_NoiseLUTRowMask;
+uint				SolveVelocitiesAndForces_NoiseLUTRowWidth;
+ByteAddressBuffer	SolveVelocitiesAndForces_NoiseLUT;
 
 struct FStatelessModule_SolveVelocitiesAndForces
 {
@@ -48,6 +45,26 @@ struct FStatelessModule_SolveVelocitiesAndForces
 	float3 Acceleration;
 };
 
+void SolveVelocitiesAndForces_EvaluateNoise(in FStatelessModule_SolveVelocitiesAndForces ModuleData, float NormalizedAge, float Age, inout float3 Position)
+{
+	const float fRow	= length2(Position * 0.01f);
+	const uint RowA		= uint(floor(abs(fRow))) & SolveVelocitiesAndForces_NoiseLUTRowMask;
+	const uint RowB		= (RowA + 1) & SolveVelocitiesAndForces_NoiseLUTRowMask;
+		
+	const float fColumn = abs(SolveVelocitiesAndForces_NoiseFrequency) * NormalizedAge;
+	const uint iColumnA = min(floor(fColumn), SolveVelocitiesAndForces_NoiseLUTRowWidth - 1);
+	const uint iColumnB = min(iColumnA + 1, SolveVelocitiesAndForces_NoiseLUTRowWidth - 1);
+
+	const float3 ValueA0 = asfloat(SolveVelocitiesAndForces_NoiseLUT.Load3(((SolveVelocitiesAndForces_NoiseLUTRowWidth * RowA) + iColumnA) * 12));
+	const float3 ValueA1 = asfloat(SolveVelocitiesAndForces_NoiseLUT.Load3(((SolveVelocitiesAndForces_NoiseLUTRowWidth * RowB) + iColumnA) * 12));
+	const float3 ValueB0 = asfloat(SolveVelocitiesAndForces_NoiseLUT.Load3(((SolveVelocitiesAndForces_NoiseLUTRowWidth * RowA) + iColumnB) * 12));
+	const float3 ValueB1 = asfloat(SolveVelocitiesAndForces_NoiseLUT.Load3(((SolveVelocitiesAndForces_NoiseLUTRowWidth * RowB) + iColumnB) * 12));
+	const float3 ValueA  = lerp(ValueA0, ValueA1, frac(fRow));
+	const float3 ValueB  = lerp(ValueB0, ValueB1, frac(fRow));
+	
+	Position += lerp(ValueA, ValueB, frac(fColumn)) * SolveVelocitiesAndForces_NoiseStrength * Age * SolveVelocitiesAndForces_NoiseLUTNormalize;
+}
+
 void SolveVelocitiesAndForces_IntegratePosition(in FStatelessModule_SolveVelocitiesAndForces ModuleData, float Age, inout float3 Position)
 {
 	if (ModuleData.Drag > 0.0001f)
@@ -108,65 +125,8 @@ void SolveVelocitiesAndForces_Simulate(inout FStatelessParticle Particle)
 
 	if ( SolveVelocitiesAndForces_NoiseEnabled  != 0 )
 	{
-		const float3 NoiseOffset = RandomFloat3(10);
-		float3 NoiseValue = 0;
-		//if ( SolveVelocitiesAndForces_NoiseMode == 0 )
-		//{
-		//	const float3 NoiseTravel = float3(Particle.Age, Particle.NormalizedAge, Particle.UniqueIndex);
-		//	const float3 NoisePosition = NoiseOffset + (NoiseTravel * SolveVelocitiesAndForces_NoiseFrequency);
-		//	NoiseValue = Texture3DSample(SolveVelocitiesAndForces_NoiseTexture, SolveVelocitiesAndForces_NoiseSampler, NoisePosition).xyz;
-		//}
-		if ( SolveVelocitiesAndForces_NoiseMode == 0 )
-		{
-			const float3 NoiseTravel = float3(Particle.Age, Particle.NormalizedAge, Particle.UniqueIndex) * 5.0f;
-			const float3 NoisePosition = NoiseOffset + (NoiseTravel * SolveVelocitiesAndForces_NoiseFrequency);
-			const float3x4 J = JacobianSimplex_ALU(NoisePosition, false, 1.0f);
-			NoiseValue = float3(J[1][2] - J[2][1], J[2][0] - J[0][2], J[0][1] - J[1][0]); // See comments to JacobianSimplex_ALU in Random.ush
-		}
-		else if ( SolveVelocitiesAndForces_NoiseMode == 1 )
-		{
-			const uint3 ChannelIndex = RandomUInt3(11) % SolveVelocitiesAndForces_NoiseLUTNumChannel;
-			const float3 ChannelLerp = (Particle.Age * SolveVelocitiesAndForces_NoiseFrequency) + RandomFloat(12);
-			const uint3 ChannelLUTOffsetA = (ChannelIndex * SolveVelocitiesAndForces_NoiseLUTChannelWidth) + min(floor(ChannelLerp) + 0, SolveVelocitiesAndForces_NoiseLUTChannelWidth - 1);
-			const uint3 ChannelLUTOffsetB = (ChannelIndex * SolveVelocitiesAndForces_NoiseLUTChannelWidth) + min(floor(ChannelLerp) + 1, SolveVelocitiesAndForces_NoiseLUTChannelWidth - 1);
-				
-			const float3 ValuesA = float3(
-					GetStaticFloat(SolveVelocitiesAndForces_NoiseLUTOffset, ChannelLUTOffsetA.x),
-					GetStaticFloat(SolveVelocitiesAndForces_NoiseLUTOffset, ChannelLUTOffsetA.y),
-					GetStaticFloat(SolveVelocitiesAndForces_NoiseLUTOffset, ChannelLUTOffsetA.z)
-			);
-			const float3 ValuesB = float3(
-					GetStaticFloat(SolveVelocitiesAndForces_NoiseLUTOffset, ChannelLUTOffsetB.x),
-					GetStaticFloat(SolveVelocitiesAndForces_NoiseLUTOffset, ChannelLUTOffsetB.y),
-					GetStaticFloat(SolveVelocitiesAndForces_NoiseLUTOffset, ChannelLUTOffsetB.z)
-			);
-			NoiseValue = lerp(ValuesA, ValuesB, frac(ChannelLerp));
-		}
-		else if ( SolveVelocitiesAndForces_NoiseMode == 2 )
-		{
-			const uint3 ChannelIndex = RandomUInt3(13) % SolveVelocitiesAndForces_NoiseLUTNumChannel;
-			const float ChannelLerp = (Particle.Age * SolveVelocitiesAndForces_NoiseFrequency.x) + RandomFloat(14);
-			const uint3 ChannelLUTOffsetA = (ChannelIndex * SolveVelocitiesAndForces_NoiseLUTChannelWidth) + min(floor(ChannelLerp) + 0, SolveVelocitiesAndForces_NoiseLUTChannelWidth - 1);
-			const uint3 ChannelLUTOffsetB = (ChannelIndex * SolveVelocitiesAndForces_NoiseLUTChannelWidth) + min(floor(ChannelLerp) + 1, SolveVelocitiesAndForces_NoiseLUTChannelWidth - 1);
-			const uint3 ChannelLUTOffsetC = (ChannelIndex * SolveVelocitiesAndForces_NoiseLUTChannelWidth) + min(floor(ChannelLerp) + 2, SolveVelocitiesAndForces_NoiseLUTChannelWidth - 1);
-			const uint3 ChannelLUTOffsetD = (ChannelIndex * SolveVelocitiesAndForces_NoiseLUTChannelWidth) + min(floor(ChannelLerp) + 3, SolveVelocitiesAndForces_NoiseLUTChannelWidth - 1);
-				
-			const float3 Values0 = float3(GetStaticFloat(SolveVelocitiesAndForces_NoiseLUTOffset, ChannelLUTOffsetA.x), GetStaticFloat(SolveVelocitiesAndForces_NoiseLUTOffset, ChannelLUTOffsetA.y), GetStaticFloat(SolveVelocitiesAndForces_NoiseLUTOffset, ChannelLUTOffsetA.z));
-			const float3 Values1 = float3(GetStaticFloat(SolveVelocitiesAndForces_NoiseLUTOffset, ChannelLUTOffsetB.x), GetStaticFloat(SolveVelocitiesAndForces_NoiseLUTOffset, ChannelLUTOffsetB.y), GetStaticFloat(SolveVelocitiesAndForces_NoiseLUTOffset, ChannelLUTOffsetB.z));
-			const float3 Values2 = float3(GetStaticFloat(SolveVelocitiesAndForces_NoiseLUTOffset, ChannelLUTOffsetC.x), GetStaticFloat(SolveVelocitiesAndForces_NoiseLUTOffset, ChannelLUTOffsetC.y), GetStaticFloat(SolveVelocitiesAndForces_NoiseLUTOffset, ChannelLUTOffsetC.z));
-			const float3 Values3 = float3(GetStaticFloat(SolveVelocitiesAndForces_NoiseLUTOffset, ChannelLUTOffsetD.x), GetStaticFloat(SolveVelocitiesAndForces_NoiseLUTOffset, ChannelLUTOffsetD.y), GetStaticFloat(SolveVelocitiesAndForces_NoiseLUTOffset, ChannelLUTOffsetD.z));
-			{
-				const float u = frac(ChannelLerp);
-				const float u2 = u * u;
-				const float3 a0 = Values3 - Values2 - Values0 + Values1;
-				const float3 a1 = Values0 - Values1 - a0;
-				const float3 a2 = Values2 - Values0;
-				const float3 a3 = Values1;
-				NoiseValue = a0*u*u2+a1*u2+a2*u+a3;
-			}
-		}
-
-		ModuleData.Velocity += NoiseValue.xyz * SolveVelocitiesAndForces_NoiseAmplitude;
+		SolveVelocitiesAndForces_EvaluateNoise(ModuleData, Particle.NormalizedAge, Particle.Age, Particle.Position);
+		SolveVelocitiesAndForces_EvaluateNoise(ModuleData, Particle.PreviousNormalizedAge, Particle.PreviousAge, Particle.PreviousPosition);
 	}
 	
 	// Simulate
diff --git a/Engine/Plugins/FX/Niagara/Source/Niagara/Internal/Stateless/Modules/NiagaraStatelessModuleCommon.h b/Engine/Plugins/FX/Niagara/Source/Niagara/Internal/Stateless/Modules/NiagaraStatelessModuleCommon.h
--- a/Engine/Plugins/FX/Niagara/Source/Niagara/Internal/Stateless/Modules/NiagaraStatelessModuleCommon.h
+++ b/Engine/Plugins/FX/Niagara/Source/Niagara/Internal/Stateless/Modules/NiagaraStatelessModuleCommon.h
@@ -40,13 +40,8 @@ namespace NiagaraStateless
 		float							PointVelocityMax = 0.0f;
 		FVector3f						PointOrigin = FVector3f::ZeroVector;
 
-		bool		bNoiseEnabled = false;
-		float		NoiseAmplitude = 0.0f;
-		float		NoiseFrequency = 0.0f;
-		UObject*	NoiseTexture = nullptr;
-		int32		NoiseMode = 0;
-		int32		NoiseLUTOffset = 0;
-		int32		NoiseLUTNumChannel = 0;
-		int32		NoiseLUTChannelWidth = 0;
+		bool							bNoiseEnabled = false;
+		float							NoiseStrength = 0.0f;
+		float							NoiseFrequency = 0.0f;
 	};
 }
diff --git a/Engine/Plugins/FX/Niagara/Source/Niagara/Internal/Stateless/Modules/NiagaraStatelessModule_CurlNoiseForce.h b/Engine/Plugins/FX/Niagara/Source/Niagara/Internal/Stateless/Modules/NiagaraStatelessModule_CurlNoiseForce.h
--- a/Engine/Plugins/FX/Niagara/Source/Niagara/Internal/Stateless/Modules/NiagaraStatelessModule_CurlNoiseForce.h
+++ b/Engine/Plugins/FX/Niagara/Source/Niagara/Internal/Stateless/Modules/NiagaraStatelessModule_CurlNoiseForce.h
@@ -9,36 +9,16 @@
 
 #include "NiagaraStatelessModule_CurlNoiseForce.generated.h"
 
-class UVectorField;
-
-UENUM()
-enum class ENSM_NoiseMode
-{
-	//VectorField,
-	JacobNoise,
-	LUTJacob,
-	LUTJacobBicubic,
-};
-
-UCLASS(MinimalAPI, EditInlineNew, Experimental, meta = (DisplayName = "Curl Noise Force"))
+UCLASS(MinimalAPI, EditInlineNew, meta = (DisplayName = "Curl Noise Force"))
 class UNiagaraStatelessModule_CurlNoiseForce : public UNiagaraStatelessModule
 {
 	GENERATED_BODY()
 
 public:
 	UPROPERTY(EditAnywhere, Category = "Parameters")
-	float NoiseAmplitude = 200.0f;
+	float NoiseStrength = 10.0f;
 	UPROPERTY(EditAnywhere, Category = "Parameters")
-	float NoiseFrequency = 1.0f;
-
-	UPROPERTY(EditAnywhere, Category = "Parameters")
-	ENSM_NoiseMode NoiseMode = ENSM_NoiseMode::JacobNoise;
-
-	//UPROPERTY(EditAnywhere, Category = "Parameters", meta = (AllowedClasses = "/Script/Engine.VectorField,/Script/Engine.VolumeTexture,/Script/Engine.TextureRenderTargetVolume"))
-	//TObjectPtr<UObject> NoiseTexture;
-
-	//-TODO: Add support for GPU once we settle on a feature set for this module
-	virtual ENiagaraStatelessFeatureMask GetFeatureMask() const { return ENiagaraStatelessFeatureMask::ExecuteGPU; }
+	float NoiseFrequency = 50.0f;
 
 	virtual void BuildEmitterData(const FNiagaraStatelessEmitterDataBuildContext& BuildContext) const override;
 
diff --git a/Engine/Plugins/FX/Niagara/Source/Niagara/Private/Stateless/Modules/NiagaraStatelessModule_CurlNoiseForce.cpp b/Engine/Plugins/FX/Niagara/Source/Niagara/Private/Stateless/Modules/NiagaraStatelessModule_CurlNoiseForce.cpp
--- a/Engine/Plugins/FX/Niagara/Source/Niagara/Private/Stateless/Modules/NiagaraStatelessModule_CurlNoiseForce.cpp
+++ b/Engine/Plugins/FX/Niagara/Source/Niagara/Private/Stateless/Modules/NiagaraStatelessModule_CurlNoiseForce.cpp
@@ -1,7 +1,6 @@
 // Copyright Epic Games, Inc. All Rights Reserved.
 
 #include "Stateless/Modules/NiagaraStatelessModule_CurlNoiseForce.h"
-#include "NiagaraSimplexNoise.h"
 
 void UNiagaraStatelessModule_CurlNoiseForce::BuildEmitterData(const FNiagaraStatelessEmitterDataBuildContext& BuildContext) const
 {
@@ -12,39 +11,6 @@ void UNiagaraStatelessModule_CurlNoiseForce::BuildEmitterData(const FNiagaraStat
 
 	NiagaraStateless::FPhysicsBuildData& PhysicsBuildData = BuildContext.GetTransientBuildData<NiagaraStateless::FPhysicsBuildData>();
 	PhysicsBuildData.bNoiseEnabled = true;
-	PhysicsBuildData.NoiseAmplitude = NoiseAmplitude;
+	PhysicsBuildData.NoiseStrength = NoiseStrength;
 	PhysicsBuildData.NoiseFrequency = NoiseFrequency;
-	//PhysicsBuildData.NoiseTexture = NoiseTexture;
-	PhysicsBuildData.NoiseMode = int32(NoiseMode);
-
-	// Build LUT
-	{
-		static int32 NumChannelsToBuild = 64;
-		static int32 ChannelWidth = 64;
-
-		TArray<float> ChannelData;
-		ChannelData.Reserve(NumChannelsToBuild * ChannelWidth);
-
-		for (int32 Channel=0; Channel < NumChannelsToBuild; ++Channel)
-		{
-			FVector3f NoisePosition = FVector3f(FMath::RandRange(-64.0f, 64.0f), FMath::RandRange(-64.0f, 64.0f), FMath::RandRange(-64.0f, 64.0f));
-			FVector3f NoiseTravel = FVector3f(FMath::RandRange(-1.0f, 1.0f), FMath::RandRange(-1.0f, 1.0f), FMath::RandRange(-1.0f, 1.0f));
-
-			float NoiseValue = 0.0f;
-			for (int32 Width=0; Width < ChannelWidth; ++Width)
-			{
-				//NoiseValue += FMath::RandRange(-1.0f, 1.0f);
-
-				FNiagaraMatrix3x4 J = JacobianSimplex_ALU(NoisePosition);
-				NoiseValue += J[1][2] - J[2][1];
-				NoisePosition += NoiseTravel;
-
-				ChannelData.Add(NoiseValue);
-			}
-		}
-
-		PhysicsBuildData.NoiseLUTOffset			= BuildContext.AddStaticData(ChannelData);
-		PhysicsBuildData.NoiseLUTNumChannel		= NumChannelsToBuild;
-		PhysicsBuildData.NoiseLUTChannelWidth	= ChannelWidth;
-	}
 }
diff --git a/Engine/Plugins/FX/Niagara/Source/Niagara/Private/Stateless/Modules/NiagaraStatelessModule_SolveVelocitiesAndForces.cpp b/Engine/Plugins/FX/Niagara/Source/Niagara/Private/Stateless/Modules/NiagaraStatelessModule_SolveVelocitiesAndForces.cpp
--- a/Engine/Plugins/FX/Niagara/Source/Niagara/Private/Stateless/Modules/NiagaraStatelessModule_SolveVelocitiesAndForces.cpp
+++ b/Engine/Plugins/FX/Niagara/Source/Niagara/Private/Stateless/Modules/NiagaraStatelessModule_SolveVelocitiesAndForces.cpp
@@ -1,6 +1,7 @@
 // Copyright Epic Games, Inc. All Rights Reserved.
 
 #include "Stateless/Modules/NiagaraStatelessModule_SolveVelocitiesAndForces.h"
+#include "Stateless/NiagaraStatelessNoiseLUT.h"
 
 #include "VectorField/VectorField.h"
 #include "VectorField/VectorFieldStatic.h"
@@ -31,14 +32,38 @@ namespace NiagaraStatelessModuleSolveVelocitiesAndForcesPrivate
 		return Age * (Velocity + Wind) + 0.5f * Acceleration * Age * Age;
 	}
 
+	static void EvaluateNoise(TConstArrayView<FVector3f> NoiseLUT, const UNiagaraStatelessModule_SolveVelocitiesAndForces::FParameters* Parameters, float NormalizedAge, float Age, FVector3f& Position)
+	{
+		const float fRow = (Position * 0.01f).SquaredLength();
+		const uint32 RowA = uint32(FMath::FloorToInt(FMath::Abs(fRow))) & Parameters->SolveVelocitiesAndForces_NoiseLUTRowMask;
+		const uint32 RowB = (RowA + 1) & Parameters->SolveVelocitiesAndForces_NoiseLUTRowMask;
+
+		const float fColumn = FMath::Abs(Parameters->SolveVelocitiesAndForces_NoiseFrequency) * NormalizedAge;
+		const uint32 iColumnA = FMath::Min(uint32(FMath::FloorToInt(fColumn)), Parameters->SolveVelocitiesAndForces_NoiseLUTRowWidth - 1);
+		const uint32 iColumnB = FMath::Min(iColumnA + 1, Parameters->SolveVelocitiesAndForces_NoiseLUTRowWidth - 1);
+
+		const FVector3f ValueA0 = NoiseLUT[(Parameters->SolveVelocitiesAndForces_NoiseLUTRowWidth * RowA) + iColumnA];
+		const FVector3f ValueA1 = NoiseLUT[(Parameters->SolveVelocitiesAndForces_NoiseLUTRowWidth * RowB) + iColumnA];
+		const FVector3f ValueB0 = NoiseLUT[(Parameters->SolveVelocitiesAndForces_NoiseLUTRowWidth * RowA) + iColumnB];
+		const FVector3f ValueB1 = NoiseLUT[(Parameters->SolveVelocitiesAndForces_NoiseLUTRowWidth * RowB) + iColumnB];
+		const FVector3f ValueA = FMath::Lerp(ValueA0, ValueA1, FMath::Fractional(fRow));
+		const FVector3f ValueB = FMath::Lerp(ValueB0, ValueB1, FMath::Fractional(fRow));
+
+		Position += FMath::Lerp(ValueA, ValueB, FMath::Fractional(fColumn)) * Parameters->SolveVelocitiesAndForces_NoiseStrength * Age * Parameters->SolveVelocitiesAndForces_NoiseLUTNormalize;
+	}
+
 	static void ParticleSimulate(const NiagaraStateless::FParticleSimulationContext& ParticleSimulationContext)
 	{
 		const FModuleBuiltData* ModuleBuiltData = ParticleSimulationContext.ReadBuiltData<FModuleBuiltData>();
 		const UNiagaraStatelessModule_SolveVelocitiesAndForces::FParameters* Parameters = ParticleSimulationContext.ReadParameterNestedStruct<UNiagaraStatelessModule_SolveVelocitiesAndForces::FParameters>();
 
-		const float* AgeData			= ParticleSimulationContext.GetParticleAge();
-		const float* PreviousAgeData	= ParticleSimulationContext.GetParticlePreviousAge();
-		const float InvDeltaTime		= ParticleSimulationContext.GetInvDeltaTime();
+		const float* AgeData					= ParticleSimulationContext.GetParticleAge();
+		const float* PreviousAgeData			= ParticleSimulationContext.GetParticlePreviousAge();
+		const float* NormalizedAgeData			= ParticleSimulationContext.GetParticleNormalizedAge();
+		const float* PreviousNormalizedAgeData	= ParticleSimulationContext.GetParticlePreviousNormalizedAge();
+		const float InvDeltaTime				= ParticleSimulationContext.GetInvDeltaTime();
+
+		TConstArrayView<FVector3f> NoiseLUT = NiagaraStateless::FNoiseLUT::GetGlobalLUT().GetCpuData();
 
 		for (uint32 i = 0; i < ParticleSimulationContext.GetNumInstances(); ++i)
 		{
@@ -54,7 +79,7 @@ namespace NiagaraStatelessModuleSolveVelocitiesAndForcesPrivate
 			FVector3f Position			= ParticleSimulationContext.ReadParticleVariable(ModuleBuiltData->PositionVariableOffset, i, FVector3f::ZeroVector);
 			FVector3f PreviousPosition	= ParticleSimulationContext.ReadParticleVariable(ModuleBuiltData->PreviousPositionVariableOffset, i, FVector3f::ZeroVector);
 
-			if (ModuleBuiltData->PhysicsData.bConeVelocity)
+			if (Parameters->SolveVelocitiesAndForces_ConeVelocityEnabled != 0)
 			{
 				const float ConeAngle = ParticleSimulationContext.RandomScaleBiasFloat(i, 5, Parameters->SolveVelocitiesAndForces_ConeAngleScale, Parameters->SolveVelocitiesAndForces_ConeAngleBias);
 				const float ConeRotation = ParticleSimulationContext.RandomFloat(i, 6) * UE_TWO_PI;
@@ -63,7 +88,7 @@ namespace NiagaraStatelessModuleSolveVelocitiesAndForcesPrivate
 				const FVector3f Direction = FVector3f(scRot.X * scAng.X, scRot.Y * scAng.X, scAng.Y);
 
 				float VelocityScale = ParticleSimulationContext.RandomScaleBiasFloat(i, 7, Parameters->SolveVelocitiesAndForces_ConeVelocityScale, Parameters->SolveVelocitiesAndForces_ConeVelocityBias);
-				if (ModuleBuiltData->PhysicsData.bConeVelocity)
+				if (Parameters->SolveVelocitiesAndForces_ConeVelocityFalloff > 0.0f)
 				{
 					const float pf = FMath::Pow(FMath::Clamp(scAng.Y, 0.0f, 1.0f), Parameters->SolveVelocitiesAndForces_ConeVelocityFalloff * 10.0f);
 					VelocityScale *= FMath::Lerp(1.0f, pf, Parameters->SolveVelocitiesAndForces_ConeVelocityFalloff);
@@ -72,7 +97,7 @@ namespace NiagaraStatelessModuleSolveVelocitiesAndForcesPrivate
 				InitialVelocity += ModuleBuiltData->PhysicsData.ConeQuat.RotateVector(Direction) * VelocityScale;
 			}
 
-			if (ModuleBuiltData->PhysicsData.bPointVelocity)
+			if (Parameters->SolveVelocitiesAndForces_PointVelocityEnabled != 0)
 			{
 				const FVector3f FallbackDir	= ParticleSimulationContext.RandomUnitFloat3(i, 8);
 				const FVector3f Delta		= Position - Parameters->SolveVelocitiesAndForces_PointOrigin;
@@ -82,9 +107,10 @@ namespace NiagaraStatelessModuleSolveVelocitiesAndForcesPrivate
 				InitialVelocity += Direction * VelocityScale;
 			}
 
-			if (ModuleBuiltData->PhysicsData.bNoiseEnabled)
+			if (Parameters->SolveVelocitiesAndForces_NoiseEnabled != 0)
 			{
-				//-TODO:
+				EvaluateNoise(NoiseLUT, Parameters, NormalizedAgeData[i], AgeData[i], Position);
+				EvaluateNoise(NoiseLUT, Parameters, PreviousNormalizedAgeData[i], PreviousAgeData[i], PreviousPosition);
 			}
 
 			Position += IntegratePosition(FMath::Max(AgeData[i], UE_KINDA_SMALL_NUMBER), Drag, InitialVelocity, Wind, Acceleration);
@@ -177,36 +203,22 @@ void UNiagaraStatelessModule_SolveVelocitiesAndForces::SetShaderParameters(const
 		Parameters->SolveVelocitiesAndForces_PointOrigin = SpaceTransforms.TransformPosition(PhysicsData.PointCoordinateSpace, PhysicsData.PointOrigin);
 	}
 
-	Parameters->SolveVelocitiesAndForces_NoiseEnabled = PhysicsData.bNoiseEnabled ? 1 : 0;
-	if (PhysicsData.bNoiseEnabled)
 	{
-		Parameters->SolveVelocitiesAndForces_NoiseAmplitude			= PhysicsData.NoiseAmplitude;
-		Parameters->SolveVelocitiesAndForces_NoiseFrequency			= FVector3f(PhysicsData.NoiseFrequency, PhysicsData.NoiseFrequency, PhysicsData.NoiseFrequency);
-		//SetShaderParameterContext.SetTextureResource(&Parameters->SolveVelocitiesAndForces_NoiseTexture, ModuleBuildData->NoiseTexture);
-		Parameters->SolveVelocitiesAndForces_NoiseMode				= PhysicsData.NoiseMode;
-		Parameters->SolveVelocitiesAndForces_NoiseLUTOffset			= PhysicsData.NoiseLUTOffset;
-		Parameters->SolveVelocitiesAndForces_NoiseLUTNumChannel		= PhysicsData.NoiseLUTNumChannel;
-		Parameters->SolveVelocitiesAndForces_NoiseLUTChannelWidth	= PhysicsData.NoiseLUTChannelWidth;
+		const NiagaraStateless::FNoiseLUT& NoiseLUT = NiagaraStateless::FNoiseLUT::GetGlobalLUT();
+		checkf(FMath::IsPowerOfTwo(NoiseLUT.GetNumRows()), TEXT("Noise LUT must be PoT"));
+
+		const float fLUTSizeMinusOne = float(NoiseLUT.GetRowWidth() - 1);
+		const float ClampedFrequency = FMath::Min(FMath::Abs(PhysicsData.NoiseFrequency / 2.0f) + 1.0f, fLUTSizeMinusOne);
+		const float Frequency = PhysicsData.NoiseFrequency >= 0.0f ? ClampedFrequency : -ClampedFrequency;
+
+		Parameters->SolveVelocitiesAndForces_NoiseEnabled		= PhysicsData.bNoiseEnabled ? 1 : 0;
+		Parameters->SolveVelocitiesAndForces_NoiseStrength		= PhysicsData.NoiseStrength;
+		Parameters->SolveVelocitiesAndForces_NoiseFrequency		= Frequency;
+		Parameters->SolveVelocitiesAndForces_NoiseLUTNormalize	= float(NoiseLUT.GetRowWidth()) / Frequency;
+		Parameters->SolveVelocitiesAndForces_NoiseLUTRowMask	= NoiseLUT.GetNumRows() - 1;
+		Parameters->SolveVelocitiesAndForces_NoiseLUTRowWidth	= NoiseLUT.GetRowWidth();
+		Parameters->SolveVelocitiesAndForces_NoiseLUT			= NoiseLUT.GetGpuData();
 	}
-
-	//FVectorFieldTextureAccessor TextureAccessor(Cast<UVectorField>(PhysicsData.NoiseTexture));
-
-	//ENQUEUE_RENDER_COMMAND(FUpdateNoiseTexture)(
-	//	[Parameters, TextureAccessor](FRHICommandListImmediate& RHICmdList)
-	//	{
-	//		FRHITexture* NoiseTextureRHI = TextureAccessor.GetTexture();
-	//		Parameters->SolveVelocitiesAndForces_NoiseTexture = NoiseTextureRHI;
-	//		Parameters->SolveVelocitiesAndForces_NoiseSampler = TStaticSamplerState<SF_Bilinear, AM_Wrap, AM_Wrap, AM_Wrap>::GetRHI();
-
-	//		if (Parameters->SolveVelocitiesAndForces_NoiseMode < 2)
-	//		{
-	//			const FIntVector TextureSize = NoiseTextureRHI ? NoiseTextureRHI->GetSizeXYZ() : FIntVector(1, 1, 1);
-	//			Parameters->SolveVelocitiesAndForces_NoiseFrequency.X *= 1.0f / float(TextureSize.X);
-	//			Parameters->SolveVelocitiesAndForces_NoiseFrequency.Y *= 1.0f / float(TextureSize.Y);
-	//			Parameters->SolveVelocitiesAndForces_NoiseFrequency.Z *= 1.0f / float(TextureSize.Z);
-	//		}
-	//	}
-	//);
 }
 
 #if WITH_EDITORONLY_DATA
diff --git a/Engine/Plugins/FX/Niagara/Source/Niagara/Private/Stateless/NiagaraStatelessNoiseLUT.cpp b/Engine/Plugins/FX/Niagara/Source/Niagara/Private/Stateless/NiagaraStatelessNoiseLUT.cpp
new file mode 100644
--- /dev/null
+++ b/Engine/Plugins/FX/Niagara/Source/Niagara/Private/Stateless/NiagaraStatelessNoiseLUT.cpp
@@ -0,0 +1,53 @@
+// Copyright Epic Games, Inc. All Rights Reserved.
+
+#include "Stateless/NiagaraStatelessNoiseLUT.h"
+#include "NiagaraSimplexNoise.h"
+
+namespace NiagaraStateless
+{
+	static TGlobalResource<NiagaraStateless::FNoiseLUT> GGlobalLUT(32, 32, 0.1f);
+
+	FNoiseLUT::FNoiseLUT(uint32 Rows, uint32 RowWidth, float FieldTravelSpeed)
+		: LUTRows(Rows)
+		, LUTRowWidth(RowWidth)
+	{
+		CpuData.AddUninitialized(LUTRows * LUTRowWidth);
+		for (uint32 Row = 0; Row < LUTRows; ++Row)
+		{
+			FVector3f Position = FVector3f(Row * 1024.0f, Row * 511.0f, Row * 30.1f);
+			FVector3f Sum = FVector3f::ZeroVector;
+
+			uint32 RowOutputOffset = Row * LUTRowWidth;
+			CpuData[RowOutputOffset++] = Sum;
+			for (uint32 Column = 1; Column < LUTRowWidth; ++Column)
+			{
+				const FNiagaraMatrix3x4 J = JacobianSimplex_ALU(Position);
+				const FVector3f Dir = FVector3f(J[1][2] - J[2][1], J[2][0] - J[0][2], J[0][1] - J[1][0]).GetSafeNormal();
+
+				Sum += Dir * float(1.0f / float(LUTRowWidth));
+				Position += Dir * FieldTravelSpeed;
+				CpuData[RowOutputOffset++] = Sum;
+			}
+		}
+	}
+
+	void FNoiseLUT::InitRHI(FRHICommandListBase& RHICmdList)
+	{
+		const uint32 BufferByteSize = CpuData.Num() * CpuData.GetTypeSize();
+		GpuData.Initialize(RHICmdList, TEXT("NiagaraStateless::FNoiseLUT"), BufferByteSize, EBufferUsageFlags::Static);
+
+		void* UploadMemory = RHICmdList.LockBuffer(GpuData.Buffer, 0, BufferByteSize, RLM_WriteOnly);
+		FMemory::Memcpy(UploadMemory, CpuData.GetData(), BufferByteSize);
+		RHICmdList.UnlockBuffer(GpuData.Buffer);
+	}
+
+	void FNoiseLUT::ReleaseRHI()
+	{
+		GpuData.Release();
+	}
+
+	const FNoiseLUT& FNoiseLUT::GetGlobalLUT()
+	{
+		return GGlobalLUT;
+	}
+}
diff --git a/Engine/Plugins/FX/Niagara/Source/Niagara/Private/Stateless/NiagaraStatelessNoiseLUT.h b/Engine/Plugins/FX/Niagara/Source/Niagara/Private/Stateless/NiagaraStatelessNoiseLUT.h
new file mode 100644
--- /dev/null
+++ b/Engine/Plugins/FX/Niagara/Source/Niagara/Private/Stateless/NiagaraStatelessNoiseLUT.h
@@ -0,0 +1,33 @@
+// Copyright Epic Games, Inc. All Rights Reserved.
+
+#pragma once
+
+#include "Stateless/NiagaraStatelessCommon.h"
+#include "RenderResource.h"
+#include "RHIUtilities.h"
+
+namespace NiagaraStateless
+{
+	class FNoiseLUT : public FRenderResource
+	{
+	public:
+		explicit FNoiseLUT(uint32 Rows, uint32 RowWidth, float FieldTravelSpeed);
+
+		virtual void InitRHI(FRHICommandListBase& RHICmdList) override;
+		virtual void ReleaseRHI() override;
+
+		uint32 GetNumRows() const { return LUTRows; }
+		uint32 GetRowWidth() const { return LUTRowWidth; }
+		TConstArrayView<FVector3f> GetCpuData() const { return CpuData; }
+		FRHIShaderResourceView* GetGpuData() const { return GpuData.SRV; }
+
+		static const FNoiseLUT& GetGlobalLUT();
+
+	private:
+		uint32				LUTRows = 0;
+		uint32				LUTRowWidth = 0;
+
+		TArray<FVector3f>	CpuData;
+		FByteAddressBuffer	GpuData;
+	};
+}
diff --git a/Engine/Plugins/FX/Niagara/Source/NiagaraShader/Internal/Stateless/NiagaraStatelessModuleShaderParameters.h b/Engine/Plugins/FX/Niagara/Source/NiagaraShader/Internal/Stateless/NiagaraStatelessModuleShaderParameters.h
--- a/Engine/Plugins/FX/Niagara/Source/NiagaraShader/Internal/Stateless/NiagaraStatelessModuleShaderParameters.h
+++ b/Engine/Plugins/FX/Niagara/Source/NiagaraShader/Internal/Stateless/NiagaraStatelessModuleShaderParameters.h
@@ -172,16 +172,12 @@ namespace NiagaraStateless
 		SHADER_PARAMETER(float,			SolveVelocitiesAndForces_PointVelocityBias)
 		SHADER_PARAMETER(FVector3f,		SolveVelocitiesAndForces_PointOrigin)
 
-		SHADER_PARAMETER(uint32,		SolveVelocitiesAndForces_NoiseEnabled)
-		SHADER_PARAMETER(float,			SolveVelocitiesAndForces_NoiseAmplitude)
-		SHADER_PARAMETER(FVector3f,		SolveVelocitiesAndForces_NoiseFrequency)
-		//SHADER_PARAMETER_TEXTURE(Texture3D,		SolveVelocitiesAndForces_NoiseTexture)
-		//SHADER_PARAMETER_SAMPLER(SamplerState,	SolveVelocitiesAndForces_NoiseSampler)
-
-		SHADER_PARAMETER(uint32, SolveVelocitiesAndForces_NoiseMode)
-		SHADER_PARAMETER(uint32, SolveVelocitiesAndForces_NoiseLUTOffset)
-		SHADER_PARAMETER(uint32, SolveVelocitiesAndForces_NoiseLUTNumChannel)
-		SHADER_PARAMETER(uint32, SolveVelocitiesAndForces_NoiseLUTChannelWidth)
-
+		SHADER_PARAMETER(uint32,				SolveVelocitiesAndForces_NoiseEnabled)
+		SHADER_PARAMETER(float,					SolveVelocitiesAndForces_NoiseStrength)
+		SHADER_PARAMETER(float,					SolveVelocitiesAndForces_NoiseFrequency)
+		SHADER_PARAMETER(float,					SolveVelocitiesAndForces_NoiseLUTNormalize)
+		SHADER_PARAMETER(uint32,				SolveVelocitiesAndForces_NoiseLUTRowMask)
+		SHADER_PARAMETER(uint32,				SolveVelocitiesAndForces_NoiseLUTRowWidth)
+		SHADER_PARAMETER_SRV(ByteAddressBuffer,	SolveVelocitiesAndForces_NoiseLUT)
 	END_SHADER_PARAMETER_STRUCT()
 }
