Remove vestigial system sim code.

diff --git a/Engine/Plugins/FX/Niagara/Source/Niagara/Private/NiagaraSystemSimulation.cpp b/Engine/Plugins/FX/Niagara/Source/Niagara/Private/NiagaraSystemSimulation.cpp
--- a/Engine/Plugins/FX/Niagara/Source/Niagara/Private/NiagaraSystemSimulation.cpp
+++ b/Engine/Plugins/FX/Niagara/Source/Niagara/Private/NiagaraSystemSimulation.cpp
@@ -12,7 +12,6 @@
 #include "NiagaraSystemGpuComputeProxy.h"
 #include "NiagaraConstants.h"
 #include "NiagaraStats.h"
-#include "Async/ParallelFor.h"
 #include "NiagaraComponent.h"
 #include "NiagaraWorldManager.h"
 #include "NiagaraGpuComputeDispatchInterface.h"
@@ -1917,8 +1916,6 @@ void FNiagaraSystemSimulation::PrepareForSystemSimulate(FNiagaraSystemSimulation
 	UpdateInstanceParameterDataSet.Allocate(NumInstances);
 
 	check(System != nullptr);
-	TConstArrayView<FNiagaraDataSetAccessor<ENiagaraExecutionState>> EmitterExecutionStateAccessors = System->GetEmitterExecutionStateAccessors();
-
 	//Tick instance parameters and transfer any needed into the system simulation dataset.
 	auto TransferInstanceParameters = [&](int32 SystemIndex)
 	{
@@ -1935,20 +1932,12 @@ void FNiagaraSystemSimulation::PrepareForSystemSimulate(FNiagaraSystemSimulation
 
 		//TODO: Find good way to check that we're not using any instance parameter data interfaces in the system scripts here.
 		//In that case we need to solo and will never get here.
-
-		TArrayView<FNiagaraEmitterInstanceRef> Emitters = Inst->GetEmitters();
-		for (int32 EmitterIdx = 0; EmitterIdx < Emitters.Num(); ++EmitterIdx)
-		{
-			FNiagaraEmitterInstance& EmitterInst = Emitters[EmitterIdx].Get();
-			if ( EmitterExecutionStateAccessors.IsValidIndex(EmitterIdx) )
-			{
-				EmitterExecutionStateAccessors[EmitterIdx].GetWriter(Context.DataSet).SetSafe(SystemIndex, EmitterInst.GetExecutionState());
-			}
-		}
 	};
 
-	//This can go wide if we have a very large number of instances.
-	ParallelFor(Context.Instances.Num(), TransferInstanceParameters, true);
+	for (int32 i = 0; i < Context.Instances.Num(); ++i)
+	{
+		TransferInstanceParameters(i);
+	}
 
 	SpawnInstanceParameterDataSet.GetDestinationDataChecked().SetNumInstances(NumInstances);
 	UpdateInstanceParameterDataSet.GetDestinationDataChecked().SetNumInstances(NumInstances);
