Fix issue writing Enums to NDC.

diff --git a/Engine/Plugins/FX/Niagara/Source/Niagara/Private/NiagaraDataChannelAccessor.cpp b/Engine/Plugins/FX/Niagara/Source/Niagara/Private/NiagaraDataChannelAccessor.cpp
--- a/Engine/Plugins/FX/Niagara/Source/Niagara/Private/NiagaraDataChannelAccessor.cpp
+++ b/Engine/Plugins/FX/Niagara/Source/Niagara/Private/NiagaraDataChannelAccessor.cpp
@@ -232,7 +232,7 @@ void UNiagaraDataChannelWriter::WriteInt(FName VarName, int32 Index, int32 InDat
 
 void UNiagaraDataChannelWriter::WriteEnum(FName VarName, int32 Index, uint8 InData)
 {
-	WriteInt(VarName, StartIndex + Index, InData);
+	WriteInt(VarName, Index, InData);
 }
 
 void UNiagaraDataChannelWriter::WriteBool(FName VarName, int32 Index, bool InData)
