- Add Create Asset from transient sim caches (i.e. ones that were captured with debugging in mind).
- Add option to remove Debug Data from an existing cache

diff --git a/Engine/Plugins/FX/Niagara/Source/Niagara/Classes/NiagaraSimCache.h b/Engine/Plugins/FX/Niagara/Source/Niagara/Classes/NiagaraSimCache.h
--- a/Engine/Plugins/FX/Niagara/Source/Niagara/Classes/NiagaraSimCache.h
+++ b/Engine/Plugins/FX/Niagara/Source/Niagara/Classes/NiagaraSimCache.h
@@ -489,6 +489,11 @@ public:
 	/** Returns the debug data stored inside the cache, intended for internal use only. */
 	NIAGARA_API const UNiagaraSimCacheDebugData* GetDebugData() const { return DebugData; }
 
+#if WITH_EDITORONLY_DATA
+	/** Remove the debug data from the cache */
+	NIAGARA_API void RemoveDebugData();
+#endif
+
 	/**
 	Get number of active instances for the emitter at the given frame.
 	An EmitterIndex or INDEX_NONE will return information about the system instance.
diff --git a/Engine/Plugins/FX/Niagara/Source/Niagara/Private/NiagaraSimCache.cpp b/Engine/Plugins/FX/Niagara/Source/Niagara/Private/NiagaraSimCache.cpp
--- a/Engine/Plugins/FX/Niagara/Source/Niagara/Private/NiagaraSimCache.cpp
+++ b/Engine/Plugins/FX/Niagara/Source/Niagara/Private/NiagaraSimCache.cpp
@@ -21,8 +21,14 @@
 #include "NiagaraSystemInstance.h"
 #include "Stateless/NiagaraStatelessEmitterInstance.h"
 
+#if WITH_EDITORONLY_DATA
+#include "ScopedTransaction.h"
+#endif
+
 #include UE_INLINE_GENERATED_CPP_BY_NAME(NiagaraSimCache)
 
+#define LOCTEXT_NAMESPACE "NiagaraSimCache"
+
 UNiagaraSimCache::FOnCacheBeginWrite	UNiagaraSimCache::OnCacheBeginWrite;
 UNiagaraSimCache::FOnCacheEndWrite		UNiagaraSimCache::OnCacheEndWrite;
 
@@ -1239,6 +1245,26 @@ const UObject* UNiagaraSimCache::GetDataInterfaceStorageObject(const FNiagaraVar
 	return nullptr;
 }
 
+#if WITH_EDITORONLY_DATA
+void UNiagaraSimCache::RemoveDebugData()
+{
+	if (DebugData == nullptr)
+	{
+		return;
+	}
+
+	const FScopedTransaction Transaction(LOCTEXT("RemoveDebugData", "Sim Cache Remove Debug Data"));
+
+	Modify();
+
+	OnCacheBeginWrite.Broadcast(this);
+
+	DebugData = nullptr;
+
+	OnCacheEndWrite.Broadcast(this);
+}
+#endif
+
 int UNiagaraSimCache::GetEmitterNumInstances(int32 EmitterIndex, int32 FrameIndex) const
 {
 	if ( CacheFrames.IsValidIndex(FrameIndex) )
@@ -1497,3 +1523,5 @@ UObject* UNiagaraSimCache::ReadDataInterfaceAs(UClass* RequestedType, FName Attr
 	}
 	return nullptr;
 }
+
+#undef LOCTEXT_NAMESPACE
diff --git a/Engine/Plugins/FX/Niagara/Source/NiagaraEditor/Private/Toolkits/NiagaraSimCacheToolkit.cpp b/Engine/Plugins/FX/Niagara/Source/NiagaraEditor/Private/Toolkits/NiagaraSimCacheToolkit.cpp
--- a/Engine/Plugins/FX/Niagara/Source/NiagaraEditor/Private/Toolkits/NiagaraSimCacheToolkit.cpp
+++ b/Engine/Plugins/FX/Niagara/Source/NiagaraEditor/Private/Toolkits/NiagaraSimCacheToolkit.cpp
@@ -236,6 +236,8 @@ void FNiagaraSimCacheToolkit::ExtendToolbar()
 	{
 		static void FillToolbar(FToolBarBuilder& ToolbarBuilder, FNiagaraSimCacheToolkit* Toolkit)
 		{
+			TSharedRef<FNiagaraSimCacheViewModel> SimCacheViewModelRef = Toolkit->SimCacheViewModel.ToSharedRef();
+
 			if (Toolkit->SimCache.IsValid() && Toolkit->SimCache->GetSystem())
 			{
 				ToolbarBuilder.BeginSection("System");
@@ -265,8 +267,6 @@ void FNiagaraSimCacheToolkit::ExtendToolbar()
 			ToolbarBuilder.EndSection();
 			ToolbarBuilder.BeginSection("Export");
 			{
-				TSharedRef<FNiagaraSimCacheViewModel> SimCacheViewModelRef = Toolkit->SimCacheViewModel.ToSharedRef();
-
 				ToolbarBuilder.AddToolBarButton(
 					FUIAction(
 						FExecuteAction::CreateSP(SimCacheViewModelRef, &FNiagaraSimCacheViewModel::CopyActiveToClipboard),
@@ -288,6 +288,33 @@ void FNiagaraSimCacheToolkit::ExtendToolbar()
 					LOCTEXT("ExportToJsonTooltip", "Exports the raw data for each frame to disk. Note that data from data interfaces is only exported if they implement support for it."),
 					FSlateIcon()//-TODO:FAppStyle::GetAppStyleSetName(), "GenericCommands.Copy")
 				);
+
+				ToolbarBuilder.AddToolBarButton(
+					FUIAction(
+						FExecuteAction::CreateSP(SimCacheViewModelRef, &FNiagaraSimCacheViewModel::CreateAsset),
+						FCanExecuteAction(),
+						FIsActionChecked(),
+						FIsActionButtonVisible::CreateSP(SimCacheViewModelRef, &FNiagaraSimCacheViewModel::IsCreateAssetVisible)
+					),
+					NAME_None,
+					LOCTEXT("CreateAsset", "Create Asset"),
+					LOCTEXT("CreateAssetTooltip", "Create an asset from this debug (transient) cache.")
+				);
+			}
+			ToolbarBuilder.EndSection();
+			ToolbarBuilder.BeginSection("Modify");
+			{
+				ToolbarBuilder.AddToolBarButton(
+					FUIAction(
+						FExecuteAction::CreateSP(SimCacheViewModelRef, &FNiagaraSimCacheViewModel::RemoveDebugData),
+						FCanExecuteAction(),
+						FIsActionChecked(),
+						FIsActionButtonVisible::CreateSP(SimCacheViewModelRef, &FNiagaraSimCacheViewModel::CanRemoveDebugData)
+					),
+					NAME_None,
+					LOCTEXT("RemoveDebugData", "Remove Debug Data"),
+					LOCTEXT("RemoveDebugDataTooltip", "Remove the debug data that is present in the cache, this will reduce the file size.")
+				);
 			}
 			ToolbarBuilder.EndSection();
 		}
diff --git a/Engine/Plugins/FX/Niagara/Source/NiagaraEditor/Private/ViewModels/NiagaraSimCacheViewModel.cpp b/Engine/Plugins/FX/Niagara/Source/NiagaraEditor/Private/ViewModels/NiagaraSimCacheViewModel.cpp
--- a/Engine/Plugins/FX/Niagara/Source/NiagaraEditor/Private/ViewModels/NiagaraSimCacheViewModel.cpp
+++ b/Engine/Plugins/FX/Niagara/Source/NiagaraEditor/Private/ViewModels/NiagaraSimCacheViewModel.cpp
@@ -6,10 +6,14 @@
 #include "NiagaraEditorCommon.h"
 #include "NiagaraEditorSimCacheUtils.h"
 #include "NiagaraDebuggerCommon.h"
+
 #include "HAL/PlatformApplicationMisc.h"
 #include "Serialization/MemoryReader.h"
+#include "Subsystems/AssetEditorSubsystem.h"
 #include "Widgets/SNiagaraSimCacheTreeView.h"
 #include "UObject/Package.h"
+#include "Editor.h"
+#include "FileHelpers.h"
 
 #define LOCTEXT_NAMESPACE "NiagaraSimCacheViewModel"
 
@@ -826,6 +830,54 @@ void FNiagaraSimCacheViewModel::ExportToDisk() const
 	FNiagaraEditorSimCacheUtils::ExportToDisk(SimCache.Get());
 }
 
+bool FNiagaraSimCacheViewModel::IsCreateAssetVisible() const
+{
+	if (SimCache)
+	{
+		return SimCache->IsAsset() == false;
+	}
+	return false;
+}
+
+void FNiagaraSimCacheViewModel::CreateAsset()
+{
+	if (SimCache == nullptr)
+	{
+		return;
+	}
+
+	SimCache->SetFlags(RF_Public | RF_Standalone);
+
+	TArray<UObject*> AssetsToSave = { SimCache };
+	TArray<UObject*> SavedAssets;
+	FEditorFileUtils::SaveAssetsAs(AssetsToSave, SavedAssets);
+
+	if (SavedAssets.Num() == 0 || SavedAssets[0] == nullptr || SavedAssets[0] == AssetsToSave[0])
+	{
+		SimCache->ClearFlags(RF_Public | RF_Standalone);
+		return;
+	}
+
+	UAssetEditorSubsystem* AssetEditorSubsystem = GEditor->GetEditorSubsystem<UAssetEditorSubsystem>();
+	AssetEditorSubsystem->CloseAllEditorsForAsset(AssetsToSave[0]);
+	AssetEditorSubsystem->OpenEditorForAssets_Advanced(SavedAssets, EToolkitMode::Standalone);
+}
+
+bool FNiagaraSimCacheViewModel::CanRemoveDebugData() const
+{
+	return SimCache && SimCache->GetDebugData() != nullptr;
+}
+
+void FNiagaraSimCacheViewModel::RemoveDebugData()
+{
+	if (SimCache == nullptr)
+	{
+		return;
+	}
+
+	SimCache->RemoveDebugData();
+}
+
 void FNiagaraSimCacheViewModel::AddReferencedObjects(FReferenceCollector& Collector)
 {
 	if(SimCache != nullptr)
diff --git a/Engine/Plugins/FX/Niagara/Source/NiagaraEditor/Public/ViewModels/NiagaraSimCacheViewModel.h b/Engine/Plugins/FX/Niagara/Source/NiagaraEditor/Public/ViewModels/NiagaraSimCacheViewModel.h
--- a/Engine/Plugins/FX/Niagara/Source/NiagaraEditor/Public/ViewModels/NiagaraSimCacheViewModel.h
+++ b/Engine/Plugins/FX/Niagara/Source/NiagaraEditor/Public/ViewModels/NiagaraSimCacheViewModel.h
@@ -102,6 +102,12 @@ public:
 	bool CanExportToDisk() const;
 	void ExportToDisk() const;
 
+	bool IsCreateAssetVisible() const;
+	void CreateAsset();
+
+	bool CanRemoveDebugData() const;
+	void RemoveDebugData();
+
 	UNiagaraComponent* GetPreviewComponent() { return PreviewComponent; }
 
 	//~ FGCObject interface
