- Fix narrowing warnings in EvaluateDerivativePosition (float4 -> float3)

#tests Run Spline ET
#rnx

diff --git a/Engine/Plugins/FX/Niagara/Shaders/Private/NiagaraDataInterfaceSplineTemplate.ush b/Engine/Plugins/FX/Niagara/Shaders/Private/NiagaraDataInterfaceSplineTemplate.ush
--- a/Engine/Plugins/FX/Niagara/Shaders/Private/NiagaraDataInterfaceSplineTemplate.ush
+++ b/Engine/Plugins/FX/Niagara/Shaders/Private/NiagaraDataInterfaceSplineTemplate.ush
@@ -109,7 +109,7 @@ float3 EvaluateDerivativePosition_{ParameterName}(float InDistance)
 			return float3(0,0,0);
 		}
 	}
-	return {ParameterName}_SplinePositionsLUT[NextKey] - {ParameterName}_SplinePositionsLUT[PrevKey];
+	return ({ParameterName}_SplinePositionsLUT[NextKey] - {ParameterName}_SplinePositionsLUT[PrevKey]).xyz;
 }
 
 float EvaluateFindNearestPosition_{ParameterName}(float3 InPosition)
