- Data Table DI updates
- Fix for reading bools on the GPU and for packing aligned types on the CPU
- Add support for auto populating the read row nodes from an existing data table

#rnx

diff --git a/Engine/Plugins/FX/Niagara/Shaders/Private/NiagaraDataInterfaceDataTableTemplate.ush b/Engine/Plugins/FX/Niagara/Shaders/Private/NiagaraDataInterfaceDataTableTemplate.ush
--- a/Engine/Plugins/FX/Niagara/Shaders/Private/NiagaraDataInterfaceDataTableTemplate.ush
+++ b/Engine/Plugins/FX/Niagara/Shaders/Private/NiagaraDataInterfaceDataTableTemplate.ush
@@ -1,9 +1,9 @@
 // Copyright Epic Games, Inc. All Rights Reserved.
 
-uint				{ParameterName}_NumRows;
+int					{ParameterName}_NumRows;
 uint				{ParameterName}_RowStride;
 uint				{ParameterName}_InvalidRowReadOffset;
-uint				{ParameterName}_NumFilteredRows;
+int					{ParameterName}_NumFilteredRows;
 ByteAddressBuffer	{ParameterName}_TableDataBuffer;
 Buffer<uint>		{ParameterName}_FilteredRowDataOffsetBuffer;
 //uint4				{ParameterName}_AttributeReadOffset[];
@@ -14,6 +14,7 @@ uint GetReadOffset_{ParameterName}(uint RowOffset, uint Column)
 	return ColumnOffset == 0xffffffff ? {ParameterName}_InvalidRowReadOffset : RowOffset + ColumnOffset;
 }
 
+void ReadValue_{ParameterName}(uint RowOffset, uint Column, out bool   Value) {	Value = asint({ParameterName}_TableDataBuffer.Load(GetReadOffset_{ParameterName}(RowOffset, Column))) != 0; }
 void ReadValue_{ParameterName}(uint RowOffset, uint Column, out int    Value) {	Value = asint({ParameterName}_TableDataBuffer.Load(GetReadOffset_{ParameterName}(RowOffset, Column))); }
 void ReadValue_{ParameterName}(uint RowOffset, uint Column, out float  Value) { Value = asfloat({ParameterName}_TableDataBuffer.Load(GetReadOffset_{ParameterName}(RowOffset, Column))); }
 void ReadValue_{ParameterName}(uint RowOffset, uint Column, out float2 Value) { Value = asfloat({ParameterName}_TableDataBuffer.Load2(GetReadOffset_{ParameterName}(RowOffset, Column))); }
@@ -31,7 +32,7 @@ uint GetFilteredRowReadOffset_{ParameterName}(int RowIndex, out bool bIsValid)
 uint GetRowReadOffset_{ParameterName}(int RowIndex, out bool bIsValid)
 {
 	bIsValid = RowIndex >= 0 && RowIndex < {ParameterName}_NumRows;
-	return bIsValid ? RowIndex * {ParameterName}_RowStride : {ParameterName}_InvalidRowReadOffset;
+	return bIsValid ? uint(RowIndex) * {ParameterName}_RowStride : {ParameterName}_InvalidRowReadOffset;
 }
 
 void IsValid_{ParameterName}(out bool bIsValid)
diff --git a/Engine/Plugins/FX/Niagara/Source/Niagara/Internal/DataInterface/NiagaraDataInterfaceDataTable.h b/Engine/Plugins/FX/Niagara/Source/Niagara/Internal/DataInterface/NiagaraDataInterfaceDataTable.h
--- a/Engine/Plugins/FX/Niagara/Source/Niagara/Internal/DataInterface/NiagaraDataInterfaceDataTable.h
+++ b/Engine/Plugins/FX/Niagara/Source/Niagara/Internal/DataInterface/NiagaraDataInterfaceDataTable.h
@@ -53,6 +53,11 @@ public:
 	NIAGARA_API virtual void SetShaderParameters(const FNiagaraDataInterfaceSetShaderParametersContext& Context) const override;
 	// UNiagaraDataInterface Interface End
 
+#if WITH_EDITOR
+	NIAGARA_API static bool IsReadFunction(const FNiagaraFunctionSignature& Signature);
+	NIAGARA_API static TArray<FNiagaraVariableBase> GetVariablesFromDataTable(const UDataTable* DataTable);
+#endif
+
 private:
 	/** Default DataTable to use can be overriden using the parameter binding. */
 	UPROPERTY(EditAnywhere, Category = "DataTable")
diff --git a/Engine/Plugins/FX/Niagara/Source/Niagara/Private/DataInterface/NiagaraDataInterfaceDataTable.cpp b/Engine/Plugins/FX/Niagara/Source/Niagara/Private/DataInterface/NiagaraDataInterfaceDataTable.cpp
--- a/Engine/Plugins/FX/Niagara/Source/Niagara/Private/DataInterface/NiagaraDataInterfaceDataTable.cpp
+++ b/Engine/Plugins/FX/Niagara/Source/Niagara/Private/DataInterface/NiagaraDataInterfaceDataTable.cpp
@@ -25,10 +25,10 @@
 namespace NDIDataTableLocal
 {
 	BEGIN_SHADER_PARAMETER_STRUCT(FShaderParameters, )
-		SHADER_PARAMETER(uint32,				NumRows)
+		SHADER_PARAMETER(int32,					NumRows)
 		SHADER_PARAMETER(uint32,				RowStride)
 		SHADER_PARAMETER(uint32,				InvalidRowReadOffset)
-		SHADER_PARAMETER(uint32,				NumFilteredRows)
+		SHADER_PARAMETER(int32,					NumFilteredRows)
 		SHADER_PARAMETER_SRV(ByteAddressBuffer,	TableDataBuffer)
 		SHADER_PARAMETER_SRV(Buffer<uint>,		FilteredRowDataOffsetBuffer)
 	END_SHADER_PARAMETER_STRUCT()
@@ -146,7 +146,15 @@ namespace NDIDataTableLocal
 		{
 			const FName PropertyName(*RowStruct->GetAuthoredNameForField(Property));
 			Columns.Emplace(FNiagaraVariableBase(TypeDef, PropertyName), RowStride);
-			RowStructConversionFuncs.Emplace([SrcOffset = Property->GetOffset_ForInternal(), DestOffset = RowStride](const uint8* Src, uint8* Dest) { *reinterpret_cast<TTypeTo*>(Dest + DestOffset) = TTypeTo(*reinterpret_cast<const TTypeFrom*>(Src + SrcOffset)); });
+			RowStructConversionFuncs.Emplace(
+				[SrcOffset = Property->GetOffset_ForInternal(), DestOffset = RowStride](const uint8* Src, uint8* Dest)
+				{
+					TTypeFrom FromValue;
+					FMemory::Memcpy(&FromValue, Src + SrcOffset, sizeof(TTypeFrom));
+					const TTypeTo ToValue = TTypeTo(FromValue);
+					FMemory::Memcpy(Dest + DestOffset, &ToValue, sizeof(TTypeTo));
+				}
+			);
 			RowStride += sizeof(TTypeTo);
 		}
 
@@ -623,14 +631,6 @@ void UNiagaraDataInterfaceDataTable::GetFunctionsInternal(TArray<FNiagaraFunctio
 	DefaultSig.bRequiresContext = false;
 	DefaultSig.Inputs.Emplace(FNiagaraTypeDefinition(GetClass()), TEXT("DataTable"));
 
-	//{
-	//	FNiagaraFunctionSignature& Sig = OutFunctions.Emplace_GetRef(DefaultSig);
-	//	Sig.Name = NAME_GetRow;
-	//	Sig.FunctionSpecifiers.Add(TEXT("RowName"));
-	//	Sig.Outputs.Emplace(FNiagaraTypeDefinition::GetIntDef(), TEXT("Value"));
-	//	//RequiredOutputs
-	//}
-
 	{
 		FNiagaraFunctionSignature& Sig = OutFunctions.Emplace_GetRef(DefaultSig);
 		Sig.Name = NAME_IsValid;
@@ -882,5 +882,67 @@ void UNiagaraDataInterfaceDataTable::SetShaderParameters(const FNiagaraDataInter
 	ShaderParameters->FilteredRowDataOffsetBuffer	= InstanceData.FilteredRowDataOffsetBuffer.SRV;
 }
 
+#if WITH_EDITOR
+bool UNiagaraDataInterfaceDataTable::IsReadFunction(const FNiagaraFunctionSignature& Signature)
+{
+	using namespace NDIDataTableLocal;
+
+	return Signature.Name == NAME_GetRow || Signature.Name == NAME_GetFilteredRow;
+}
+
+TArray<FNiagaraVariableBase> UNiagaraDataInterfaceDataTable::GetVariablesFromDataTable(const UDataTable* DataTable)
+{
+	UScriptStruct* RowStruct = DataTable ? DataTable->RowStruct.Get() : nullptr;
+	if (RowStruct == nullptr)
+	{
+		return {};
+	}
+
+	TArray<FNiagaraVariableBase> Variables;
+	for (FProperty* Property = RowStruct->PropertyLink; Property != nullptr; Property = Property->PropertyLinkNext)
+	{
+		const FName PropertyName(*RowStruct->GetAuthoredNameForField(Property));
+
+		if (Property->IsA<const FIntProperty>())
+		{
+			Variables.Emplace(FNiagaraTypeDefinition::GetIntDef(), PropertyName);
+		}
+		else if (Property->IsA<const FFloatProperty>())
+		{
+			Variables.Emplace(FNiagaraTypeDefinition::GetFloatDef(), PropertyName);
+		}
+		else if (Property->IsA<const FDoubleProperty>())
+		{
+			Variables.Emplace(FNiagaraTypeDefinition::GetFloatDef(), PropertyName);
+		}
+		else if (const FStructProperty* StructProperty = CastField<const FStructProperty>(Property))
+		{
+			if (StructProperty->Struct == TBaseStructure<FVector2D>::Get())
+			{
+				Variables.Emplace(FNiagaraTypeDefinition::GetVec2Def(), PropertyName);
+			}
+			else if (StructProperty->Struct == TBaseStructure<FVector>::Get())
+			{
+				Variables.Emplace(FNiagaraTypeDefinition::GetVec3Def(), PropertyName);
+			}
+			else if (StructProperty->Struct == TBaseStructure<FVector4>::Get())
+			{
+				Variables.Emplace(FNiagaraTypeDefinition::GetVec4Def(), PropertyName);
+			}
+			else if (StructProperty->Struct == TBaseStructure<FQuat>::Get())
+			{
+				Variables.Emplace(FNiagaraTypeDefinition::GetQuatDef(), PropertyName);
+			}
+			else if (StructProperty->Struct == TBaseStructure<FLinearColor>::Get())
+			{
+				Variables.Emplace(FNiagaraTypeDefinition::GetColorDef(), PropertyName);
+			}
+		}
+	}
+
+	return Variables;
+}
+#endif
+
 #undef LOCTEXT_NAMESPACE
 
diff --git a/Engine/Plugins/FX/Niagara/Source/NiagaraEditor/Private/NiagaraActions.cpp b/Engine/Plugins/FX/Niagara/Source/NiagaraEditor/Private/NiagaraActions.cpp
--- a/Engine/Plugins/FX/Niagara/Source/NiagaraEditor/Private/NiagaraActions.cpp
+++ b/Engine/Plugins/FX/Niagara/Source/NiagaraEditor/Private/NiagaraActions.cpp
@@ -1,6 +1,7 @@
 // Copyright Epic Games, Inc. All Rights Reserved.
 
 #include "NiagaraActions.h"
+#include "DataInterface/NiagaraDataInterfaceDataTable.h"
 #include "NiagaraNodeParameterMapBase.h"
 #include "NiagaraNodeParameterMapGet.h"
 #include "NiagaraNodeParameterMapSet.h"
@@ -25,6 +26,7 @@
 
 #include "NiagaraDataChannel.h"
 #include "Config/NiagaraFavoriteActionsConfig.h"
+#include "Engine/DataTable.h"
 
 #include UE_INLINE_GENERATED_CPP_BY_NAME(NiagaraActions)
 
@@ -677,7 +679,7 @@ EVisibility FNiagaraParameterDragOperation::IsTextVisible() const
 
 TMap<FName, TUniquePtr<INiagaraDataInterfaceNodeActionProvider>> INiagaraDataInterfaceNodeActionProvider::RegisteredActionProviders;
 
-void INiagaraDataInterfaceNodeActionProvider::GetNodeContextMenuActions(UClass* DIClass, UToolMenu* Menu, UGraphNodeContextMenuContext* Context, FNiagaraFunctionSignature Signature)
+void INiagaraDataInterfaceNodeActionProvider::GetNodeContextMenuActions(UClass* DIClass, UToolMenu* Menu, UGraphNodeContextMenuContext* Context, const FNiagaraFunctionSignature& Signature)
 {
 	if(ensure(DIClass))
 	{
@@ -761,7 +763,7 @@ namespace NiagaraActionsLocal
 	FText InitForDataChannelMenuTooltipText = LOCTEXT("InitForDataChannelTooltip", "Initializes this node to write to all members of a given data channel.");
 }
 
-void FNiagaraDataInterfaceNodeActionProvider_DataChannelWrite::GetNodeContextMenuActionsImpl(UToolMenu* Menu, UGraphNodeContextMenuContext* Context, FNiagaraFunctionSignature Signature) const
+void FNiagaraDataInterfaceNodeActionProvider_DataChannelWrite::GetNodeContextMenuActionsImpl(UToolMenu* Menu, UGraphNodeContextMenuContext* Context, const FNiagaraFunctionSignature& Signature) const
 {
 	using namespace NiagaraActionsLocal;
 
@@ -933,7 +935,7 @@ void FNiagaraDataInterfaceNodeActionProvider_DataChannelWrite::AddDataChannelIni
 }
 
 
-void FNiagaraDataInterfaceNodeActionProvider_DataChannelRead::GetNodeContextMenuActionsImpl(UToolMenu* Menu, UGraphNodeContextMenuContext* Context, FNiagaraFunctionSignature Signature) const
+void FNiagaraDataInterfaceNodeActionProvider_DataChannelRead::GetNodeContextMenuActionsImpl(UToolMenu* Menu, UGraphNodeContextMenuContext* Context, const FNiagaraFunctionSignature& Signature) const
 {
 	using namespace NiagaraActionsLocal;
 
@@ -1105,4 +1107,83 @@ TSharedPtr<SWidget> FNiagaraDataInterfaceNodeActionProvider_DataChannelRead::Get
 		.AllowedTypes(MoveTemp(AllowedTypes));
 }
 
+/////////////////////////////////////////////////////////////////////////////////////////////
+
+void FNiagaraDataInterfaceNodeActionProvider_DataTable::GetInlineNodeContextMenuActionsImpl(UToolMenu* ToolMenu) const
+{
+	AddInitializeActions(ToolMenu);
+}
+
+INiagaraDataInterfaceNodeActionProvider::FInlineMenuDisplayOptions FNiagaraDataInterfaceNodeActionProvider_DataTable::GetInlineMenuDisplayOptionsImpl(UClass* DIClass, UEdGraphNode* Source) const
+{
+	if (UNiagaraNodeFunctionCall* FunctionCall = Cast<UNiagaraNodeFunctionCall>(Source))
+	{
+		if (UNiagaraDataInterfaceDataTable::IsReadFunction(FunctionCall->Signature))
+		{
+			FInlineMenuDisplayOptions DisplayOptions;
+			DisplayOptions.bDisplayInline = true;
+			DisplayOptions.DisplayBrush = FAppStyle::GetBrush("Icons.Edit");
+			DisplayOptions.TooltipText = LOCTEXT("InitDataTableReadNode", "Initialize the data table read using the selected data table structure.");
+
+			return DisplayOptions;
+		}
+	}
+	return {};
+}
+
+void FNiagaraDataInterfaceNodeActionProvider_DataTable::AddInitializeActions(UToolMenu* ToolMenu)
+{
+	UGraphNodeContextMenuContext* Context = ToolMenu->FindContext<UGraphNodeContextMenuContext>();
+	if (Context == nullptr || Context->Node == nullptr)
+	{
+		return;
+	}
+
+	FAssetRegistryModule& AssetRegistryModule = FModuleManager::LoadModuleChecked<FAssetRegistryModule>("AssetRegistry");
+
+	FARFilter Filter;
+	Filter.bRecursiveClasses = true;
+	Filter.ClassPaths.Add(FTopLevelAssetPath(UDataTable::StaticClass()));
+
+	TArray<FAssetData> DataTableAssets;
+	AssetRegistryModule.Get().GetAssets(Filter, DataTableAssets);
+
+	if (DataTableAssets.Num() == 0)
+	{
+		return;
+	}
+
+	FToolMenuSection& MenuSection = ToolMenu->AddSection("DataTable", LOCTEXT("DataTable", "Data Table"));
+	TWeakObjectPtr<UNiagaraNodeFunctionCall> WeakFunctionNode = CastChecked<UNiagaraNodeFunctionCall>(ConstCast(Context->Node));
+
+	for (const FAssetData& AssetData : DataTableAssets)
+	{
+		FText AssetName = FText::FromName(AssetData.AssetName);
+		MenuSection.AddMenuEntry(
+			NAME_None,
+			AssetName,
+			FText::Format(LOCTEXT("SetNodeToDataTableTooltip", "Set node to read data table '{0}'"), AssetName),
+			FSlateIcon(FAppStyle::GetAppStyleSetName(), "Icons.Import"),
+			FUIAction(FExecuteAction::CreateStatic(&AllocateNodePins, AssetData.ToSoftObjectPath(), WeakFunctionNode))
+		);
+	}
+}
+
+void FNiagaraDataInterfaceNodeActionProvider_DataTable::AllocateNodePins(FSoftObjectPath DataTablePath, TWeakObjectPtr<UNiagaraNodeFunctionCall> WeakFunctionNode)
+{
+	UDataTable* DataTable = Cast<UDataTable>(DataTablePath.TryLoad());
+	UNiagaraNodeFunctionCall* FunctionNode = WeakFunctionNode.Get();
+	if (!DataTable || !FunctionNode)
+	{
+		return;
+	}
+
+	FunctionNode->RemoveAllDynamicPins();
+
+	for (const FNiagaraVariableBase& Variable : UNiagaraDataInterfaceDataTable::GetVariablesFromDataTable(DataTable))
+	{
+		FunctionNode->AddParameter(Variable, EEdGraphPinDirection::EGPD_Output);
+	}
+}
+
 #undef LOCTEXT_NAMESPACE
diff --git a/Engine/Plugins/FX/Niagara/Source/NiagaraEditor/Private/NiagaraEditorModule.cpp b/Engine/Plugins/FX/Niagara/Source/NiagaraEditor/Private/NiagaraEditorModule.cpp
--- a/Engine/Plugins/FX/Niagara/Source/NiagaraEditor/Private/NiagaraEditorModule.cpp
+++ b/Engine/Plugins/FX/Niagara/Source/NiagaraEditor/Private/NiagaraEditorModule.cpp
@@ -58,6 +58,7 @@
 #include "NiagaraDataInterfaceColorCurve.h"
 #include "DataInterface/NiagaraDataInterfaceDataChannelRead.h"
 #include "DataInterface/NiagaraDataInterfaceDataChannelWrite.h"
+#include "DataInterface/NiagaraDataInterfaceDataTable.h"
 #include "DataInterface/NiagaraDataInterfaceMemoryBuffer.h"
 #include "DataInterface/NiagaraDataInterfaceSimpleCounter.h"
 #include "NiagaraDataInterfaceRenderTarget2D.h"
@@ -1614,6 +1615,7 @@ void FNiagaraEditorModule::StartupModule()
 	//Register node action providers for data interface functions.
 	INiagaraDataInterfaceNodeActionProvider::Register<UNiagaraDataInterfaceDataChannelWrite, FNiagaraDataInterfaceNodeActionProvider_DataChannelWrite>();
 	INiagaraDataInterfaceNodeActionProvider::Register<UNiagaraDataInterfaceDataChannelRead, FNiagaraDataInterfaceNodeActionProvider_DataChannelRead>();
+	INiagaraDataInterfaceNodeActionProvider::Register<UNiagaraDataInterfaceDataTable, FNiagaraDataInterfaceNodeActionProvider_DataTable>();
 
 	FAssetCompilingManager::Get().RegisterManager(&FNiagaraSystemCompilingManager::Get());
 	ILevelSequenceModule& LevelSequenceModule = FModuleManager::LoadModuleChecked<ILevelSequenceModule>("LevelSequence");
diff --git a/Engine/Plugins/FX/Niagara/Source/NiagaraEditor/Public/NiagaraActions.h b/Engine/Plugins/FX/Niagara/Source/NiagaraEditor/Public/NiagaraActions.h
--- a/Engine/Plugins/FX/Niagara/Source/NiagaraEditor/Public/NiagaraActions.h
+++ b/Engine/Plugins/FX/Niagara/Source/NiagaraEditor/Public/NiagaraActions.h
@@ -409,7 +409,7 @@ public:
 	virtual ~INiagaraDataInterfaceNodeActionProvider() = default;
 
 	/** Allows DIs to add context menu actions for the whole node. */
-	virtual void GetNodeContextMenuActionsImpl(UToolMenu* Menu, UGraphNodeContextMenuContext* Context, FNiagaraFunctionSignature Signature) const {}
+	virtual void GetNodeContextMenuActionsImpl(UToolMenu* Menu, UGraphNodeContextMenuContext* Context, const FNiagaraFunctionSignature& Signature) const {}
 
 	virtual void GetInlineNodeContextMenuActionsImpl(UToolMenu* ToolMenu) const {}
 
@@ -427,9 +427,9 @@ public:
 	static void Unregister();
 
 	template<typename DIClass>
-	static void GetNodeContextMenuActions(UToolMenu* Menu, UGraphNodeContextMenuContext* Context, FNiagaraFunctionSignature Signature);
+	static void GetNodeContextMenuActions(UToolMenu* Menu, UGraphNodeContextMenuContext* Context, const FNiagaraFunctionSignature& Signature);
 
-	static NIAGARAEDITOR_API void GetNodeContextMenuActions(UClass* DIClass, UToolMenu* Menu, UGraphNodeContextMenuContext* Context, FNiagaraFunctionSignature Signature);
+	static NIAGARAEDITOR_API void GetNodeContextMenuActions(UClass* DIClass, UToolMenu* Menu, UGraphNodeContextMenuContext* Context, const FNiagaraFunctionSignature& Signature);
 
 	static NIAGARAEDITOR_API void GetInlineNodeContextMenuActions(UClass* DIClass, UToolMenu* ToolMenu);
 	
@@ -460,7 +460,7 @@ void INiagaraDataInterfaceNodeActionProvider::Unregister()
 }
 
 template<typename DIClass>
-void INiagaraDataInterfaceNodeActionProvider::GetNodeContextMenuActions(UToolMenu* Menu, UGraphNodeContextMenuContext* Context, FNiagaraFunctionSignature Signature)
+void INiagaraDataInterfaceNodeActionProvider::GetNodeContextMenuActions(UToolMenu* Menu, UGraphNodeContextMenuContext* Context, const FNiagaraFunctionSignature& Signature)
 {
 	if(TUniquePtr<INiagaraDataInterfaceNodeActionProvider>* Provider = RegisteredActionProviders.Find(DIClass::StaticClass()->GetFName()))
 	{
@@ -485,7 +485,7 @@ class FNiagaraDataInterfaceNodeActionProvider_DataChannelWrite : public INiagara
 {
 public:
 
-	virtual void GetNodeContextMenuActionsImpl(UToolMenu* Menu, UGraphNodeContextMenuContext* Context, FNiagaraFunctionSignature Signature) const override;
+	virtual void GetNodeContextMenuActionsImpl(UToolMenu* Menu, UGraphNodeContextMenuContext* Context, const FNiagaraFunctionSignature& Signature) const override;
 	virtual void CollectAddPinActionsImpl(FNiagaraMenuActionCollector& Collector, UEdGraphPin* AddPin)const override;
 	virtual void GetInlineNodeContextMenuActionsImpl(UToolMenu* ToolMenu) const override;
 	virtual FInlineMenuDisplayOptions GetInlineMenuDisplayOptionsImpl(UClass* DIClass, UEdGraphNode* Source) const override;
@@ -498,7 +498,7 @@ class FNiagaraDataInterfaceNodeActionProvider_DataChannelRead : public INiagaraD
 {
 public:
 
-	virtual void GetNodeContextMenuActionsImpl(UToolMenu* Menu, UGraphNodeContextMenuContext* Context, FNiagaraFunctionSignature Signature) const override;
+	virtual void GetNodeContextMenuActionsImpl(UToolMenu* Menu, UGraphNodeContextMenuContext* Context, const FNiagaraFunctionSignature& Signature) const override;
 	virtual void GetInlineNodeContextMenuActionsImpl(UToolMenu* ToolMenu) const override;
 	virtual FInlineMenuDisplayOptions GetInlineMenuDisplayOptionsImpl(UClass* DIClass, UEdGraphNode* Source) const override;
 	virtual void CollectAddPinActionsImpl(FNiagaraMenuActionCollector& Collector, UEdGraphPin* AddPin)const override;
@@ -507,3 +507,14 @@ public:
 private:
 	static void AddDataChannelInitActions(UToolMenu* ToolMenu);
 };
+
+class FNiagaraDataInterfaceNodeActionProvider_DataTable : public INiagaraDataInterfaceNodeActionProvider
+{
+public:
+	virtual void GetInlineNodeContextMenuActionsImpl(UToolMenu* ToolMenu) const override;
+	virtual FInlineMenuDisplayOptions GetInlineMenuDisplayOptionsImpl(UClass* DIClass, UEdGraphNode* Source) const override;
+
+private:
+	static void AddInitializeActions(UToolMenu* ToolMenu);
+	static void AllocateNodePins(FSoftObjectPath DataTablePath, TWeakObjectPtr<UNiagaraNodeFunctionCall> WeakFunctionNode);
+};
