- Speculative fix for null component pool on shutdown

#rnx

diff --git a/Engine/Plugins/FX/Niagara/Source/Niagara/Private/NiagaraWorldManager.cpp b/Engine/Plugins/FX/Niagara/Source/Niagara/Private/NiagaraWorldManager.cpp
--- a/Engine/Plugins/FX/Niagara/Source/Niagara/Private/NiagaraWorldManager.cpp
+++ b/Engine/Plugins/FX/Niagara/Source/Niagara/Private/NiagaraWorldManager.cpp
@@ -719,7 +719,10 @@ void FNiagaraWorldManager::DestroySystemSimulation(UNiagaraSystem* System)
 			}
 		}
 	}
-	ComponentPool->RemoveComponentsBySystem(System);
+	if (ComponentPool)
+	{
+		ComponentPool->RemoveComponentsBySystem(System);
+	}
 
 #if WITH_EDITOR
 	if (FDelegateHandle* Handle = SystemPostChangeDelegateHandles.Find(System))
@@ -762,7 +765,10 @@ void FNiagaraWorldManager::OnComputeDispatchInterfaceDestroyed_Internal(FNiagara
 void FNiagaraWorldManager::OnWorldBeginTearDown()
 {
 	bIsTearingDown = true;
-	ComponentPool->Cleanup(World);
+	if (ComponentPool)
+	{
+		ComponentPool->Cleanup(World);
+	}
 	DataChannelManager->Cleanup();
 }
 
@@ -771,7 +777,10 @@ void FNiagaraWorldManager::OnWorldCleanup(bool bSessionEnded, bool bCleanupResou
 	bIsCleaningUp = true;
 
 	DeferredMethods.ExecuteAndClear();
-	ComponentPool->Cleanup(World);
+	if (ComponentPool)
+	{
+		ComponentPool->Cleanup(World);
+	}
 
 	DataChannelManager->Cleanup();
 
@@ -812,7 +821,10 @@ void FNiagaraWorldManager::OnWorldCleanup(bool bSessionEnded, bool bCleanupResou
 
 void FNiagaraWorldManager::OnPostWorldCleanup(bool bSessionEnded, bool bCleanupResources)
 {
-	ComponentPool->Cleanup(World);
+	if (ComponentPool)
+	{
+		ComponentPool->Cleanup(World);
+	}
 }
 
 void FNiagaraWorldManager::PreGarbageCollect()
@@ -2262,7 +2274,7 @@ void FNiagaraWorldManager::PrimePoolForAllWorlds(UNiagaraSystem* System)
 
 void FNiagaraWorldManager::PrimePoolForAllSystems()
 {
-	if (GNigaraAllowPrimedPools && HasActiveWorld() && World->IsGameWorld())
+	if (GNigaraAllowPrimedPools && HasActiveWorld() && World->IsGameWorld() && ComponentPool)
 	{
 		//Prime the pool for all currently loaded systems.
 		for (TObjectIterator<UNiagaraSystem> It; It; ++It)
@@ -2278,7 +2290,7 @@ void FNiagaraWorldManager::PrimePoolForAllSystems()
 
 void FNiagaraWorldManager::PrimePool(UNiagaraSystem* System)
 {
-	if (GNigaraAllowPrimedPools && HasActiveWorld() && World->IsGameWorld())
+	if (GNigaraAllowPrimedPools && HasActiveWorld() && World->IsGameWorld() && ComponentPool)
 	{
 		ComponentPool->PrimePool(System, World);
 	}
