CIS fix

[RN] rnx

diff --git a/Engine/Plugins/FX/Niagara/Source/Niagara/Niagara.Build.cs b/Engine/Plugins/FX/Niagara/Source/Niagara/Niagara.Build.cs
--- a/Engine/Plugins/FX/Niagara/Source/Niagara/Niagara.Build.cs
+++ b/Engine/Plugins/FX/Niagara/Source/Niagara/Niagara.Build.cs
@@ -41,6 +41,7 @@ public class Niagara : ModuleRules
                 "PhysicsCore",
                 "RenderCore",
                 "RHI",
+				"SlateCore",
                 "VectorVM",
             }
 		);
