- Adjustments to stateless noise
- Improves quality, and more consistent over varying lifetimes

#rnx

diff --git a/Engine/Plugins/FX/Niagara/Shaders/Private/Stateless/Modules/NiagaraStatelessModule_SolveVelocitiesAndForces.ush b/Engine/Plugins/FX/Niagara/Shaders/Private/Stateless/Modules/NiagaraStatelessModule_SolveVelocitiesAndForces.ush
--- a/Engine/Plugins/FX/Niagara/Shaders/Private/Stateless/Modules/NiagaraStatelessModule_SolveVelocitiesAndForces.ush
+++ b/Engine/Plugins/FX/Niagara/Shaders/Private/Stateless/Modules/NiagaraStatelessModule_SolveVelocitiesAndForces.ush
@@ -45,13 +45,18 @@ struct FStatelessModule_SolveVelocitiesAndForces
 	float3 Acceleration;
 };
 
-void SolveVelocitiesAndForces_EvaluateNoise(in FStatelessModule_SolveVelocitiesAndForces ModuleData, float NormalizedAge, float Age, inout float3 Position)
+void SolveVelocitiesAndForces_EvaluateNoise(in FStatelessModule_SolveVelocitiesAndForces ModuleData, float Lifetime, float NormalizedAge, float Age, inout float3 Position)
 {
-	const float fRow	= length2(Position * 0.01f);
+	const float3 SamplePosition	= Position * SolveVelocitiesAndForces_NoiseFrequency;
+	
+	const float fRow	= dot(abs(SamplePosition), float3(1.0f / 8.0f, 2.0f / 8.0f, 5.0f / 8.0f));
 	const uint RowA		= uint(floor(abs(fRow))) & SolveVelocitiesAndForces_NoiseLUTRowMask;
 	const uint RowB		= (RowA + 1) & SolveVelocitiesAndForces_NoiseLUTRowMask;
-		
-	const float fColumn = abs(SolveVelocitiesAndForces_NoiseFrequency) * NormalizedAge;
+
+	const float TotalDistance	= Lifetime * SolveVelocitiesAndForces_NoiseStrength;
+	const float EndColumn		= clamp(TotalDistance * SolveVelocitiesAndForces_NoiseLUTNormalize, 1.0f, float(SolveVelocitiesAndForces_NoiseLUTRowWidth - 1));
+    
+	const float fColumn = NormalizedAge * EndColumn;
 	const uint iColumnA = min(floor(fColumn), SolveVelocitiesAndForces_NoiseLUTRowWidth - 1);
 	const uint iColumnB = min(iColumnA + 1, SolveVelocitiesAndForces_NoiseLUTRowWidth - 1);
 
@@ -62,7 +67,7 @@ void SolveVelocitiesAndForces_EvaluateNoise(in FStatelessModule_SolveVelocitiesA
 	const float3 ValueA  = lerp(ValueA0, ValueA1, frac(fRow));
 	const float3 ValueB  = lerp(ValueB0, ValueB1, frac(fRow));
 	
-	Position += lerp(ValueA, ValueB, frac(fColumn)) * SolveVelocitiesAndForces_NoiseStrength * Age * SolveVelocitiesAndForces_NoiseLUTNormalize;
+	Position += lerp(ValueA, ValueB, frac(fColumn)) * TotalDistance / EndColumn;
 }
 
 void SolveVelocitiesAndForces_IntegratePosition(in FStatelessModule_SolveVelocitiesAndForces ModuleData, float Age, inout float3 Position)
@@ -125,8 +130,8 @@ void SolveVelocitiesAndForces_Simulate(inout FStatelessParticle Particle)
 
 	if ( SolveVelocitiesAndForces_NoiseEnabled  != 0 )
 	{
-		SolveVelocitiesAndForces_EvaluateNoise(ModuleData, Particle.NormalizedAge, Particle.Age, Particle.Position);
-		SolveVelocitiesAndForces_EvaluateNoise(ModuleData, Particle.PreviousNormalizedAge, Particle.PreviousAge, Particle.PreviousPosition);
+		SolveVelocitiesAndForces_EvaluateNoise(ModuleData, Particle.Lifetime, Particle.NormalizedAge, Particle.Age, Particle.Position);
+		SolveVelocitiesAndForces_EvaluateNoise(ModuleData, Particle.Lifetime, Particle.PreviousNormalizedAge, Particle.PreviousAge, Particle.PreviousPosition);
 	}
 	
 	// Simulate
diff --git a/Engine/Plugins/FX/Niagara/Source/Niagara/Private/Stateless/Modules/NiagaraStatelessModule_SolveVelocitiesAndForces.cpp b/Engine/Plugins/FX/Niagara/Source/Niagara/Private/Stateless/Modules/NiagaraStatelessModule_SolveVelocitiesAndForces.cpp
--- a/Engine/Plugins/FX/Niagara/Source/Niagara/Private/Stateless/Modules/NiagaraStatelessModule_SolveVelocitiesAndForces.cpp
+++ b/Engine/Plugins/FX/Niagara/Source/Niagara/Private/Stateless/Modules/NiagaraStatelessModule_SolveVelocitiesAndForces.cpp
@@ -9,6 +9,22 @@
 
 namespace NiagaraStatelessModuleSolveVelocitiesAndForcesPrivate
 {
+	float GFrequencyNormalize = 50.0f;
+	FAutoConsoleVariableRef CVarFrequencyNormalize(
+		TEXT("fx.NiagaraStateless.SolveVelocities.Noise.FrequencyNormalize"),
+		GFrequencyNormalize,
+		TEXT("Translation from position -> lookup space."),
+		ECVF_Default
+	);
+
+	float GFrequencyLUTNormalize = 0.1f;
+	FAutoConsoleVariableRef CVarFrequencyLUTNormalize(
+		TEXT("fx.NiagaraStateless.SolveVelocities.Noise.FrequencyLUTNormalize"),
+		GFrequencyLUTNormalize,
+		TEXT("Modifier applied when looking up from normalized total travel into the LUT."),
+		ECVF_Default
+	);
+
 	struct FModuleBuiltData
 	{
 		NiagaraStateless::FPhysicsBuildData	PhysicsData;
@@ -32,13 +48,18 @@ namespace NiagaraStatelessModuleSolveVelocitiesAndForcesPrivate
 		return Age * (Velocity + Wind) + 0.5f * Acceleration * Age * Age;
 	}
 
-	static void EvaluateNoise(TConstArrayView<FVector3f> NoiseLUT, const UNiagaraStatelessModule_SolveVelocitiesAndForces::FParameters* Parameters, float NormalizedAge, float Age, FVector3f& Position)
+	static void EvaluateNoise(TConstArrayView<FVector3f> NoiseLUT, const UNiagaraStatelessModule_SolveVelocitiesAndForces::FParameters* Parameters, float Lifetime, float NormalizedAge, FVector3f& Position)
 	{
-		const float fRow = (Position * 0.01f).SquaredLength();
+		const FVector3f SamplePosition = Position * Parameters->SolveVelocitiesAndForces_NoiseFrequency;
+
+		const float fRow = FVector3f::DotProduct(SamplePosition.GetAbs(), FVector3f(1.0f / 8.0f, 2.0f / 8.0f, 5.0f / 8.0f));
 		const uint32 RowA = uint32(FMath::FloorToInt(FMath::Abs(fRow))) & Parameters->SolveVelocitiesAndForces_NoiseLUTRowMask;
 		const uint32 RowB = (RowA + 1) & Parameters->SolveVelocitiesAndForces_NoiseLUTRowMask;
 
-		const float fColumn = FMath::Abs(Parameters->SolveVelocitiesAndForces_NoiseFrequency) * NormalizedAge;
+		const float TotalDistance = Lifetime * Parameters->SolveVelocitiesAndForces_NoiseStrength;
+		const float EndColumn = FMath::Clamp(TotalDistance * Parameters->SolveVelocitiesAndForces_NoiseLUTNormalize, 1.0f, float(Parameters->SolveVelocitiesAndForces_NoiseLUTRowWidth - 1));
+
+		const float fColumn = NormalizedAge * EndColumn;
 		const uint32 iColumnA = FMath::Min(uint32(FMath::FloorToInt(fColumn)), Parameters->SolveVelocitiesAndForces_NoiseLUTRowWidth - 1);
 		const uint32 iColumnB = FMath::Min(iColumnA + 1, Parameters->SolveVelocitiesAndForces_NoiseLUTRowWidth - 1);
 
@@ -49,7 +70,7 @@ namespace NiagaraStatelessModuleSolveVelocitiesAndForcesPrivate
 		const FVector3f ValueA = FMath::Lerp(ValueA0, ValueA1, FMath::Fractional(fRow));
 		const FVector3f ValueB = FMath::Lerp(ValueB0, ValueB1, FMath::Fractional(fRow));
 
-		Position += FMath::Lerp(ValueA, ValueB, FMath::Fractional(fColumn)) * Parameters->SolveVelocitiesAndForces_NoiseStrength * Age * Parameters->SolveVelocitiesAndForces_NoiseLUTNormalize;
+		Position += FMath::Lerp(ValueA, ValueB, FMath::Fractional(fColumn)) * TotalDistance / EndColumn;
 	}
 
 	static void ParticleSimulate(const NiagaraStateless::FParticleSimulationContext& ParticleSimulationContext)
@@ -57,6 +78,7 @@ namespace NiagaraStatelessModuleSolveVelocitiesAndForcesPrivate
 		const FModuleBuiltData* ModuleBuiltData = ParticleSimulationContext.ReadBuiltData<FModuleBuiltData>();
 		const UNiagaraStatelessModule_SolveVelocitiesAndForces::FParameters* Parameters = ParticleSimulationContext.ReadParameterNestedStruct<UNiagaraStatelessModule_SolveVelocitiesAndForces::FParameters>();
 
+		const float* LifetimeData				= ParticleSimulationContext.GetParticleLifetime();
 		const float* AgeData					= ParticleSimulationContext.GetParticleAge();
 		const float* PreviousAgeData			= ParticleSimulationContext.GetParticlePreviousAge();
 		const float* NormalizedAgeData			= ParticleSimulationContext.GetParticleNormalizedAge();
@@ -109,8 +131,8 @@ namespace NiagaraStatelessModuleSolveVelocitiesAndForcesPrivate
 
 			if (Parameters->SolveVelocitiesAndForces_NoiseEnabled != 0)
 			{
-				EvaluateNoise(NoiseLUT, Parameters, NormalizedAgeData[i], AgeData[i], Position);
-				EvaluateNoise(NoiseLUT, Parameters, PreviousNormalizedAgeData[i], PreviousAgeData[i], PreviousPosition);
+				EvaluateNoise(NoiseLUT, Parameters, LifetimeData[i], NormalizedAgeData[i], Position);
+				EvaluateNoise(NoiseLUT, Parameters, LifetimeData[i], PreviousNormalizedAgeData[i], PreviousPosition);
 			}
 
 			Position += IntegratePosition(FMath::Max(AgeData[i], UE_KINDA_SMALL_NUMBER), Drag, InitialVelocity, Wind, Acceleration);
@@ -207,14 +229,10 @@ void UNiagaraStatelessModule_SolveVelocitiesAndForces::SetShaderParameters(const
 		const NiagaraStateless::FNoiseLUT& NoiseLUT = NiagaraStateless::FNoiseLUT::GetGlobalLUT();
 		checkf(FMath::IsPowerOfTwo(NoiseLUT.GetNumRows()), TEXT("Noise LUT must be PoT"));
 
-		const float fLUTSizeMinusOne = float(NoiseLUT.GetRowWidth() - 1);
-		const float ClampedFrequency = FMath::Min(FMath::Abs(PhysicsData.NoiseFrequency / 2.0f) + 1.0f, fLUTSizeMinusOne);
-		const float Frequency = PhysicsData.NoiseFrequency >= 0.0f ? ClampedFrequency : -ClampedFrequency;
-
 		Parameters->SolveVelocitiesAndForces_NoiseEnabled		= PhysicsData.bNoiseEnabled ? 1 : 0;
 		Parameters->SolveVelocitiesAndForces_NoiseStrength		= PhysicsData.NoiseStrength;
-		Parameters->SolveVelocitiesAndForces_NoiseFrequency		= Frequency;
-		Parameters->SolveVelocitiesAndForces_NoiseLUTNormalize	= float(NoiseLUT.GetRowWidth()) / Frequency;
+		Parameters->SolveVelocitiesAndForces_NoiseFrequency		= PhysicsData.NoiseFrequency / GFrequencyNormalize;
+		Parameters->SolveVelocitiesAndForces_NoiseLUTNormalize	= (PhysicsData.NoiseFrequency / GFrequencyNormalize) * GFrequencyLUTNormalize;
 		Parameters->SolveVelocitiesAndForces_NoiseLUTRowMask	= NoiseLUT.GetNumRows() - 1;
 		Parameters->SolveVelocitiesAndForces_NoiseLUTRowWidth	= NoiseLUT.GetRowWidth();
 		Parameters->SolveVelocitiesAndForces_NoiseLUT			= NoiseLUT.GetGpuData();
diff --git a/Engine/Plugins/FX/Niagara/Source/Niagara/Private/Stateless/NiagaraStatelessNoiseLUT.cpp b/Engine/Plugins/FX/Niagara/Source/Niagara/Private/Stateless/NiagaraStatelessNoiseLUT.cpp
--- a/Engine/Plugins/FX/Niagara/Source/Niagara/Private/Stateless/NiagaraStatelessNoiseLUT.cpp
+++ b/Engine/Plugins/FX/Niagara/Source/Niagara/Private/Stateless/NiagaraStatelessNoiseLUT.cpp
@@ -5,28 +5,29 @@
 
 namespace NiagaraStateless
 {
-	static TGlobalResource<NiagaraStateless::FNoiseLUT> GGlobalLUT(32, 32, 0.1f);
+	static TGlobalResource<NiagaraStateless::FNoiseLUT> GGlobalLUT(32, 64, 0.1f);
 
-	FNoiseLUT::FNoiseLUT(uint32 Rows, uint32 RowWidth, float FieldTravelSpeed)
+	FNoiseLUT::FNoiseLUT(uint32 Rows, uint32 RowWidth, float InFieldTravelSpeed)
 		: LUTRows(Rows)
 		, LUTRowWidth(RowWidth)
+		, FieldTravelSpeed(InFieldTravelSpeed)
 	{
 		CpuData.AddUninitialized(LUTRows * LUTRowWidth);
 		for (uint32 Row = 0; Row < LUTRows; ++Row)
 		{
-			FVector3f Position = FVector3f(Row * 1024.0f, Row * 511.0f, Row * 30.1f);
-			FVector3f Sum = FVector3f::ZeroVector;
+			const float fRow = float(Row) - (float(LUTRows) * 0.5f);
+			FVector3f SamplePosition = FVector3f(fRow, fRow * 3.0f, fRow * 9.0f);
+			FVector3f PositionOffset = FVector3f::ZeroVector;
 
 			uint32 RowOutputOffset = Row * LUTRowWidth;
-			CpuData[RowOutputOffset++] = Sum;
+			CpuData[RowOutputOffset++] = PositionOffset;
 			for (uint32 Column = 1; Column < LUTRowWidth; ++Column)
 			{
-				const FNiagaraMatrix3x4 J = JacobianSimplex_ALU(Position);
+				const FNiagaraMatrix3x4 J = JacobianSimplex_ALU(SamplePosition);
 				const FVector3f Dir = FVector3f(J[1][2] - J[2][1], J[2][0] - J[0][2], J[0][1] - J[1][0]).GetSafeNormal();
-
-				Sum += Dir * float(1.0f / float(LUTRowWidth));
-				Position += Dir * FieldTravelSpeed;
-				CpuData[RowOutputOffset++] = Sum;
+				SamplePosition += Dir * FieldTravelSpeed;
+				PositionOffset += Dir;
+				CpuData[RowOutputOffset++] = PositionOffset;
 			}
 		}
 	}
diff --git a/Engine/Plugins/FX/Niagara/Source/Niagara/Private/Stateless/NiagaraStatelessNoiseLUT.h b/Engine/Plugins/FX/Niagara/Source/Niagara/Private/Stateless/NiagaraStatelessNoiseLUT.h
--- a/Engine/Plugins/FX/Niagara/Source/Niagara/Private/Stateless/NiagaraStatelessNoiseLUT.h
+++ b/Engine/Plugins/FX/Niagara/Source/Niagara/Private/Stateless/NiagaraStatelessNoiseLUT.h
@@ -18,14 +18,16 @@ namespace NiagaraStateless
 
 		uint32 GetNumRows() const { return LUTRows; }
 		uint32 GetRowWidth() const { return LUTRowWidth; }
+		float GetFieldTravelSpeed() const { return FieldTravelSpeed; }
 		TConstArrayView<FVector3f> GetCpuData() const { return CpuData; }
 		FRHIShaderResourceView* GetGpuData() const { return GpuData.SRV; }
 
 		static const FNoiseLUT& GetGlobalLUT();
 
 	private:
-		uint32				LUTRows = 0;
-		uint32				LUTRowWidth = 0;
+		const uint32		LUTRows = 0;
+		const uint32		LUTRowWidth = 0;
+		const float			FieldTravelSpeed = 0.1f;
 
 		TArray<FVector3f>	CpuData;
 		FByteAddressBuffer	GpuData;
