- Disable the component tick if we are not in solo mode

#rnx

diff --git a/Engine/Plugins/FX/Niagara/Source/Niagara/Private/NiagaraComponent.cpp b/Engine/Plugins/FX/Niagara/Source/Niagara/Private/NiagaraComponent.cpp
--- a/Engine/Plugins/FX/Niagara/Source/Niagara/Private/NiagaraComponent.cpp
+++ b/Engine/Plugins/FX/Niagara/Source/Niagara/Private/NiagaraComponent.cpp
@@ -881,6 +881,16 @@ void UNiagaraComponent::TickComponent(float DeltaSeconds, enum ELevelTick TickTy
 		return;
 	}
 
+	// The component tick function should only be execute when in solo mode
+	// When in batched mode the world manager will ensure the system ticks
+	// The replay net driver can sometimes enable our ticks via SetComponentTickEnabled so disable here
+	// Ideally we would intercept SetComponentTickEnabled and have a concept of if we are managing the tick or not
+	if (SystemInstanceController->IsSolo() == false)
+	{
+		SetComponentTickEnabled(false);
+		return;
+	}
+
 	if (Asset->HasFixedTickDelta() && Asset->GetFixedTickDeltaTime() <= 0)
 	{
 		return;
@@ -891,7 +901,6 @@ void UNiagaraComponent::TickComponent(float DeltaSeconds, enum ELevelTick TickTy
 		Activate();
 	}
 
-	check(SystemInstanceController->IsSolo());
 	if (IsActive() && SystemInstanceController.Get() && !SystemInstanceController->IsComplete())
 	{
 		check(Asset != nullptr);
