- Reset stateless module to defaults on delete

#rnx

diff --git a/Engine/Plugins/FX/Niagara/Source/NiagaraEditor/Private/ViewModels/Stack/NiagaraStackStatelessEmitterSimulateGroup.cpp b/Engine/Plugins/FX/Niagara/Source/NiagaraEditor/Private/ViewModels/Stack/NiagaraStackStatelessEmitterSimulateGroup.cpp
--- a/Engine/Plugins/FX/Niagara/Source/NiagaraEditor/Private/ViewModels/Stack/NiagaraStackStatelessEmitterSimulateGroup.cpp
+++ b/Engine/Plugins/FX/Niagara/Source/NiagaraEditor/Private/ViewModels/Stack/NiagaraStackStatelessEmitterSimulateGroup.cpp
@@ -337,6 +337,7 @@ void UNiagaraStackStatelessModuleItem::Delete()
 	if (StatelessModule != nullptr && StatelessModule->CanDisableModule())
 	{
 		StatelessModule->Modify();
+		UEngine::CopyPropertiesForUnrelatedObjects(StatelessModule->GetClass()->GetDefaultObject(), StatelessModule);
 		StatelessModule->SetIsModuleEnabled(false);
 		StatelessModule->PostEditChange();
 		OnDataObjectModified().Broadcast({ StatelessModule }, ENiagaraDataObjectChange::Changed);
