- Add space rotations to stateless common parameters

#rnx

diff --git a/Engine/Plugins/FX/Niagara/Shaders/Private/Stateless/NiagaraStatelessCommon.ush b/Engine/Plugins/FX/Niagara/Shaders/Private/Stateless/NiagaraStatelessCommon.ush
--- a/Engine/Plugins/FX/Niagara/Shaders/Private/Stateless/NiagaraStatelessCommon.ush
+++ b/Engine/Plugins/FX/Niagara/Shaders/Private/Stateless/NiagaraStatelessCommon.ush
@@ -25,6 +25,7 @@ bool IsValidComponent(int ComponentOffset)
 	return ComponentOffset != INVALID_COMPONENT;
 }
 
+float4				Common_ToSimulationRotations[3];
 uint				Common_OutputBufferStride;
 RWBuffer<float>		Common_FloatOutputBuffer;
 RWBuffer<int>		Common_IntOutputBuffer;
@@ -33,6 +34,13 @@ Buffer<uint>		Common_ParameterBuffer;
 
 ////////////////////////////////////////////////////////////////////////////////////
 
+float4 GetToSimulationRotation(int SourceSpace)
+{
+	return Common_ToSimulationRotations[SourceSpace];
+}
+
+////////////////////////////////////////////////////////////////////////////////////
+
 float2 SinCos(float v)
 {
 	float2 sc;
diff --git a/Engine/Plugins/FX/Niagara/Source/Niagara/Internal/Stateless/NiagaraStatelessCommon.h b/Engine/Plugins/FX/Niagara/Source/Niagara/Internal/Stateless/NiagaraStatelessCommon.h
--- a/Engine/Plugins/FX/Niagara/Source/Niagara/Internal/Stateless/NiagaraStatelessCommon.h
+++ b/Engine/Plugins/FX/Niagara/Source/Niagara/Internal/Stateless/NiagaraStatelessCommon.h
@@ -31,6 +31,7 @@ struct FNiagaraStatelessSpaceTransforms
 	
 	const FTransform3f& GetLocalToWorld() const { return Transforms[LocalToWorldIndex]; }
 	const FTransform3f& GetWorldToLocal() const { return Transforms[WorldToLocalIndex]; }
+	const FTransform3f& GetTransform(ENiagaraCoordinateSpace SourceSpace, ENiagaraCoordinateSpace DestinationSpace) const;
 
 	FQuat4f   TransformRotation(ENiagaraCoordinateSpace SourceSpace, ENiagaraCoordinateSpace DestinationSpace, FQuat4f Rotation) const;
 	FVector3f TransformPosition(ENiagaraCoordinateSpace SourceSpace, ENiagaraCoordinateSpace DestinationSpace, FVector3f Position) const;
diff --git a/Engine/Plugins/FX/Niagara/Source/Niagara/Internal/Stateless/NiagaraStatelessParticleSimContext.h b/Engine/Plugins/FX/Niagara/Source/Niagara/Internal/Stateless/NiagaraStatelessParticleSimContext.h
--- a/Engine/Plugins/FX/Niagara/Source/Niagara/Internal/Stateless/NiagaraStatelessParticleSimContext.h
+++ b/Engine/Plugins/FX/Niagara/Source/Niagara/Internal/Stateless/NiagaraStatelessParticleSimContext.h
@@ -50,6 +50,8 @@ public:
 	float GetDeltaTime() const { return DeltaTime; }
 	float GetInvDeltaTime() const { return InvDeltaTime; }
 
+	const FQuat4f& GetToSimulationRotation(ENiagaraCoordinateSpace SourceSpace) const;
+
 	// Get a pointer to required particle components
 	int32* GetParticleAlive() const { return static_cast<int32*>(RequiredComponents[int32(EParticleComponent::Alive)]); }
 	float* GetParticleLifetime() { return static_cast<float*>(RequiredComponents[int32(EParticleComponent::Lifetime)]); }
diff --git a/Engine/Plugins/FX/Niagara/Source/Niagara/Private/Stateless/NiagaraStatelessCommon.cpp b/Engine/Plugins/FX/Niagara/Source/Niagara/Private/Stateless/NiagaraStatelessCommon.cpp
--- a/Engine/Plugins/FX/Niagara/Source/Niagara/Private/Stateless/NiagaraStatelessCommon.cpp
+++ b/Engine/Plugins/FX/Niagara/Source/Niagara/Private/Stateless/NiagaraStatelessCommon.cpp
@@ -3,6 +3,7 @@
 #include "Stateless/NiagaraStatelessCommon.h"
 
 #include "Stateless/NiagaraStatelessEmitterTemplate.h"
+#include "NiagaraConstants.h"
 #include "NiagaraSystem.h"
 
 #include "UObject/UObjectIterator.h"
@@ -40,6 +41,13 @@ bool FNiagaraStatelessSpaceTransforms::UpdateTransforms(const FTransform3f& Loca
 	return true;
 }
 
+const FTransform3f& FNiagaraStatelessSpaceTransforms::GetTransform(ENiagaraCoordinateSpace SourceSpace, ENiagaraCoordinateSpace DestinationSpace) const
+{
+	check(int(SourceSpace) >= 0 && int(SourceSpace) < TransformRemapSize && int(DestinationSpace) >= 0 && int(DestinationSpace) < TransformRemapSize);
+	const int8 i = TransformRemap[int(SourceSpace)][int(DestinationSpace)];
+	return i == INDEX_NONE ? FTransform3f::Identity : Transforms[i];
+}
+
 FQuat4f FNiagaraStatelessSpaceTransforms::TransformRotation(ENiagaraCoordinateSpace SourceSpace, ENiagaraCoordinateSpace DestinationSpace, FQuat4f Rotation) const
 {
 	check(int(SourceSpace) >= 0 && int(SourceSpace) < TransformRemapSize && int(DestinationSpace) >= 0 && int(DestinationSpace) < TransformRemapSize);
diff --git a/Engine/Plugins/FX/Niagara/Source/Niagara/Private/Stateless/NiagaraStatelessEmitterInstance.cpp b/Engine/Plugins/FX/Niagara/Source/Niagara/Private/Stateless/NiagaraStatelessEmitterInstance.cpp
--- a/Engine/Plugins/FX/Niagara/Source/Niagara/Private/Stateless/NiagaraStatelessEmitterInstance.cpp
+++ b/Engine/Plugins/FX/Niagara/Source/Niagara/Private/Stateless/NiagaraStatelessEmitterInstance.cpp
@@ -69,6 +69,13 @@ namespace NiagaraStatelessEmitterInstancePrivate
 		}
 		return (OutMin > 0.0f) || (OutMax > 0.0f);
 	}
+
+	void SetShaderParameterTransforms(const FNiagaraStatelessSpaceTransforms& Transforms, NiagaraStateless::FCommonShaderParameters* ShaderParameters)
+	{
+		ShaderParameters->Common_ToSimulationRotations[int(ENiagaraCoordinateSpace::Simulation)] = Transforms.GetTransform(ENiagaraCoordinateSpace::Simulation, ENiagaraCoordinateSpace::Simulation).GetRotation();
+		ShaderParameters->Common_ToSimulationRotations[int(ENiagaraCoordinateSpace::World)] = Transforms.GetTransform(ENiagaraCoordinateSpace::World, ENiagaraCoordinateSpace::Simulation).GetRotation();
+		ShaderParameters->Common_ToSimulationRotations[int(ENiagaraCoordinateSpace::Local)] = Transforms.GetTransform(ENiagaraCoordinateSpace::Local, ENiagaraCoordinateSpace::Simulation).GetRotation();
+	}
 }
 
 FNiagaraDataBuffer* NiagaraStateless::FEmitterInstance_RT::GetDataToRender(FRHICommandListBase& RHICmdList, bool bIsLowLatencyTranslucent) const
@@ -146,6 +153,7 @@ void FNiagaraStatelessEmitterInstance::Init(int32 InEmitterIndex)
 	{
 		ShaderParameters.Reset(WeakStatelessEmitter->AllocateShaderParameters(EmitterTransforms, RendererBindings));
 		ShaderParameters->Common_RandomSeed = RandomSeed;
+		NiagaraStatelessEmitterInstancePrivate::SetShaderParameterTransforms(EmitterTransforms, ShaderParameters.Get());
 	}
 	else
 	{
@@ -157,6 +165,7 @@ void FNiagaraStatelessEmitterInstance::Init(int32 InEmitterIndex)
 		RenderThreadDataPtr->ExecutionState = ENiagaraExecutionState::Active;
 		RenderThreadDataPtr->ShaderParameters.Reset(WeakStatelessEmitter->AllocateShaderParameters(EmitterTransforms, RendererBindings));
 		RenderThreadDataPtr->ShaderParameters->Common_RandomSeed = RandomSeed;
+		NiagaraStatelessEmitterInstancePrivate::SetShaderParameterTransforms(EmitterTransforms, RenderThreadDataPtr->ShaderParameters.Get());
 
 		ENQUEUE_RENDER_COMMAND(FInitStatelessEmitter)(
 			[RenderThreadData=RenderThreadDataPtr.Get(), ComputeInterface = ParentSystemInstance->GetComputeDispatchInterface()](FRHICommandListImmediate& RHICmdList)
@@ -566,6 +575,7 @@ void FNiagaraStatelessEmitterInstance::UpdateSimulationData(float DeltaSeconds)
 		{
 			ShaderParameters.Reset(WeakStatelessEmitter->AllocateShaderParameters(EmitterTransforms, RendererBindings));
 			ShaderParameters->Common_RandomSeed = RandomSeed;
+			NiagaraStatelessEmitterInstancePrivate::SetShaderParameterTransforms(EmitterTransforms, ShaderParameters.Get());
 		}
 
 		FNiagaraDataBuffer& DataBuffer = ParticleDataSet->BeginSimulate();
@@ -607,6 +617,7 @@ void FNiagaraStatelessEmitterInstance::UpdateSimulationData(float DeltaSeconds)
 		{
 			DataForRenderThread.ShaderParameters = WeakStatelessEmitter->AllocateShaderParameters(EmitterTransforms, RendererBindings);
 			DataForRenderThread.ShaderParameters->Common_RandomSeed = RandomSeed;
+			NiagaraStatelessEmitterInstancePrivate::SetShaderParameterTransforms(EmitterTransforms, DataForRenderThread.ShaderParameters);
 		}
 
 		if (bSpawnInfosDirty)
diff --git a/Engine/Plugins/FX/Niagara/Source/Niagara/Private/Stateless/NiagaraStatelessParticleSimContext.cpp b/Engine/Plugins/FX/Niagara/Source/Niagara/Private/Stateless/NiagaraStatelessParticleSimContext.cpp
--- a/Engine/Plugins/FX/Niagara/Source/Niagara/Private/Stateless/NiagaraStatelessParticleSimContext.cpp
+++ b/Engine/Plugins/FX/Niagara/Source/Niagara/Private/Stateless/NiagaraStatelessParticleSimContext.cpp
@@ -235,6 +235,13 @@ namespace NiagaraStateless
 		}
 	}
 
+	const FQuat4f& FParticleSimulationContext::GetToSimulationRotation(ENiagaraCoordinateSpace SourceSpace) const
+	{
+		using namespace NiagaraStateless;
+		const FCommonShaderParameters* CommonShaderParameters = reinterpret_cast<const FCommonShaderParameters*>(ShaderParametersData);
+		return CommonShaderParameters->Common_ToSimulationRotations[int(SourceSpace)];
+	}
+
 	uint32 FParticleSimulationContext::RandomUInt(uint32 iInstance, uint32 RandomSeedOffset) const
 	{
 		const uint32 UniqueIndex = uint32(GetParticleUniqueIndex()[iInstance]);
diff --git a/Engine/Plugins/FX/Niagara/Source/NiagaraShader/Internal/Stateless/NiagaraStatelessSimulationShader.h b/Engine/Plugins/FX/Niagara/Source/NiagaraShader/Internal/Stateless/NiagaraStatelessSimulationShader.h
--- a/Engine/Plugins/FX/Niagara/Source/NiagaraShader/Internal/Stateless/NiagaraStatelessSimulationShader.h
+++ b/Engine/Plugins/FX/Niagara/Source/NiagaraShader/Internal/Stateless/NiagaraStatelessSimulationShader.h
@@ -26,6 +26,7 @@ namespace NiagaraStateless
 		SHADER_PARAMETER(float,					Common_SimulationTime)
 		SHADER_PARAMETER(float,					Common_SimulationDeltaTime)
 		SHADER_PARAMETER(float,					Common_SimulationInvDeltaTime)
+		SHADER_PARAMETER_ARRAY(FQuat4f,			Common_ToSimulationRotations, [3])
 
 		SHADER_PARAMETER(uint32,				Common_OutputBufferStride)
 		SHADER_PARAMETER(uint32,				Common_GPUCountBufferOffset)
