- Light & decal attribute modules for stateless
- Enable decal renderer for stateless

diff --git a/Engine/Plugins/FX/Niagara/Shaders/Private/Stateless/NiagaraStatelessSimulationDefault.usf b/Engine/Plugins/FX/Niagara/Shaders/Private/Stateless/NiagaraStatelessSimulationDefault.usf
--- a/Engine/Plugins/FX/Niagara/Shaders/Private/Stateless/NiagaraStatelessSimulationDefault.usf
+++ b/Engine/Plugins/FX/Niagara/Shaders/Private/Stateless/NiagaraStatelessSimulationDefault.usf
@@ -50,7 +50,9 @@
 	PARTICLE_MODULE(SolveVelocitiesAndForces) \
 	/* Post Solve */ \
 	PARTICLE_MODULE(CameraOffset) \
+	PARTICLE_MODULE(Null) /* DecalAttributes */ \
 	PARTICLE_MODULE(DynamicMaterialParameters) \
+	PARTICLE_MODULE(Null) /* LightAttributes */ \
 	PARTICLE_MODULE(MeshIndex) \
 	PARTICLE_MODULE(MeshRotationRate) \
 	PARTICLE_MODULE(ScaleColor) \
diff --git a/Engine/Plugins/FX/Niagara/Source/Niagara/Internal/Stateless/Modules/NiagaraStatelessModule_DecalAttributes.h b/Engine/Plugins/FX/Niagara/Source/Niagara/Internal/Stateless/Modules/NiagaraStatelessModule_DecalAttributes.h
new file mode 100644
--- /dev/null
+++ b/Engine/Plugins/FX/Niagara/Source/Niagara/Internal/Stateless/Modules/NiagaraStatelessModule_DecalAttributes.h
@@ -0,0 +1,50 @@
+// Copyright Epic Games, Inc. All Rights Reserved.
+
+#pragma once
+
+#include "Stateless/NiagaraStatelessCommon.h"
+#include "Stateless/NiagaraStatelessModule.h"
+
+#include "NiagaraStatelessModule_DecalAttributes.generated.h"
+
+UCLASS(MinimalAPI, EditInlineNew, Experimental, meta = (DisplayName = "Decal Attributes"))
+class UNiagaraStatelessModule_DecalAttributes : public UNiagaraStatelessModule
+{
+	GENERATED_BODY()
+
+public:
+	UPROPERTY(EditAnywhere, Category = "Parameters", meta = (InlineEditConditionToggle))
+	bool bApplyOrientation = false;
+
+	UPROPERTY(EditAnywhere, Category = "Parameters", meta = (InlineEditConditionToggle))
+	bool bApplySize = false;
+
+	UPROPERTY(EditAnywhere, Category = "Parameters", meta = (InlineEditConditionToggle))
+	bool bApplyFade = false;
+
+	UPROPERTY(EditAnywhere, Category = "Parameters", meta = (DisplayName = "Orientation", EditCondition = "bApplyOrientation"))
+	FNiagaraDistributionVector3 Orientation = FNiagaraDistributionVector3(FVector3f::ZeroVector);
+
+	UPROPERTY(EditAnywhere, Category = "Parameters", meta = (EditCondition = "IsOrientationEnabled()", EditConditionHides, SegmentedDisplay))
+	ENiagaraCoordinateSpace OrientationCoordinateSpace = ENiagaraCoordinateSpace::Local;
+
+	UPROPERTY(EditAnywhere, Category = "Parameters", meta = (DisplayName = "Size", EditCondition = "bApplySize"))
+	FNiagaraDistributionVector3 Size = FNiagaraDistributionVector3(FVector3f(50.0f, 50.0f, 50.0f));
+
+	UPROPERTY(EditAnywhere, Category = "Parameters", meta = (DisplayName = "Fade", EditCondition = "bApplyFade"))
+	FNiagaraDistributionFloat Fade = FNiagaraDistributionFloat(1.0f);
+
+	virtual ENiagaraStatelessFeatureMask GetFeatureMask() const override;
+	virtual void BuildEmitterData(const FNiagaraStatelessEmitterDataBuildContext& BuildContext) const override;
+	virtual void SetShaderParameters(const FNiagaraStatelessSetShaderParameterContext& SetShaderParameterContext) const override;
+
+	UFUNCTION()
+	bool IsOrientationEnabled() const { return bApplyOrientation; }
+
+#if WITH_EDITOR
+	virtual bool CanDisableModule() const override { return true; }
+#endif
+#if WITH_EDITORONLY_DATA
+	virtual void GetOutputVariables(TArray<FNiagaraVariableBase>& OutVariables) const override;
+#endif
+};
diff --git a/Engine/Plugins/FX/Niagara/Source/Niagara/Internal/Stateless/Modules/NiagaraStatelessModule_LightAttributes.h b/Engine/Plugins/FX/Niagara/Source/Niagara/Internal/Stateless/Modules/NiagaraStatelessModule_LightAttributes.h
new file mode 100644
--- /dev/null
+++ b/Engine/Plugins/FX/Niagara/Source/Niagara/Internal/Stateless/Modules/NiagaraStatelessModule_LightAttributes.h
@@ -0,0 +1,56 @@
+// Copyright Epic Games, Inc. All Rights Reserved.
+
+#pragma once
+
+#include "Stateless/NiagaraStatelessModule.h"
+
+#include "NiagaraStatelessModule_LightAttributes.generated.h"
+
+// Sets attributes that the light renderer can consume
+UCLASS(MinimalAPI, EditInlineNew, meta = (DisplayName = "Light Attributes"))
+class UNiagaraStatelessModule_LightAttributes : public UNiagaraStatelessModule
+{
+	GENERATED_BODY()
+
+public:
+	UPROPERTY(EditAnywhere, Category = "Parameters", meta = (InlineEditConditionToggle))
+	uint32 bApplyRadius : 1 = false;
+
+	UPROPERTY(EditAnywhere, Category = "Parameters", meta = (InlineEditConditionToggle))
+	uint32 bApplyFalloffExponent : 1 = false;
+
+	UPROPERTY(EditAnywhere, Category = "Parameters", meta = (InlineEditConditionToggle))
+	uint32 bApplyDiffuseScale : 1 = false;
+
+	UPROPERTY(EditAnywhere, Category = "Parameters", meta = (InlineEditConditionToggle))
+	uint32 bApplySpecularScale : 1 = false;
+
+	UPROPERTY(EditAnywhere, Category = "Parameters", meta = (InlineEditConditionToggle))
+	uint32 bApplyVolumetricScattering : 1 = false;
+
+	UPROPERTY(EditAnywhere, Category = "Parameters", meta = (EditCondition = "bApplyRadius"))
+	FNiagaraDistributionFloat Radius = FNiagaraDistributionFloat(100.0f);
+
+	UPROPERTY(EditAnywhere, Category = "Parameters", meta = (EditCondition = "bApplyFalloffExponent"))
+	FNiagaraDistributionFloat FalloffExponent = FNiagaraDistributionFloat(1.0f);
+
+	UPROPERTY(EditAnywhere, Category = "Parameters", meta = (EditCondition = "bApplyDiffuseScale"))
+	FNiagaraDistributionFloat DiffuseScale = FNiagaraDistributionFloat(1.0f);
+
+	UPROPERTY(EditAnywhere, Category = "Parameters", meta = (EditCondition = "bApplySpecularScale"))
+	FNiagaraDistributionFloat SpecularScale = FNiagaraDistributionFloat(1.0f);
+
+	UPROPERTY(EditAnywhere, Category = "Parameters", meta = (EditCondition = "bApplyVolumetricScattering"))
+	FNiagaraDistributionFloat VolumetricScattering = FNiagaraDistributionFloat(0.0f);
+
+	virtual ENiagaraStatelessFeatureMask GetFeatureMask() const override;
+	virtual void BuildEmitterData(const FNiagaraStatelessEmitterDataBuildContext& BuildContext) const override;
+	virtual void SetShaderParameters(const FNiagaraStatelessSetShaderParameterContext& SetShaderParameterContext) const override;
+
+#if WITH_EDITOR
+	virtual bool CanDisableModule() const override { return true; }
+#endif
+#if WITH_EDITORONLY_DATA
+	virtual void GetOutputVariables(TArray<FNiagaraVariableBase>& OutVariables) const override;
+#endif
+};
diff --git a/Engine/Plugins/FX/Niagara/Source/Niagara/Internal/Stateless/NiagaraStatelessCommon.h b/Engine/Plugins/FX/Niagara/Source/Niagara/Internal/Stateless/NiagaraStatelessCommon.h
--- a/Engine/Plugins/FX/Niagara/Source/Niagara/Internal/Stateless/NiagaraStatelessCommon.h
+++ b/Engine/Plugins/FX/Niagara/Source/Niagara/Internal/Stateless/NiagaraStatelessCommon.h
@@ -107,4 +107,6 @@ namespace NiagaraStatelessCommon
 {
 	extern void Initialize();
 	extern void UpdateSettings();
+
+	extern FNiagaraVariableBase ConvertParticleVariableToStateless(const FNiagaraVariableBase& InVariable);
 }
diff --git a/Engine/Plugins/FX/Niagara/Source/Niagara/Private/NiagaraDecalRendererProperties.cpp b/Engine/Plugins/FX/Niagara/Source/Niagara/Private/NiagaraDecalRendererProperties.cpp
--- a/Engine/Plugins/FX/Niagara/Source/Niagara/Private/NiagaraDecalRendererProperties.cpp
+++ b/Engine/Plugins/FX/Niagara/Source/Niagara/Private/NiagaraDecalRendererProperties.cpp
@@ -33,30 +33,6 @@ namespace NiagaraDecalRendererPropertiesLocal
 		return Variable;
 	}
 
-	static FNiagaraVariable GetDecalOrientationVariable()
-	{
-		static FNiagaraVariable Variable = MakeNiagaraVariableWithValue(FNiagaraTypeDefinition::GetQuatDef(), TEXT("Particles.DecalOrientation"), UNiagaraDecalRendererProperties::GetDefaultOrientation());
-		return Variable;
-	}
-
-	static FNiagaraVariable GetDecalSizeVariable()
-	{
-		static FNiagaraVariable Variable = MakeNiagaraVariableWithValue(FNiagaraTypeDefinition::GetVec3Def(), TEXT("Particles.DecalSize"), UNiagaraDecalRendererProperties::GetDefaultDecalSize());
-		return Variable;
-	}
-
-	static FNiagaraVariable GetDecalFadeVariable()
-	{
-		static FNiagaraVariable Variable = MakeNiagaraVariableWithValue(FNiagaraTypeDefinition::GetFloatDef(), TEXT("Particles.DecalFade"), UNiagaraDecalRendererProperties::GetDefaultDecalFade());
-		return Variable;
-	}
-
-	static FNiagaraVariable GetDecalVisibleVariable()
-	{
-		static FNiagaraVariable Variable = MakeNiagaraVariableWithValue(FNiagaraTypeDefinition::GetBoolDef(), TEXT("Particles.DecalVisible"), UNiagaraDecalRendererProperties::GetDefaultDecalVisible());
-		return Variable;
-	}
-
 	static void SetupBindings(UNiagaraDecalRendererProperties* Props)
 	{
 		if (Props->PositionBinding.IsValid())
@@ -64,12 +40,12 @@ namespace NiagaraDecalRendererPropertiesLocal
 			return;
 		}
 		Props->PositionBinding = FNiagaraConstants::GetAttributeDefaultBinding(SYS_PARAM_PARTICLES_POSITION);
-		Props->DecalOrientationBinding.Setup(GetDecalOrientationVariable(), GetDecalOrientationVariable());
-		Props->DecalSizeBinding.Setup(GetDecalSizeVariable(), GetDecalSizeVariable());
-		Props->DecalFadeBinding.Setup(GetDecalFadeVariable(), GetDecalFadeVariable());
+		Props->DecalOrientationBinding.Setup(UNiagaraDecalRendererProperties::GetDecalOrientationVariable(), UNiagaraDecalRendererProperties::GetDecalOrientationVariable());
+		Props->DecalSizeBinding.Setup(UNiagaraDecalRendererProperties::GetDecalSizeVariable(), UNiagaraDecalRendererProperties::GetDecalSizeVariable());
+		Props->DecalFadeBinding.Setup(UNiagaraDecalRendererProperties::GetDecalFadeVariable(), UNiagaraDecalRendererProperties::GetDecalFadeVariable());
 		Props->DecalSortOrderBinding = FNiagaraConstants::GetAttributeDefaultBinding(SYS_PARAM_PARTICLES_UNIQUE_ID);
 		Props->DecalColorBinding = FNiagaraConstants::GetAttributeDefaultBinding(SYS_PARAM_PARTICLES_COLOR);
-		Props->DecalVisibleBinding.Setup(GetDecalVisibleVariable(), GetDecalVisibleVariable());
+		Props->DecalVisibleBinding.Setup(UNiagaraDecalRendererProperties::GetDecalVisibleVariable(), UNiagaraDecalRendererProperties::GetDecalVisibleVariable());
 		Props->RendererVisibilityTagBinding = FNiagaraConstants::GetAttributeDefaultBinding(SYS_PARAM_PARTICLES_VISIBILITY_TAG);
 
 	#if WITH_EDITORONLY_DATA
@@ -390,4 +366,33 @@ UMaterialInterface* UNiagaraDecalRendererProperties::GetMaterial(const FNiagaraE
 	return MaterialInterface ? MaterialInterface : ToRawPtr(Material);
 }
 
+FNiagaraVariable UNiagaraDecalRendererProperties::GetDecalOrientationVariable()
+{
+	using namespace NiagaraDecalRendererPropertiesLocal;
+	static FNiagaraVariable Variable = MakeNiagaraVariableWithValue(FNiagaraTypeDefinition::GetQuatDef(), TEXT("Particles.DecalOrientation"), UNiagaraDecalRendererProperties::GetDefaultOrientation());
+	return Variable;
+}
+
+FNiagaraVariable UNiagaraDecalRendererProperties::GetDecalSizeVariable()
+{
+	using namespace NiagaraDecalRendererPropertiesLocal;
+	static FNiagaraVariable Variable = MakeNiagaraVariableWithValue(FNiagaraTypeDefinition::GetVec3Def(), TEXT("Particles.DecalSize"), UNiagaraDecalRendererProperties::GetDefaultDecalSize());
+	return Variable;
+}
+
+FNiagaraVariable UNiagaraDecalRendererProperties::GetDecalFadeVariable()
+{
+	using namespace NiagaraDecalRendererPropertiesLocal;
+	static FNiagaraVariable Variable = MakeNiagaraVariableWithValue(FNiagaraTypeDefinition::GetFloatDef(), TEXT("Particles.DecalFade"), UNiagaraDecalRendererProperties::GetDefaultDecalFade());
+	return Variable;
+}
+
+FNiagaraVariable UNiagaraDecalRendererProperties::GetDecalVisibleVariable()
+{
+	using namespace NiagaraDecalRendererPropertiesLocal;
+	static FNiagaraVariable Variable = MakeNiagaraVariableWithValue(FNiagaraTypeDefinition::GetBoolDef(), TEXT("Particles.DecalVisible"), UNiagaraDecalRendererProperties::GetDefaultDecalVisible());
+	return Variable;
+}
+
+
 #undef LOCTEXT_NAMESPACE
diff --git a/Engine/Plugins/FX/Niagara/Source/Niagara/Private/Stateless/Modules/NiagaraStatelessModule_DecalAttributes.cpp b/Engine/Plugins/FX/Niagara/Source/Niagara/Private/Stateless/Modules/NiagaraStatelessModule_DecalAttributes.cpp
new file mode 100644
--- /dev/null
+++ b/Engine/Plugins/FX/Niagara/Source/Niagara/Private/Stateless/Modules/NiagaraStatelessModule_DecalAttributes.cpp
@@ -0,0 +1,160 @@
+// Copyright Epic Games, Inc. All Rights Reserved.
+
+#include "Stateless/Modules/NiagaraStatelessModule_DecalAttributes.h"
+
+#include "Stateless/NiagaraStatelessCommon.h"
+#include "Stateless/NiagaraStatelessEmitterDataBuildContext.h"
+#include "Stateless/NiagaraStatelessParticleSimContext.h"
+#include "NiagaraDecalRendererProperties.h"
+
+#include "NiagaraConstants.h"
+#include "NiagaraModule.h"
+
+namespace NSMDecalAttributesPrivate
+{
+	enum class EModuleAttribute
+	{
+		Orientation,
+		Size,
+		Fade,
+		Num
+	};
+
+	const FNiagaraVariableBase* GetAttributeVariable(EModuleAttribute Attribute, const UNiagaraStatelessModule_DecalAttributes* Module)
+	{
+		static const FNiagaraVariableBase DecalOrientationVariable = NiagaraStatelessCommon::ConvertParticleVariableToStateless(UNiagaraDecalRendererProperties::GetDecalOrientationVariable());
+		static const FNiagaraVariableBase DecalSizeVariable = NiagaraStatelessCommon::ConvertParticleVariableToStateless(UNiagaraDecalRendererProperties::GetDecalSizeVariable());
+		static const FNiagaraVariableBase DecalFadeVariable = NiagaraStatelessCommon::ConvertParticleVariableToStateless(UNiagaraDecalRendererProperties::GetDecalFadeVariable());
+
+		switch (Attribute)
+		{
+			case EModuleAttribute::Orientation:		return Module->bApplyOrientation ? &DecalOrientationVariable : nullptr;
+			case EModuleAttribute::Size:			return Module->bApplySize ? &DecalSizeVariable : nullptr;
+			case EModuleAttribute::Fade:			return Module->bApplyFade ? &DecalFadeVariable : nullptr;
+		}
+		return nullptr;
+	}
+
+	struct FModuleBuiltData
+	{
+		FModuleBuiltData()
+		{
+			for (int i = 0; i < int(EModuleAttribute::Num); ++i)
+			{
+				AttributeDistributionParameters[i] = FUintVector3::ZeroValue;
+				AttributeOffset[i] = INDEX_NONE;
+			}
+		}
+
+		ENiagaraCoordinateSpace OrientationCoordinateSpace = ENiagaraCoordinateSpace::Local;
+		FUintVector3			AttributeDistributionParameters[int(EModuleAttribute::Num)];
+		int						AttributeOffset[int(EModuleAttribute::Num)];
+	};
+
+	static void ParticleSimulate(const NiagaraStateless::FParticleSimulationContext& ParticleSimulationContext)
+	{
+		using namespace NiagaraStateless;
+
+		const FModuleBuiltData* ModuleBuiltData = ParticleSimulationContext.ReadBuiltData<FModuleBuiltData>();
+		const float* NormalizedAgeData = ParticleSimulationContext.GetParticleNormalizedAge();
+
+		const FQuat4f OrientationRotation = ParticleSimulationContext.GetToSimulationRotation(ModuleBuiltData->OrientationCoordinateSpace);
+
+		for (uint32 i = 0; i < ParticleSimulationContext.GetNumInstances(); ++i)
+		{
+			if (ModuleBuiltData->AttributeOffset[int(EModuleAttribute::Orientation)] != INDEX_NONE)
+			{
+				FStatelessDistributionSampler<FVector3f> DistributionSampler(ParticleSimulationContext, ModuleBuiltData->AttributeDistributionParameters[int(EModuleAttribute::Orientation)], i, 0);
+				const FVector3f Rotation = DistributionSampler.GetValue(ParticleSimulationContext, NormalizedAgeData[i]);
+				const FQuat4f SourceValue = FRotator3f(Rotation.X, Rotation.Y, Rotation.Z).Quaternion();
+				const FQuat4f Value = OrientationRotation * SourceValue;
+				ParticleSimulationContext.WriteParticleVariable(ModuleBuiltData->AttributeOffset[int(EModuleAttribute::Orientation)], i, Value);
+			}
+			if (ModuleBuiltData->AttributeOffset[int(EModuleAttribute::Size)] != INDEX_NONE)
+			{
+				FStatelessDistributionSampler<FVector3f> DistributionSampler(ParticleSimulationContext, ModuleBuiltData->AttributeDistributionParameters[int(EModuleAttribute::Size)], i, 0);
+				const FVector3f Value = DistributionSampler.GetValue(ParticleSimulationContext, NormalizedAgeData[i]);
+				ParticleSimulationContext.WriteParticleVariable(ModuleBuiltData->AttributeOffset[int(EModuleAttribute::Size)], i, Value);
+			}
+			if (ModuleBuiltData->AttributeOffset[int(EModuleAttribute::Fade)] != INDEX_NONE)
+			{
+				FStatelessDistributionSampler<float> DistributionSampler(ParticleSimulationContext, ModuleBuiltData->AttributeDistributionParameters[int(EModuleAttribute::Fade)], i, 0);
+				const float Value = DistributionSampler.GetValue(ParticleSimulationContext, NormalizedAgeData[i]);
+				ParticleSimulationContext.WriteParticleVariable(ModuleBuiltData->AttributeOffset[int(EModuleAttribute::Fade)], i, Value);
+			}
+		}
+	}
+}
+
+ENiagaraStatelessFeatureMask UNiagaraStatelessModule_DecalAttributes::GetFeatureMask() const
+{
+	return ENiagaraStatelessFeatureMask::ExecuteCPU;
+}
+
+void UNiagaraStatelessModule_DecalAttributes::BuildEmitterData(const FNiagaraStatelessEmitterDataBuildContext& BuildContext) const
+{
+	using namespace NSMDecalAttributesPrivate;
+
+	FModuleBuiltData* BuiltData = BuildContext.AllocateBuiltData<FModuleBuiltData>();
+	if (!IsModuleEnabled())
+	{
+		return;
+	}
+
+	// Gather our attribute bindings
+	{
+		bool bAnyValidAttributes = false;
+		for (int i = 0; i < int(EModuleAttribute::Num); ++i)
+		{
+			if (const FNiagaraVariableBase* Variable = GetAttributeVariable(EModuleAttribute(i), this))
+			{
+				BuiltData->AttributeOffset[i] = BuildContext.FindParticleVariableIndex(*Variable);
+				bAnyValidAttributes |= BuiltData->AttributeOffset[i] != INDEX_NONE;
+			}
+		}
+
+		if ( !bAnyValidAttributes )
+		{
+			return;
+		}
+	}
+
+	// Build distributions
+	if (bApplyOrientation)
+	{
+		BuiltData->OrientationCoordinateSpace = OrientationCoordinateSpace;
+		BuiltData->AttributeDistributionParameters[int(EModuleAttribute::Orientation)] = BuildContext.AddDistribution(Orientation);
+	}
+	if (bApplySize)
+	{
+		BuiltData->AttributeDistributionParameters[int(EModuleAttribute::Size)] = BuildContext.AddDistribution(Size);
+	}
+	if (bApplyFade)
+	{
+		BuiltData->AttributeDistributionParameters[int(EModuleAttribute::Fade)] = BuildContext.AddDistribution(Fade);
+	}
+
+	BuildContext.AddParticleSimulationExecSimulate(&ParticleSimulate);
+}
+
+void UNiagaraStatelessModule_DecalAttributes::SetShaderParameters(const FNiagaraStatelessSetShaderParameterContext& SetShaderParameterContext) const
+{
+	using namespace NSMDecalAttributesPrivate;
+
+	const FModuleBuiltData* ModuleBuiltData = SetShaderParameterContext.ReadBuiltData<FModuleBuiltData>();
+}
+
+#if WITH_EDITORONLY_DATA
+void UNiagaraStatelessModule_DecalAttributes::GetOutputVariables(TArray<FNiagaraVariableBase>& OutVariables) const
+{
+	using namespace NSMDecalAttributesPrivate;
+
+	for (int i = 0; i < int(EModuleAttribute::Num); ++i)
+	{
+		if (const FNiagaraVariableBase* Variable = GetAttributeVariable(EModuleAttribute(i), this))
+		{
+			OutVariables.AddUnique(*Variable);
+		}
+	}
+}
+#endif
diff --git a/Engine/Plugins/FX/Niagara/Source/Niagara/Private/Stateless/Modules/NiagaraStatelessModule_LightAttributes.cpp b/Engine/Plugins/FX/Niagara/Source/Niagara/Private/Stateless/Modules/NiagaraStatelessModule_LightAttributes.cpp
new file mode 100644
--- /dev/null
+++ b/Engine/Plugins/FX/Niagara/Source/Niagara/Private/Stateless/Modules/NiagaraStatelessModule_LightAttributes.cpp
@@ -0,0 +1,159 @@
+// Copyright Epic Games, Inc. All Rights Reserved.
+
+#include "Stateless/Modules/NiagaraStatelessModule_LightAttributes.h"
+
+#include "Stateless/NiagaraStatelessCommon.h"
+#include "Stateless/NiagaraStatelessEmitterDataBuildContext.h"
+#include "Stateless/NiagaraStatelessParticleSimContext.h"
+
+#include "NiagaraConstants.h"
+#include "NiagaraModule.h"
+
+namespace NSMLightAttributesPrivate
+{
+	enum class EModuleAttribute
+	{
+		Radius,
+		Falloff,
+		DiffuseScale,
+		SpecularScale,
+		VolumetricScale,
+		Num
+	};
+
+	const FNiagaraVariableBase* GetAttributeVariable(EModuleAttribute Attribute, const UNiagaraStatelessModule_LightAttributes* Module)
+	{
+		static const FNiagaraVariableBase LightRadiusVariable = NiagaraStatelessCommon::ConvertParticleVariableToStateless(SYS_PARAM_PARTICLES_LIGHT_RADIUS);
+		static const FNiagaraVariableBase LightFalloffExponent = NiagaraStatelessCommon::ConvertParticleVariableToStateless(SYS_PARAM_PARTICLES_LIGHT_EXPONENT);
+		static const FNiagaraVariableBase LightDiffuseScale = NiagaraStatelessCommon::ConvertParticleVariableToStateless(SYS_PARAM_PARTICLES_LIGHT_DIFFUSE_SCALE);
+		static const FNiagaraVariableBase LightSpecularScale = NiagaraStatelessCommon::ConvertParticleVariableToStateless(SYS_PARAM_PARTICLES_LIGHT_SPECULAR_SCALE);
+		static const FNiagaraVariableBase LightVolumetricScattering = NiagaraStatelessCommon::ConvertParticleVariableToStateless(SYS_PARAM_PARTICLES_LIGHT_VOLUMETRIC_SCATTERING);
+
+		switch (Attribute)
+		{
+			case EModuleAttribute::Radius:			return Module->bApplyRadius ? &LightRadiusVariable : nullptr;
+			case EModuleAttribute::Falloff:			return Module->bApplyFalloffExponent ? &LightFalloffExponent : nullptr;
+			case EModuleAttribute::DiffuseScale:	return Module->bApplyDiffuseScale ? &LightDiffuseScale : nullptr;
+			case EModuleAttribute::SpecularScale:	return Module->bApplySpecularScale ? &LightSpecularScale : nullptr;
+			case EModuleAttribute::VolumetricScale:	return Module->bApplyVolumetricScattering ? &LightVolumetricScattering : nullptr;
+		}
+		return nullptr;
+	}
+
+	const FNiagaraDistributionFloat* GetAttributeDistribution(EModuleAttribute Attribute, const UNiagaraStatelessModule_LightAttributes* Module)
+	{
+		switch (Attribute)
+		{
+			case EModuleAttribute::Radius:			return Module->bApplyRadius ? &Module->Radius : nullptr;
+			case EModuleAttribute::Falloff:			return Module->bApplyFalloffExponent ? &Module->FalloffExponent : nullptr;
+			case EModuleAttribute::DiffuseScale:	return Module->bApplyDiffuseScale ? &Module->DiffuseScale : nullptr;
+			case EModuleAttribute::SpecularScale:	return Module->bApplySpecularScale ? &Module->SpecularScale : nullptr;
+			case EModuleAttribute::VolumetricScale:	return Module->bApplyVolumetricScattering ? &Module->VolumetricScattering : nullptr;
+		}
+		return nullptr;
+	}
+
+	struct FModuleBuiltData
+	{
+		FModuleBuiltData()
+		{
+			for (int i = 0; i < int(EModuleAttribute::Num); ++i)
+			{
+				AttributeDistributionParameters[i] = FUintVector3::ZeroValue;
+				AttributeOffset[i] = INDEX_NONE;
+			}
+		}
+
+		FUintVector3	AttributeDistributionParameters[int(EModuleAttribute::Num)];
+		int				AttributeOffset[int(EModuleAttribute::Num)];
+	};
+
+	static void ParticleSimulate(const NiagaraStateless::FParticleSimulationContext& ParticleSimulationContext)
+	{
+		using namespace NiagaraStateless;
+
+		const FModuleBuiltData* ModuleBuiltData = ParticleSimulationContext.ReadBuiltData<FModuleBuiltData>();
+		const float* NormalizedAgeData = ParticleSimulationContext.GetParticleNormalizedAge();
+
+		for (uint32 i=0; i < ParticleSimulationContext.GetNumInstances(); ++i)
+		{
+			for ( int a=0; a < int(EModuleAttribute::Num); ++a )
+			{
+				if (ModuleBuiltData->AttributeOffset[a] != INDEX_NONE)
+				{
+					FStatelessDistributionSampler<float> DistributionSampler(ParticleSimulationContext, ModuleBuiltData->AttributeDistributionParameters[a], i, 0);
+					const float Value = DistributionSampler.GetValue(ParticleSimulationContext, NormalizedAgeData[i]);
+					ParticleSimulationContext.WriteParticleVariable(ModuleBuiltData->AttributeOffset[a], i, Value);
+				}
+			}
+		}
+	}
+}
+
+ENiagaraStatelessFeatureMask UNiagaraStatelessModule_LightAttributes::GetFeatureMask() const
+{
+	return ENiagaraStatelessFeatureMask::ExecuteCPU;
+}
+
+void UNiagaraStatelessModule_LightAttributes::BuildEmitterData(const FNiagaraStatelessEmitterDataBuildContext& BuildContext) const
+{
+	using namespace NSMLightAttributesPrivate;
+
+	FModuleBuiltData* BuiltData = BuildContext.AllocateBuiltData<FModuleBuiltData>();
+	if (!IsModuleEnabled())
+	{
+		return;
+	}
+
+	// Gather our attribute bindings
+	{
+		bool bAnyValidAttributes = false;
+		for (int i = 0; i < int(EModuleAttribute::Num); ++i)
+		{
+			if (const FNiagaraVariableBase* Variable = GetAttributeVariable(EModuleAttribute(i), this))
+			{
+				BuiltData->AttributeOffset[i] = BuildContext.FindParticleVariableIndex(*Variable);
+				bAnyValidAttributes |= BuiltData->AttributeOffset[i] != INDEX_NONE;
+			}
+		}
+
+		if ( !bAnyValidAttributes )
+		{
+			return;
+		}
+	}
+
+	// Build distributions
+	for (int i = 0; i < int(EModuleAttribute::Num); ++i)
+	{
+		const FNiagaraDistributionFloat* Distribution = BuiltData->AttributeOffset[i] != INDEX_NONE ? GetAttributeDistribution(EModuleAttribute(i), this) : nullptr;
+		if (Distribution)
+		{
+			BuiltData->AttributeDistributionParameters[i] = BuildContext.AddDistribution(*Distribution);
+		}
+	}
+
+	BuildContext.AddParticleSimulationExecSimulate(&ParticleSimulate);
+}
+
+void UNiagaraStatelessModule_LightAttributes::SetShaderParameters(const FNiagaraStatelessSetShaderParameterContext& SetShaderParameterContext) const
+{
+	using namespace NSMLightAttributesPrivate;
+
+	const FModuleBuiltData* ModuleBuiltData = SetShaderParameterContext.ReadBuiltData<FModuleBuiltData>();
+}
+
+#if WITH_EDITORONLY_DATA
+void UNiagaraStatelessModule_LightAttributes::GetOutputVariables(TArray<FNiagaraVariableBase>& OutVariables) const
+{
+	using namespace NSMLightAttributesPrivate;
+
+	for (int i = 0; i < int(EModuleAttribute::Num); ++i)
+	{
+		if (const FNiagaraVariableBase* Variable = GetAttributeVariable(EModuleAttribute(i), this))
+		{
+			OutVariables.AddUnique(*Variable);
+		}
+	}
+}
+#endif
diff --git a/Engine/Plugins/FX/Niagara/Source/Niagara/Private/Stateless/NiagaraStatelessCommon.cpp b/Engine/Plugins/FX/Niagara/Source/Niagara/Private/Stateless/NiagaraStatelessCommon.cpp
--- a/Engine/Plugins/FX/Niagara/Source/Niagara/Private/Stateless/NiagaraStatelessCommon.cpp
+++ b/Engine/Plugins/FX/Niagara/Source/Niagara/Private/Stateless/NiagaraStatelessCommon.cpp
@@ -169,6 +169,13 @@ namespace NiagaraStatelessCommon
 		}
 		GUpdatedFeatureMask.Reset();
 	}
+
+	FNiagaraVariableBase ConvertParticleVariableToStateless(const FNiagaraVariableBase& InVariable)
+	{
+		FNiagaraVariableBase Variable = InVariable;
+		Variable.RemoveRootNamespace(FNiagaraConstants::ParticleAttributeNamespaceString);
+		return Variable;
+	}
 } //NiagaraStatelessCommon
 
 const FNiagaraStatelessGlobals& FNiagaraStatelessGlobals::Get()
diff --git a/Engine/Plugins/FX/Niagara/Source/Niagara/Private/Stateless/NiagaraStatelessEmitter.cpp b/Engine/Plugins/FX/Niagara/Source/Niagara/Private/Stateless/NiagaraStatelessEmitter.cpp
--- a/Engine/Plugins/FX/Niagara/Source/Niagara/Private/Stateless/NiagaraStatelessEmitter.cpp
+++ b/Engine/Plugins/FX/Niagara/Source/Niagara/Private/Stateless/NiagaraStatelessEmitter.cpp
@@ -297,6 +297,7 @@ const UNiagaraStatelessEmitterTemplate* UNiagaraStatelessEmitter::GetEmitterTemp
 void UNiagaraStatelessEmitter::CacheFromCompiledData()
 {
 	StatelessEmitterData = MakeShareable(new FNiagaraStatelessEmitterData(), FNiagaraRenderThreadDeletor<FNiagaraStatelessEmitterData>());
+	StatelessEmitterData->bCanEverExecute = true;
 
 #if WITH_NIAGARA_DEBUG_EMITTER_NAME
 	StatelessEmitterData->DebugSimulationName = GetPackage()->GetFName();
@@ -304,6 +305,27 @@ void UNiagaraStatelessEmitter::CacheFromCompiledData()
 #endif
 	StatelessEmitterData->EmitterTemplate = GetEmitterTemplate();
 
+	// Determine feature mask (needed before we build the data set)
+	// Determine our supported feature set mask
+	// If we can not execute on any enabled units then the emitter is considered disabled
+	StatelessEmitterData->FeatureMask = FNiagaraStatelessGlobals::Get().FeatureMask;
+	StatelessEmitterData->FeatureMask &= ENiagaraStatelessFeatureMask(AllowedFeatureMask);
+
+	for (const UNiagaraStatelessModule* Module : Modules)
+	{
+		if (Module->IsModuleEnabled())
+		{
+			StatelessEmitterData->FeatureMask &= Module->GetFeatureMask();
+		}
+	}
+
+	if (StatelessEmitterData->FeatureMask == ENiagaraStatelessFeatureMask::None)
+	{
+		StatelessEmitterData->bCanEverExecute = false;
+		UE_LOG(LogNiagara, Log, TEXT("Stateless Emitter (%s) can not execute on any available path and will be disabled."), *GetFullName());
+		return;
+	}
+
 	// Build data set
 	BuildCompiledDataSet();
 
@@ -343,31 +365,9 @@ void UNiagaraStatelessEmitter::CacheFromCompiledData()
 		}
 	}
 
-	StatelessEmitterData->bCanEverExecute = NiagaraStatelessInternal::IsValid(*StatelessEmitterData);
+	StatelessEmitterData->bCanEverExecute &= NiagaraStatelessInternal::IsValid(*StatelessEmitterData);
 	StatelessEmitterData->bCanEverExecute &= IsAllowedByScalability();
 
-	// Determine our supported feature set mask
-	// If we can not execute on any enabled units then the emitter is considered disabled
-	if (StatelessEmitterData->bCanEverExecute)
-	{
-		StatelessEmitterData->FeatureMask = FNiagaraStatelessGlobals::Get().FeatureMask;
-		StatelessEmitterData->FeatureMask &= ENiagaraStatelessFeatureMask(AllowedFeatureMask);
-
-		for (const UNiagaraStatelessModule* Module : Modules)
-		{
-			if (Module->IsModuleEnabled())
-			{
-				StatelessEmitterData->FeatureMask &= Module->GetFeatureMask();
-			}
-		}
-
-		if (StatelessEmitterData->FeatureMask == ENiagaraStatelessFeatureMask::None)
-		{
-			StatelessEmitterData->bCanEverExecute = false;
-			UE_LOG(LogNiagara, Verbose, TEXT("Stateless Emitter (%s) can not execute on any available path and will be disabled."), *GetFullName());
-		}
-	}
-
 	// Resolve scalability settings
 	ResolveScalabilitySettings();
 
@@ -523,12 +523,17 @@ void UNiagaraStatelessEmitter::BuildCompiledDataSet()
 		}
 
 		// Remove any variables we don't output
+		// Note: We only need to do this for GPU path as it can't vary the outputs
 		TConstArrayView<FNiagaraVariableBase> OutputComponents = EmitterTemplate->GetOututputComponents();
-		for (auto it=AvailableVariables.CreateIterator(); it; ++it)
+		if (EnumHasAnyFlags(StatelessEmitterData->FeatureMask, ENiagaraStatelessFeatureMask::ExecuteGPU))
 		{
-			if (!OutputComponents.Contains(*it))
+			for (auto it = AvailableVariables.CreateIterator(); it; ++it)
 			{
-				it.RemoveCurrent();
+				if (!OutputComponents.Contains(*it))
+				{
+					UE_LOG(LogNiagara, Log, TEXT("Removed variable '%s' for emitter '%s' as it's not part of the output components"), *it->GetName().ToString(), *GetFullName());
+					it.RemoveCurrent();
+				}
 			}
 		}
 
diff --git a/Engine/Plugins/FX/Niagara/Source/Niagara/Private/Stateless/NiagaraStatelessEmitterTemplate.cpp b/Engine/Plugins/FX/Niagara/Source/Niagara/Private/Stateless/NiagaraStatelessEmitterTemplate.cpp
--- a/Engine/Plugins/FX/Niagara/Source/Niagara/Private/Stateless/NiagaraStatelessEmitterTemplate.cpp
+++ b/Engine/Plugins/FX/Niagara/Source/Niagara/Private/Stateless/NiagaraStatelessEmitterTemplate.cpp
@@ -13,10 +13,12 @@
 #include "Stateless/Modules/NiagaraStatelessModule_CalculateAccurateVelocity.h"
 #include "Stateless/Modules/NiagaraStatelessModule_CameraOffset.h"
 #include "Stateless/Modules/NiagaraStatelessModule_CurlNoiseForce.h"
+#include "Stateless/Modules/NiagaraStatelessModule_DecalAttributes.h"
 #include "Stateless/Modules/NiagaraStatelessModule_Drag.h"
 #include "Stateless/Modules/NiagaraStatelessModule_GravityForce.h"
 #include "Stateless/Modules/NiagaraStatelessModule_InitializeParticle.h"
 #include "Stateless/Modules/NiagaraStatelessModule_InitialMeshOrientation.h"
+#include "Stateless/Modules/NiagaraStatelessModule_LightAttributes.h"
 #include "Stateless/Modules/NiagaraStatelessModule_MeshIndex.h"
 #include "Stateless/Modules/NiagaraStatelessModule_MeshRotationRate.h"
 #include "Stateless/Modules/NiagaraStatelessModule_RotateAroundPoint.h"
@@ -91,7 +93,9 @@ void UNiagaraStatelessEmitterDefault::InitModulesAndAttributes()
 		UNiagaraStatelessModule_SolveVelocitiesAndForces::StaticClass(),
 		// Post Solve
 		UNiagaraStatelessModule_CameraOffset::StaticClass(),
+		UNiagaraStatelessModule_DecalAttributes::StaticClass(),
 		UNiagaraStatelessModule_DynamicMaterialParameters::StaticClass(),
+		UNiagaraStatelessModule_LightAttributes::StaticClass(),
 		UNiagaraStatelessModule_MeshIndex::StaticClass(),
 		UNiagaraStatelessModule_MeshRotationRate::StaticClass(),
 		UNiagaraStatelessModule_ScaleColor::StaticClass(),
diff --git a/Engine/Plugins/FX/Niagara/Source/Niagara/Public/NiagaraDecalRendererProperties.h b/Engine/Plugins/FX/Niagara/Source/Niagara/Public/NiagaraDecalRendererProperties.h
--- a/Engine/Plugins/FX/Niagara/Source/Niagara/Public/NiagaraDecalRendererProperties.h
+++ b/Engine/Plugins/FX/Niagara/Source/Niagara/Public/NiagaraDecalRendererProperties.h
@@ -15,7 +15,7 @@ class UMaterialInterface;
 class FNiagaraEmitterInstance;
 class SWidget;
 
-UCLASS(editinlinenew, MinimalAPI, meta = (DisplayName = "Decal Renderer"))
+UCLASS(editinlinenew, MinimalAPI, meta = (DisplayName = "Decal Renderer", SupportsStateless))
 class UNiagaraDecalRendererProperties : public UNiagaraRendererProperties
 {
 public:
@@ -66,6 +66,11 @@ public:
 	static const int32 GetDefaultDecalSortOrder() { return 0; }
 	static const FNiagaraBool GetDefaultDecalVisible() { return FNiagaraBool(true); }
 
+	static FNiagaraVariable GetDecalOrientationVariable();
+	static FNiagaraVariable GetDecalSizeVariable();
+	static FNiagaraVariable GetDecalFadeVariable();
+	static FNiagaraVariable GetDecalVisibleVariable();
+
 	/** What material to use for the decal. */
 	UPROPERTY(EditAnywhere, Category = "Decal Rendering")
 	TObjectPtr<UMaterialInterface> Material;
