When clearing out renderer MICs be sure to rename them as well as marking them for garbage so that we don't run into conflicts at a later point.

diff --git a/Engine/Plugins/FX/Niagara/Source/Niagara/Private/NiagaraRendererProperties.cpp b/Engine/Plugins/FX/Niagara/Source/Niagara/Private/NiagaraRendererProperties.cpp
--- a/Engine/Plugins/FX/Niagara/Source/Niagara/Private/NiagaraRendererProperties.cpp
+++ b/Engine/Plugins/FX/Niagara/Source/Niagara/Private/NiagaraRendererProperties.cpp
@@ -43,6 +43,12 @@ static FAutoConsoleVariableRef CVarNiagaraRendererCookOutStaticEnabledBinding(
 	TEXT("If none zero renderers with static variables used for enabled binding will cook out if they are not enabled."),
 	ECVF_Scalability
 );
+#endif
+
+namespace NiagaraRendererPropertiesPrivate
+{
+
+#if WITH_EDITORONLY_DATA
 
 // Attempts to resolve a static variable across all the scripts
 // If the state is undetermined (i.e. does not exist or is inconsistent across scripts) we will not return a value
@@ -98,8 +104,38 @@ static TOptional<bool> TryResolveStaticVariableBool(const UNiagaraEmitter* Niaga
 
 	return StaticValue;
 }
+
 #endif
 
+void MarkAndRenameMaterialForGarbage(UMaterialInstanceConstant* MIC)
+{
+	if (MIC)
+	{
+		const ERenameFlags RenameFlags = REN_NonTransactional | REN_DoNotDirty | REN_DontCreateRedirectors;
+
+		FNameBuilder UniqueNameBuilder;
+		int32 UniqueIndex = 0;
+
+		do 
+		{
+			UniqueNameBuilder.Reset();
+			UniqueNameBuilder << TEXT("TRASH_");
+			MIC->GetFName().AppendString(UniqueNameBuilder);
+			if (UniqueIndex)
+			{
+				UniqueNameBuilder << TEXT("_");
+				UniqueNameBuilder << UniqueIndex;
+			}
+			++UniqueIndex;
+		} while (!MIC->Rename(*UniqueNameBuilder, nullptr, REN_Test | RenameFlags));
+
+		MIC->Rename(*UniqueNameBuilder, nullptr, RenameFlags);
+		MIC->MarkAsGarbage();
+	}
+}
+
+} // NiagaraRendererPropertiesPrivate
+
 void FNiagaraRendererLayout::Initialize(int32 NumVariables)
 {
 	VFVariables_GT.Reset(NumVariables);
@@ -569,7 +605,7 @@ void UNiagaraRendererProperties::UpdateMaterialParametersMIC(const FNiagaraRende
 	}
 	else if (InOutMIC)
 	{
-		InOutMIC->MarkAsGarbage();
+		NiagaraRendererPropertiesPrivate::MarkAndRenameMaterialForGarbage(InOutMIC);
 		InOutMIC = nullptr;
 	}
 }
@@ -659,7 +695,7 @@ void UNiagaraRendererProperties::UpdateMaterialParametersMIC(const FNiagaraRende
 	// Garabge the MICs we are no longer going to use
 	for (int i=0; i < MICPool.Num(); ++i)
 	{
-		MICPool[i]->MarkAsGarbage();
+		NiagaraRendererPropertiesPrivate::MarkAndRenameMaterialForGarbage(MICPool[i]);
 	}
 }
 
@@ -1042,7 +1078,7 @@ bool UNiagaraRendererProperties::NeedsLoadForTargetPlatform(const ITargetPlatfor
 			#if WITH_EDITORONLY_DATA
 				if (GNiagaraRendererCookOutStaticEnabledBinding && RendererEnabledBinding.IsValid())
 				{
-					TOptional<bool> ResolvedStaticValue = TryResolveStaticVariableBool(OwnerEmitter, RendererEnabledBinding.GetParamMapBindableVariable());
+					TOptional<bool> ResolvedStaticValue = NiagaraRendererPropertiesPrivate::TryResolveStaticVariableBool(OwnerEmitter, RendererEnabledBinding.GetParamMapBindableVariable());
 					return ResolvedStaticValue.Get(true);
 				}
 			#endif
@@ -1060,7 +1096,7 @@ bool UNiagaraRendererProperties::NeedsLoadForTargetPlatform(const ITargetPlatfor
 //#if WITH_EDITORONLY_DATA
 //				if (GNiagaraRendererCookOutStaticEnabledBinding && RendererEnabledBinding.IsValid())
 //				{
-//					TOptional<bool> ResolvedStaticValue = TryResolveStaticVariableBool(OwnerEmitter, RendererEnabledBinding.GetParamMapBindableVariable());
+//					TOptional<bool> ResolvedStaticValue = NiagaraRendererPropertiesPrivate::TryResolveStaticVariableBool(OwnerEmitter, RendererEnabledBinding.GetParamMapBindableVariable());
 //					return ResolvedStaticValue.Get(true);
 //				}
 //#endif
