- Dynamic Mesh Sim Cache Storage + Visualizer

#tests Create a system using the dunamic mesh DI capture a cache, examine the cache, export the data, etc.

diff --git a/Engine/Plugins/FX/Niagara/Source/Niagara/Internal/DataInterface/NiagaraDataInterfaceDynamicMesh.h b/Engine/Plugins/FX/Niagara/Source/Niagara/Internal/DataInterface/NiagaraDataInterfaceDynamicMesh.h
new file mode 100644
--- /dev/null
+++ b/Engine/Plugins/FX/Niagara/Source/Niagara/Internal/DataInterface/NiagaraDataInterfaceDynamicMesh.h
@@ -0,0 +1,172 @@
+// Copyright Epic Games, Inc. All Rights Reserved.
+
+#pragma once
+
+#include "NiagaraCommon.h"
+#include "NiagaraDataInterface.h"
+#include "NiagaraRenderableMeshInterface.h"
+#include "NiagaraSimCacheCustomStorageInterface.h"
+#include "VectorVM.h"
+
+#include "NiagaraDataInterfaceDynamicMesh.generated.h"
+
+USTRUCT()
+struct FNiagaraDynamicMeshSection
+{
+	GENERATED_BODY()
+
+	/** Number of triangles in this section. */
+	UPROPERTY(EditAnywhere, Category = "Section")
+	int32 NumTriangles = 0;
+
+	/** Index of the material to use. */
+	UPROPERTY(EditAnywhere, Category = "Section")
+	int32 MaterialIndex = 0;
+
+	bool operator==(const FNiagaraDynamicMeshSection& Other) const
+	{
+		return NumTriangles == Other.NumTriangles && MaterialIndex == Other.MaterialIndex;
+	}
+};
+
+USTRUCT()
+struct FNiagaraDynamicMeshMaterial
+{
+	GENERATED_BODY()
+
+	FNiagaraDynamicMeshMaterial();
+
+	/** Optional material to use for this section. */
+	UPROPERTY(EditAnywhere, Category = "Material")
+	TObjectPtr<UMaterialInterface> Material;
+
+	/** Optional binding to a user parameter */
+	UPROPERTY(EditAnywhere, Category = "Material")
+	FNiagaraUserParameterBinding MaterialUserParamBinding;
+
+	bool operator==(const FNiagaraDynamicMeshMaterial& Other) const
+	{
+		return Material == Other.Material && MaterialUserParamBinding == Other.MaterialUserParamBinding;
+	}
+};
+
+/** Data Interface that can create mesh geometry at runtime from triangle data (e.g. to render a custom generated mesh per particle) */
+UCLASS(Experimental, EditInlineNew, Category = "Mesh", CollapseCategories, meta = (DisplayName = "Dynamic Mesh"), MinimalAPI)
+class UNiagaraDataInterfaceDynamicMesh : public UNiagaraDataInterface, public INiagaraRenderableMeshInterface, public INiagaraSimCacheCustomStorageInterface
+{
+	GENERATED_UCLASS_BODY()
+
+public:
+	/**
+	Sections to render, each section will generally result in a draw call.
+	Triangles are contiguous between sections, i.e. Section[1] triangles will begin after Section[0].NumTriangles
+	*/
+	UPROPERTY(EditAnywhere, Category = "Dynamic Mesh")
+	TArray<FNiagaraDynamicMeshSection> Sections;
+
+	/** List of materials to use */
+	UPROPERTY(EditAnywhere, Category = "Dynamic Mesh")
+	TArray<FNiagaraDynamicMeshMaterial> Materials;
+
+	/** Allocates space for the number of vertices we will use, leave as zero if you intend to allocate dynamically via VM calls. */
+	UPROPERTY(EditAnywhere, Category = "Dynamic Mesh")
+	int32 NumVertices = 0;
+
+	/** Allocates space for the number of texture coordinates requested. */
+	UPROPERTY(EditAnywhere, Category = "Dynamic Mesh")
+	int NumTexCoords = 0;
+
+	/** Allocates space for vertex colors when enabled. */
+	UPROPERTY(EditAnywhere, Category = "Dynamic Mesh")
+	bool bHasColors = false;
+
+	/** Allocates space for tangent basis when enabled. */
+	UPROPERTY(EditAnywhere, Category = "Dynamic Mesh")
+	bool bHasTangentBasis = false;
+
+	/** Should we auto clear section triangle allocations per frame or not. */
+	UPROPERTY(EditAnywhere, Category = "Dynamic Mesh")
+	bool bClearTrianglesPerFrame = false;
+
+	/** Should we auto clear section triangle allocations per frame or not. */
+	UPROPERTY(EditAnywhere, Category = "Dynamic Mesh")
+	FBox DefaultBounds = FBox(FVector(-100.0f), FVector(100.0f));
+
+	//UObject Interface
+	NIAGARA_API virtual void PostInitProperties() override;
+	//UObject Interface End
+
+	//UNiagaraDataInterface Interface
+	NIAGARA_API virtual void GetVMExternalFunction(const FVMExternalFunctionBindingInfo& BindingInfo, void* InstanceData, FVMExternalFunction &OutFunc) override;
+	virtual bool CanExecuteOnTarget(ENiagaraSimTarget Target) const override { return true; }
+
+	NIAGARA_API virtual int32 PerInstanceDataSize() const override;
+	NIAGARA_API virtual bool InitPerInstanceData(void* PerInstanceData, FNiagaraSystemInstance* SystemInstance) override;
+	NIAGARA_API virtual void DestroyPerInstanceData(void* PerInstanceData, FNiagaraSystemInstance* SystemInstance) override;
+
+	virtual bool HasPreSimulateTick() const override { return true; }
+	virtual bool HasPostSimulateTick() const override { return true; }
+	NIAGARA_API virtual bool PerInstanceTick(void* PerInstanceData, FNiagaraSystemInstance* SystemInstance, float DeltaSeconds) override;
+	NIAGARA_API virtual bool PerInstanceTickPostSimulate(void* PerInstanceData, FNiagaraSystemInstance* SystemInstance, float DeltaSeconds) override;
+
+	NIAGARA_API virtual bool Equals(const UNiagaraDataInterface* Other) const override;
+	NIAGARA_API virtual bool CopyToInternal(UNiagaraDataInterface* Destination) const override;
+
+	// GPU sim functionality
+#if WITH_EDITORONLY_DATA
+	NIAGARA_API virtual bool AppendCompileHash(FNiagaraCompileHashVisitor* InVisitor) const override;
+	NIAGARA_API virtual void GetParameterDefinitionHLSL(const FNiagaraDataInterfaceGPUParamInfo& ParamInfo, FString& OutHLSL) override;
+	NIAGARA_API virtual bool GetFunctionHLSL(const FNiagaraDataInterfaceGPUParamInfo& ParamInfo, const FNiagaraDataInterfaceGeneratedFunction& FunctionInfo, int FunctionInstanceIndex, FString& OutHLSL) override;
+#endif
+	NIAGARA_API virtual void BuildShaderParameters(FNiagaraShaderParametersBuilder& ShaderParametersBuilder) const override;
+	NIAGARA_API virtual void SetShaderParameters(const FNiagaraDataInterfaceSetShaderParametersContext& Context) const override;
+	//UNiagaraDataInterface Interface End
+
+	//INiagaraRenderableMeshInterface Interface Begin
+	NIAGARA_API virtual FNiagaraRenderableMeshPtr GetRenderableMesh(FNiagaraSystemInstanceID SystemInstanceID) override;
+	NIAGARA_API virtual void GetUsedMaterials(FNiagaraSystemInstanceID SystemInstanceID, TArray<UMaterialInterface*>& OutMaterials) const override;
+	//INiagaraRenderableMeshInterface Interface End
+
+	//~ INiagaraSimCacheCustomStorageInterface interface BEGIN
+	virtual UObject* SimCacheBeginWrite(UObject* SimCache, FNiagaraSystemInstance* NiagaraSystemInstance, const void* OptionalPerInstanceData, FNiagaraSimCacheFeedbackContext& FeedbackContext) const override;
+	virtual bool SimCacheWriteFrame(UObject* StorageObject, int FrameIndex, FNiagaraSystemInstance* SystemInstance, const void* OptionalPerInstanceData, FNiagaraSimCacheFeedbackContext& FeedbackContext) const override;
+	virtual bool SimCacheReadFrame(const FNiagaraSimCacheDataInterfaceReadContext& ReadContext) override;
+	virtual bool SimCacheCompareFrame(const UObject* LhsStorageObject, const UObject* RhsStorageObject, int FrameIndex, TOptional<float> Tolerance, FString& OutErrors) const override;
+	//~ UNiagaraDataInterface interface END
+
+	// VM Functions
+	static NIAGARA_API void VMGetMeshProperties(FVectorVMExternalFunctionContext& Context);
+	static NIAGARA_API void VMGetSectionCount(FVectorVMExternalFunctionContext& Context);
+	static NIAGARA_API void VMGetSectionData(FVectorVMExternalFunctionContext& Context);
+	static NIAGARA_API void VMGetLocalBounds(FVectorVMExternalFunctionContext& Context);
+
+	static NIAGARA_API void VMSetMeshProperties(FVectorVMExternalFunctionContext& Context);
+	static NIAGARA_API void VMSetSectionCount(FVectorVMExternalFunctionContext& Context);
+	static NIAGARA_API void VMSetSectionData(FVectorVMExternalFunctionContext& Context);
+	static NIAGARA_API void VMSetLocalBounds(FVectorVMExternalFunctionContext& Context);
+
+	static NIAGARA_API void VMClearAllSectionTriangles(FVectorVMExternalFunctionContext& Context);
+	static NIAGARA_API void VMClearSectionTriangles(FVectorVMExternalFunctionContext& Context);
+	static NIAGARA_API void VMAllocateSectionTriangles(FVectorVMExternalFunctionContext& Context);
+
+	static NIAGARA_API void VMGetTriangleVertices(FVectorVMExternalFunctionContext& Context);
+	static NIAGARA_API void VMGetVertexPosition(FVectorVMExternalFunctionContext& Context);
+	static NIAGARA_API void VMGetVertexTangentBasis(FVectorVMExternalFunctionContext& Context);
+	static NIAGARA_API void VMGetVertexTexCoord(FVectorVMExternalFunctionContext& Context);
+	static NIAGARA_API void VMGetVertexColor(FVectorVMExternalFunctionContext& Context);
+	static NIAGARA_API void VMGetVertexData(FVectorVMExternalFunctionContext& Context);
+
+	static NIAGARA_API void VMSetTriangleVertices(FVectorVMExternalFunctionContext& Context);
+	static NIAGARA_API void VMSetVertexPosition(FVectorVMExternalFunctionContext& Context);
+	static NIAGARA_API void VMSetVertexTangentBasis(FVectorVMExternalFunctionContext& Context);
+	static NIAGARA_API void VMSetVertexTexCoord(FVectorVMExternalFunctionContext& Context);
+	static NIAGARA_API void VMSetVertexColor(FVectorVMExternalFunctionContext& Context);
+	static NIAGARA_API void VMSetVertexData(FVectorVMExternalFunctionContext& Context);
+
+	static NIAGARA_API void VMAppendTriangle(FVectorVMExternalFunctionContext& Context);
+
+protected:
+#if WITH_EDITORONLY_DATA
+	NIAGARA_API virtual void GetFunctionsInternal(TArray<FNiagaraFunctionSignature>& OutFunctions) const override;
+#endif
+};
diff --git a/Engine/Plugins/FX/Niagara/Source/Niagara/Internal/DataInterface/NiagaraDataInterfaceDynamicMeshSimCacheData.h b/Engine/Plugins/FX/Niagara/Source/Niagara/Internal/DataInterface/NiagaraDataInterfaceDynamicMeshSimCacheData.h
new file mode 100644
--- /dev/null
+++ b/Engine/Plugins/FX/Niagara/Source/Niagara/Internal/DataInterface/NiagaraDataInterfaceDynamicMeshSimCacheData.h
@@ -0,0 +1,165 @@
+// Copyright Epic Games, Inc. All Rights Reserved.
+
+#pragma once
+
+#include "NiagaraCommon.h"
+
+#include "NiagaraDataInterfaceDynamicMeshSimCacheData.generated.h"
+
+class UMaterialInterface;
+
+USTRUCT()
+struct FNDIDynamicMeshSimCacheSection
+{
+	GENERATED_BODY()
+
+	UPROPERTY(EditAnywhere, Category = "Section")
+	uint32 TriangleOffset = 0;
+
+	UPROPERTY(EditAnywhere, Category = "Section")
+	uint32 MaxTriangles = 0;
+
+	UPROPERTY(EditAnywhere, Category = "Section")
+	uint32 AllocatedTriangles = 0;
+
+	UPROPERTY(EditAnywhere, Category = "Section")
+	TObjectPtr<UMaterialInterface> Material;
+
+	bool operator==(const FNDIDynamicMeshSimCacheSection& Other) const
+	{
+		return
+			TriangleOffset == Other.TriangleOffset &&
+			MaxTriangles == Other.MaxTriangles &&
+			AllocatedTriangles == Other.AllocatedTriangles &&
+			Material == Other.Material;
+	}
+};
+
+USTRUCT()
+struct FNDIDynamicMeshSimCacheFrame
+{
+	GENERATED_BODY()
+
+	UPROPERTY(EditAnywhere, Category ="Mesh")
+	uint32 NumTriangles = 0;
+
+	UPROPERTY(EditAnywhere, Category = "Mesh")
+	uint32 NumVertices = 0;
+
+	UPROPERTY(EditAnywhere, Category = "Mesh")
+	uint32 NumTexCoords = 0;
+
+	UPROPERTY(EditAnywhere, Category = "Mesh")
+	FBox LocalBounds = FBox(FVector(-100.0f), FVector(100.0f));
+
+	UPROPERTY(EditAnywhere, Category = "Mesh")
+	uint32 VertexBufferSize = 0;
+
+	UPROPERTY(EditAnywhere, Category = "Mesh")
+	uint32 PositionOffset = INDEX_NONE;
+
+	UPROPERTY(EditAnywhere, Category = "Mesh")
+	uint32 TangentBasisOffset = INDEX_NONE;
+
+	UPROPERTY(EditAnywhere, Category = "Mesh")
+	uint32 TexCoordOffset = INDEX_NONE;
+
+	UPROPERTY(EditAnywhere, Category = "Mesh")
+	uint32 ColorOffset = INDEX_NONE;
+
+	UPROPERTY()
+	TArray<uint8> IndexData;
+
+	UPROPERTY()
+	TArray<uint8> VertexData;
+
+	UPROPERTY(EditAnywhere, Category = "Mesh")
+	TArray<FNDIDynamicMeshSimCacheSection> Sections;
+
+	bool Equals(const FNDIDynamicMeshSimCacheFrame& Other, TConstArrayView<uint8> OtherIndexData, TConstArrayView<uint8> OtherVertexData) const
+	{
+		return
+			NumTriangles == Other.NumTriangles &&
+			NumVertices == Other.NumVertices &&
+			NumTexCoords == Other.NumTexCoords &&
+			LocalBounds == Other.LocalBounds &&
+			VertexBufferSize == Other.VertexBufferSize &&
+			PositionOffset == Other.PositionOffset &&
+			TangentBasisOffset == Other.TangentBasisOffset &&
+			TexCoordOffset == Other.TexCoordOffset &&
+			ColorOffset == Other.ColorOffset&&
+			Sections == Other.Sections &&
+			IndexData == OtherIndexData &&
+			VertexData == OtherVertexData;
+	}
+
+	bool Equals(const FNDIDynamicMeshSimCacheFrame& Other) const
+	{
+		return Equals(Other, IndexData, VertexData);
+	}
+};
+
+UCLASS(MinimalAPI)
+class UNiagaraDataInterfaceDynamicMeshSimCacheData : public UObject
+{
+	GENERATED_BODY()
+
+public:
+	UPROPERTY()
+	TArray<FNDIDynamicMeshSimCacheFrame> UniqueFrames;
+
+	// Indicies into UniqueFrames where INDEX_NONE means no data
+	UPROPERTY()
+	TArray<int> CpuFrameIndices;
+
+	UPROPERTY()
+	TArray<int> GpuFrameIndices;
+
+	void ResizeFrames(int FrameIndex)
+	{
+		while (CpuFrameIndices.Num() <= FrameIndex)
+		{
+			CpuFrameIndices.Add(INDEX_NONE);
+		}
+		while (GpuFrameIndices.Num() <= FrameIndex)
+		{
+			GpuFrameIndices.Add(INDEX_NONE);
+		}
+	}
+
+	int32 FindOrCreateMesh(const FNDIDynamicMeshSimCacheFrame& MeshDesc, TConstArrayView<uint8> IndexData, TConstArrayView<uint8>  VertexData)
+	{
+		if (MeshDesc.VertexBufferSize == 0 || IndexData.Num() == 0 || VertexData.Num() == 0)
+		{
+			return INDEX_NONE;
+		}
+
+		if (UniqueFrames.Num() > 0)
+		{
+			const FNDIDynamicMeshSimCacheFrame& LastFrame = UniqueFrames.Last();
+			if (LastFrame.Equals(MeshDesc, IndexData, VertexData))
+			{
+				return UniqueFrames.Num() - 1;
+			}
+		}
+
+		const int32 MeshIndex = UniqueFrames.Add(MeshDesc);
+		UniqueFrames[MeshIndex].IndexData = IndexData;
+		UniqueFrames[MeshIndex].VertexData = VertexData;
+		return MeshIndex;
+	}
+
+	const FNDIDynamicMeshSimCacheFrame* GetFrameData(int32 FrameIndex, ENiagaraSimTarget SimTarget) const
+	{
+		int32 MeshIndex = INDEX_NONE;
+		if (SimTarget == ENiagaraSimTarget::CPUSim)
+		{
+			MeshIndex = CpuFrameIndices.IsValidIndex(FrameIndex) ? CpuFrameIndices[FrameIndex] : INDEX_NONE;
+		}
+		else if (SimTarget == ENiagaraSimTarget::GPUComputeSim)
+		{
+			MeshIndex = GpuFrameIndices.IsValidIndex(FrameIndex) ? GpuFrameIndices[FrameIndex] : INDEX_NONE;
+		}
+		return UniqueFrames.IsValidIndex(MeshIndex) ? &UniqueFrames[MeshIndex] : nullptr;
+	}
+};
diff --git a/Engine/Plugins/FX/Niagara/Source/Niagara/Private/DataInterface/NiagaraDataInterfaceDynamicMesh.cpp b/Engine/Plugins/FX/Niagara/Source/Niagara/Private/DataInterface/NiagaraDataInterfaceDynamicMesh.cpp
--- a/Engine/Plugins/FX/Niagara/Source/Niagara/Private/DataInterface/NiagaraDataInterfaceDynamicMesh.cpp
+++ b/Engine/Plugins/FX/Niagara/Source/Niagara/Private/DataInterface/NiagaraDataInterfaceDynamicMesh.cpp
@@ -1,7 +1,11 @@
 // Copyright Epic Games, Inc. All Rights Reserved.
 
-#include "NiagaraDataInterfaceDynamicMesh.h"
+#include "DataInterface/NiagaraDataInterfaceDynamicMesh.h"
+#include "DataInterface/NiagaraDataInterfaceDynamicMeshSimCacheData.h"
+
 #include "NiagaraCompileHashVisitor.h"
+#include "NiagaraGpuComputeDispatchInterface.h"
+#include "NiagaraGpuReadbackManager.h"
 #include "NiagaraMeshVertexFactory.h"
 #include "NiagaraParameterStore.h"
 #include "NiagaraSimStageData.h"
@@ -244,7 +248,7 @@ namespace NDIDynamicMeshLocal
 			OutTangentX = FVector3f(float((TangentData[0] >> 0) & 0xff), float((TangentData[0] >> 8) & 0xff), float((TangentData[0] >> 16) & 0xff));
 			OutTangentZ = FVector3f(float((TangentData[1] >> 0) & 0xff), float((TangentData[1] >> 8) & 0xff), float((TangentData[1] >> 16) & 0xff));
 			OutTangentX = (OutTangentX / 127.5f) - 1.0f;
-			OutTangentZ = (OutTangentX / 127.5f) - 1.0f;
+			OutTangentZ = (OutTangentZ / 127.5f) - 1.0f;
 
 			const float TangentSign = (((TangentData[1] >> 24) & 0xff) / 127.5f) - 1.0f;
 
@@ -345,7 +349,6 @@ namespace NDIDynamicMeshLocal
 		uint32			TangentBasisOffset = INDEX_NONE;
 		uint32			TexCoordOffset = INDEX_NONE;
 		uint32			ColorOffset = INDEX_NONE;
-		uint32			VertexStride = INDEX_NONE;
 
 		TResourceArray<uint8>	IndexData;
 		TResourceArray<uint8>	VertexData;
@@ -366,7 +369,6 @@ namespace NDIDynamicMeshLocal
 		uint32			TangentBasisOffset = INDEX_NONE;
 		uint32			TexCoordOffset = INDEX_NONE;
 		uint32			ColorOffset = INDEX_NONE;
-		uint32			VertexStride = INDEX_NONE;
 
 		TRefCountPtr<FRDGPooledBuffer>	SectionDefaultPooledBuffer;
 		TRefCountPtr<FRDGPooledBuffer>	SectionPooledBuffer;
@@ -409,11 +411,11 @@ namespace NDIDynamicMeshLocal
 			NumTexCoords				= GameToRenderData.NumTexCoords;
 			LocalBounds					= GameToRenderData.LocalBounds;
 
+			VertexBufferSize			= GameToRenderData.VertexBufferSize;
 			PositionOffset				= GameToRenderData.PositionOffset;
 			TangentBasisOffset			= GameToRenderData.TangentBasisOffset;
 			TexCoordOffset				= GameToRenderData.TexCoordOffset;
 			ColorOffset					= GameToRenderData.ColorOffset;
-			VertexStride				= GameToRenderData.VertexStride;
 
 			{
 				FRDGBuilder GraphBuilder(RHICmdList);
@@ -510,10 +512,17 @@ namespace NDIDynamicMeshLocal
 					VertexBufferColorSRV = RHICmdList.CreateShaderResourceView(FShaderResourceViewInitializer(VertexBuffer.VertexBufferRHI, PixelFormatColor, ColorOffset, NumVertices));
 				}
 			}
+		#if STATS
+			UpdateMemoryStats();
+		#endif
 		}
 
 		void ReleaseData()
 		{
+		#if STATS
+			ResetMemoryStats();
+		#endif
+
 			IndexBuffer.ReleaseResource();
 			IndexBufferSRV.SafeRelease();
 			VertexBuffer.ReleaseResource();
@@ -1131,13 +1140,6 @@ void UNiagaraDataInterfaceDynamicMesh::DestroyPerInstanceData(void* PerInstanceD
 	(
 		[RT_Proxy=GetProxyAs<FNDIProxy>(), InstanceID=SystemInstance->GetId()](FRHICommandListImmediate& CmdList)
 		{
-#if STATS
-			if (FNDIInstanceData_RenderThread* InstanceData = RT_Proxy->InstanceData_RT.Find(InstanceID))
-			{
-				InstanceData->ResetMemoryStats();
-			}
-#endif
-			
 			RT_Proxy->InstanceData_RT.Remove(InstanceID);
 		}
 	);
@@ -1190,10 +1192,6 @@ bool UNiagaraDataInterfaceDynamicMesh::PerInstanceTickPostSimulate(void* PerInst
 			{
 				FNDIInstanceData_RenderThread& InstanceData_RT = RT_Proxy->InstanceData_RT.FindOrAdd(RT_SystemInstanceID);
 				InstanceData_RT.UpdateData(RHICmdList, RT_GameToRenderData);
-
-#if STATS
-				InstanceData_RT.UpdateMemoryStats();
-#endif
 			}
 		);
 	}
@@ -1341,6 +1339,257 @@ void UNiagaraDataInterfaceDynamicMesh::GetUsedMaterials(FNiagaraSystemInstanceID
 	}
 }
 
+UObject* UNiagaraDataInterfaceDynamicMesh::SimCacheBeginWrite(UObject* SimCache, FNiagaraSystemInstance* NiagaraSystemInstance, const void* OptionalPerInstanceData, FNiagaraSimCacheFeedbackContext& FeedbackContext) const
+{
+	return OptionalPerInstanceData ? NewObject<UNiagaraDataInterfaceDynamicMeshSimCacheData>(SimCache) : nullptr;
+}
+
+bool UNiagaraDataInterfaceDynamicMesh::SimCacheWriteFrame(UObject* StorageObject, int FrameIndex, FNiagaraSystemInstance* SystemInstance, const void* OptionalPerInstanceData, FNiagaraSimCacheFeedbackContext& FeedbackContext) const
+{
+	using namespace NDIDynamicMeshLocal;
+
+	check(OptionalPerInstanceData && StorageObject);
+	UNiagaraDataInterfaceDynamicMeshSimCacheData* CacheData = CastChecked<UNiagaraDataInterfaceDynamicMeshSimCacheData>(StorageObject);
+	CacheData->ResizeFrames(FrameIndex);
+
+	// Build MeshDesc (shared between RT / GT)
+	const FNDIInstanceData_GameThread* InstanceData_GT = reinterpret_cast<const FNDIInstanceData_GameThread*>(OptionalPerInstanceData);
+
+	FNDIDynamicMeshSimCacheFrame MeshDesc;
+	MeshDesc.NumTriangles		= InstanceData_GT->NumTriangles;
+	MeshDesc.NumVertices		= InstanceData_GT->NumVertices;
+	MeshDesc.NumTexCoords		= InstanceData_GT->NumTexCoords;
+	MeshDesc.LocalBounds		= InstanceData_GT->LocalBounds;
+	MeshDesc.VertexBufferSize	= InstanceData_GT->VertexBufferSize;
+	MeshDesc.PositionOffset		= InstanceData_GT->PositionOffset;
+	MeshDesc.TangentBasisOffset	= InstanceData_GT->TangentBasisOffset;
+	MeshDesc.TexCoordOffset		= InstanceData_GT->TexCoordOffset;
+	MeshDesc.ColorOffset		= InstanceData_GT->ColorOffset;
+	MeshDesc.Sections.Reserve(InstanceData_GT->Sections.Num());
+	for (const FNDIInstanceSection& Section : InstanceData_GT->Sections)
+	{
+		FNDIDynamicMeshSimCacheSection& CacheSection = MeshDesc.Sections.AddDefaulted_GetRef();
+		CacheSection.TriangleOffset		= Section.TriangleOffset;
+		CacheSection.MaxTriangles		= Section.MaxTriangles;
+		CacheSection.AllocatedTriangles	= Section.AllocatedTriangles;
+		CacheSection.Material			= InstanceData_GT->Materials.IsValidIndex(Section.MaterialIndex) ? InstanceData_GT->Materials[Section.MaterialIndex].Material : nullptr;
+	}
+
+	// Add to CPU meshes
+	CacheData->CpuFrameIndices[FrameIndex] = CacheData->FindOrCreateMesh(MeshDesc, InstanceData_GT->IndexData, InstanceData_GT->VertexData);
+
+	// Gather from the GPU
+	if (IsUsedWithGPUScript())
+	{
+		FNiagaraGpuComputeDispatchInterface* ComputeInterface = FNiagaraGpuComputeDispatchInterface::Get(SystemInstance->GetWorld());
+		ENQUEUE_RENDER_COMMAND(NDIDynamicMesh_SimCacheWrite)(
+			[Proxy_RT=GetProxyAs<FNDIProxy>(), InstanceID=SystemInstance->GetId(), CacheData, MeshDesc, FrameIndex, ComputeInterface](FRHICommandListImmediate& RHICmdList) mutable
+			{
+				const FNDIInstanceData_RenderThread* InstanceData_RT = Proxy_RT->InstanceData_RT.Find(InstanceID);
+				if (InstanceData_RT == nullptr)
+				{
+					return;
+				}
+
+				// Right now we are assuming that the GT data must match if it doesn't don't attempt to record the frame
+				if ( (MeshDesc.NumTriangles != InstanceData_RT->NumTriangles) ||
+					 (MeshDesc.NumVertices != InstanceData_RT->NumVertices) ||
+					 (MeshDesc.NumTexCoords != InstanceData_RT->NumTexCoords) ||
+					 //(MeshDesc.LocalBounds != InstanceData_RT->LocalBounds) ||
+					 (MeshDesc.VertexBufferSize != InstanceData_RT->VertexBufferSize) ||
+					 (MeshDesc.PositionOffset != InstanceData_RT->PositionOffset) ||
+					 (MeshDesc.TangentBasisOffset != InstanceData_RT->TangentBasisOffset) ||
+					 (MeshDesc.TexCoordOffset != InstanceData_RT->TexCoordOffset) ||
+					 (MeshDesc.ColorOffset != InstanceData_RT->ColorOffset) ||
+					 (MeshDesc.Sections.Num() != InstanceData_RT->NumSections) )
+				{
+					return;
+				}
+
+				if (InstanceData_RT->SectionPooledBuffer == nullptr || InstanceData_RT->IndexPooledBuffer == nullptr || InstanceData_RT->VertexPooledBuffer == nullptr)
+				{
+					return;
+				}
+
+				FRDGBuilder GraphBuilder(RHICmdList);
+
+				FRDGBufferRef BuffersToReadback[3];
+				BuffersToReadback[0] = GraphBuilder.RegisterExternalBuffer(InstanceData_RT->SectionPooledBuffer);
+				BuffersToReadback[1] = GraphBuilder.RegisterExternalBuffer(InstanceData_RT->IndexPooledBuffer);
+				BuffersToReadback[2] = GraphBuilder.RegisterExternalBuffer(InstanceData_RT->VertexPooledBuffer);
+
+				TArray<uint32> GpuSectionData;
+				TArray<uint8> GpuIndexData;
+				TArray<uint8> GpuVertexData;
+
+				FNiagaraGpuReadbackManager* ReadbackManager = ComputeInterface->GetGpuReadbackManager();
+				ReadbackManager->EnqueueReadbacks(
+					GraphBuilder,
+					BuffersToReadback,
+					[&](TConstArrayView<TPair<void*, uint32>> ReadbackData)
+					{
+						GpuSectionData.AddUninitialized(ReadbackData[0].Value / sizeof(uint32));
+						FMemory::Memcpy(GpuSectionData.GetData(), ReadbackData[0].Key, GpuSectionData.NumBytes());
+			
+						GpuIndexData.AddUninitialized(ReadbackData[1].Value);
+						FMemory::Memcpy(GpuIndexData.GetData(), ReadbackData[1].Key, ReadbackData[1].Value);
+
+						GpuVertexData.AddUninitialized(ReadbackData[2].Value);
+						FMemory::Memcpy(GpuVertexData.GetData(), ReadbackData[2].Key, ReadbackData[2].Value);
+					}
+				);
+
+				GraphBuilder.Execute();
+				ReadbackManager->WaitCompletion(RHICmdList);
+
+				for (int iSection = 0; iSection < MeshDesc.Sections.Num(); ++iSection)
+				{
+					FNDIDynamicMeshSimCacheSection& Section = MeshDesc.Sections[iSection];
+					Section.AllocatedTriangles = GpuSectionData[iSection * 2 + 1];
+				}
+
+				CacheData->GpuFrameIndices[FrameIndex] = CacheData->FindOrCreateMesh(MeshDesc, GpuIndexData, GpuVertexData);
+			}
+		);
+
+		FlushRenderingCommands();
+	}
+
+	return true;
+}
+
+bool UNiagaraDataInterfaceDynamicMesh::SimCacheReadFrame(const FNiagaraSimCacheDataInterfaceReadContext& ReadContext)
+{
+	using namespace NDIDynamicMeshLocal;
+
+	const UNiagaraDataInterfaceDynamicMeshSimCacheData* CacheData = ReadContext.GetOptionalStorageObject<UNiagaraDataInterfaceDynamicMeshSimCacheData>();
+	const int32 FrameIndex = ReadContext.GetFrameIndexA();
+
+	// Copy data from cache if we have GPU data we take that as the authoritative version
+	ENiagaraSimTarget SimTarget = ENiagaraSimTarget::GPUComputeSim;
+	const FNDIDynamicMeshSimCacheFrame* CacheFrame = CacheData->GetFrameData(FrameIndex, SimTarget);
+	if (CacheFrame == nullptr)
+	{
+		SimTarget = ENiagaraSimTarget::CPUSim;
+		CacheFrame = CacheData->GetFrameData(FrameIndex, SimTarget);
+	}
+
+	FNDIInstanceData_GameThread* InstanceData_GT = ReadContext.GetPerInstanceData<FNDIInstanceData_GameThread>();
+	if (CacheFrame)
+	{
+		InstanceData_GT->bHasChanged		= false;
+		InstanceData_GT->NumTriangles		= CacheFrame->NumTriangles;
+		InstanceData_GT->NumVertices		= CacheFrame->NumVertices;
+		InstanceData_GT->NumTexCoords		= CacheFrame->NumTexCoords;
+		InstanceData_GT->VertexBufferSize	= CacheFrame->VertexBufferSize;
+		InstanceData_GT->PositionOffset		= CacheFrame->PositionOffset;
+		InstanceData_GT->TangentBasisOffset	= CacheFrame->TangentBasisOffset;
+		InstanceData_GT->TexCoordOffset		= CacheFrame->TexCoordOffset;
+		InstanceData_GT->ColorOffset		= CacheFrame->ColorOffset;
+		InstanceData_GT->IndexData			= CacheFrame->IndexData;
+		InstanceData_GT->VertexData			= CacheFrame->VertexData;
+		InstanceData_GT->Sections.SetNum(CacheFrame->Sections.Num());
+		InstanceData_GT->Materials.Empty();
+
+		for (int iSection = 0; iSection < CacheFrame->Sections.Num(); ++iSection)
+		{
+			const FNDIDynamicMeshSimCacheSection& CacheMeshSection = CacheFrame->Sections[iSection];
+			FNDIInstanceSection& MeshSection = InstanceData_GT->Sections[iSection];
+			UMaterialInterface* SectionMaterial = CacheMeshSection.Material ? CacheMeshSection.Material : UMaterial::GetDefaultMaterial(MD_Surface);
+
+			MeshSection.TriangleOffset		= CacheMeshSection.TriangleOffset;
+			MeshSection.MaxTriangles		= CacheMeshSection.MaxTriangles;
+			MeshSection.AllocatedTriangles	= CacheMeshSection.AllocatedTriangles;
+			MeshSection.MaterialIndex		= InstanceData_GT->Materials.IndexOfByPredicate([SectionMaterial](const FNDIInstanceMaterial& ExistingMaterial) { return ExistingMaterial.Material == SectionMaterial; });
+			if (MeshSection.MaterialIndex == INDEX_NONE)
+			{
+				MeshSection.MaterialIndex = InstanceData_GT->Materials.AddDefaulted();
+				InstanceData_GT->Materials[MeshSection.MaterialIndex].Material = SectionMaterial;
+			}
+		}
+
+		TUniquePtr<FGameToRenderData> GameToRenderData;
+		GameToRenderData.Reset(new FGameToRenderData(*InstanceData_GT));
+		if (SimTarget == ENiagaraSimTarget::GPUComputeSim)
+		{
+			// Clear out the index / vertex data if it came from a GPU simulation
+			// Note: We could do something more optimal here and avoid the copy
+			InstanceData_GT->IndexData.Empty();
+			InstanceData_GT->VertexData.Empty();
+		}
+
+		ENQUEUE_RENDER_COMMAND(NDIDynamicMesh_SimCacheRead)(
+			[Proxy_RT=GetProxyAs<FNDIProxy>(), InstanceID=ReadContext.GetSystemInstance()->GetId(), GameToRenderData_RT=MoveTemp(GameToRenderData)](FRHICommandListImmediate& RHICmdList)
+			{
+				FNDIInstanceData_RenderThread& InstanceData_RT = Proxy_RT->InstanceData_RT.FindOrAdd(InstanceID);
+				InstanceData_RT.UpdateData(RHICmdList, *GameToRenderData_RT.Get());
+			}
+		);
+	}
+	else
+	{
+		InstanceData_GT->bHasChanged		= false;
+		InstanceData_GT->NumTriangles		= 0;
+		InstanceData_GT->NumVertices		= 0;
+		InstanceData_GT->NumTexCoords		= 0;
+		InstanceData_GT->VertexBufferSize	= 0;
+		InstanceData_GT->PositionOffset		= INDEX_NONE;
+		InstanceData_GT->TangentBasisOffset	= INDEX_NONE;
+		InstanceData_GT->TexCoordOffset		= INDEX_NONE;
+		InstanceData_GT->ColorOffset		= INDEX_NONE;
+		InstanceData_GT->Sections.Empty();
+		InstanceData_GT->Materials.Empty();
+		InstanceData_GT->IndexData.Empty();
+		InstanceData_GT->VertexData.Empty();
+	}
+
+	return true;
+}
+
+bool UNiagaraDataInterfaceDynamicMesh::SimCacheCompareFrame(const UObject* LhsStorageObject, const UObject* RhsStorageObject, int FrameIndex, TOptional<float> Tolerance, FString& OutErrors) const
+{
+	const UNiagaraDataInterfaceDynamicMeshSimCacheData* LhsCacheData = CastChecked<UNiagaraDataInterfaceDynamicMeshSimCacheData>(LhsStorageObject);
+	const UNiagaraDataInterfaceDynamicMeshSimCacheData* RhsCacheData = CastChecked<UNiagaraDataInterfaceDynamicMeshSimCacheData>(RhsStorageObject);
+
+	const bool LhsValid = LhsCacheData->CpuFrameIndices.IsValidIndex(FrameIndex) && LhsCacheData->GpuFrameIndices.IsValidIndex(FrameIndex);
+	const bool RhsValid = RhsCacheData->CpuFrameIndices.IsValidIndex(FrameIndex) && RhsCacheData->GpuFrameIndices.IsValidIndex(FrameIndex);
+	if (LhsValid != RhsValid)
+	{
+		OutErrors = FString::Format(TEXT("Frame Validity different LhsValid = {0} RhsValid = {1}"), { LhsValid, RhsValid });
+		return false;
+	}
+
+	if (!LhsValid)
+	{
+		return true;
+	}
+
+	for (ENiagaraSimTarget SimTarget : {ENiagaraSimTarget::CPUSim, ENiagaraSimTarget::GPUComputeSim})
+	{
+		const FNDIDynamicMeshSimCacheFrame* LhsMeshData = LhsCacheData->GetFrameData(FrameIndex, SimTarget);
+		const FNDIDynamicMeshSimCacheFrame* RhsMeshData = LhsCacheData->GetFrameData(FrameIndex, SimTarget);
+
+		const bool LhsMeshDataValid = LhsMeshData != nullptr;
+		const bool RhsMeshDataValid = RhsMeshData != nullptr;
+		if (LhsMeshDataValid != RhsMeshDataValid)
+		{
+			OutErrors = FString::Format(TEXT("Mesh Data Validity different LhsMeshData = {0} RhsMeshData = {1}"), { LhsMeshDataValid, RhsMeshDataValid });
+			return false;
+		}
+
+		if (LhsMeshDataValid)
+		{
+			if (!LhsMeshData->Equals(*RhsMeshData))
+			{
+				OutErrors = TEXT("Mesh Data Different");
+				return false;
+			}
+		}
+	}
+
+	return true;
+}
+
 void UNiagaraDataInterfaceDynamicMesh::VMGetMeshProperties(FVectorVMExternalFunctionContext& Context)
 {
 	using namespace NDIDynamicMeshLocal;
diff --git a/Engine/Plugins/FX/Niagara/Source/Niagara/Private/DataInterface/NiagaraDataInterfaceDynamicMesh.h b/Engine/Plugins/FX/Niagara/Source/Niagara/Private/DataInterface/NiagaraDataInterfaceDynamicMesh.h
deleted file mode 100644
diff --git a/Engine/Plugins/FX/Niagara/Source/NiagaraEditor/Private/Customizations/SimCache/NiagaraDynamicMeshSimCacheVisualizer.cpp b/Engine/Plugins/FX/Niagara/Source/NiagaraEditor/Private/Customizations/SimCache/NiagaraDynamicMeshSimCacheVisualizer.cpp
new file mode 100644
--- /dev/null
+++ b/Engine/Plugins/FX/Niagara/Source/NiagaraEditor/Private/Customizations/SimCache/NiagaraDynamicMeshSimCacheVisualizer.cpp
@@ -0,0 +1,541 @@
+// Copyright Epic Games, Inc. All Rights Reserved.
+
+#include "NiagaraDynamicMeshSimCacheVisualizer.h"
+
+#include "DataInterface/NiagaraDataInterfaceDynamicMeshSimCacheData.h"
+#include "NiagaraEditorStyle.h"
+
+#include "Engine/StaticMesh.h"
+#include "Framework/MultiBox/MultiBoxBuilder.h"
+#include "HAL/PlatformApplicationMisc.h"
+#include "Modules/ModuleManager.h"
+#include "Subsystems/AssetEditorSubsystem.h"
+#include "UObject/StrongObjectPtr.h"
+#include "ViewModels/NiagaraSimCacheViewModel.h"
+#include "Widgets/Input/SSpinBox.h"
+#include "Widgets/Layout/SWidgetSwitcher.h"
+#include "Widgets/Text/STextBlock.h"
+#include "DetailLayoutBuilder.h"
+#include "Editor.h"
+#include "FileHelpers.h"
+#include "MeshDescription.h"
+#include "MeshAttributes.h"
+#include "PropertyEditorModule.h"
+#include "StaticMeshAttributes.h"
+#include "StaticMeshOperations.h"
+#include "IStructureDetailsView.h"
+#include "SEnumCombo.h"
+
+#define LOCTEXT_NAMESPACE "NiagaraDynamicMeshSimCacheVisualizer"
+
+////////////////////////////////////////////////////////////////////////////////////////////////////////////////////
+
+namespace NDIDynamicMeshSimCacheVisualizer
+{
+	struct FTangentBasis
+	{
+		FVector3f XAxis = FVector3f::XAxisVector;
+		FVector3f YAxis = FVector3f::YAxisVector;
+		FVector3f ZAxis = FVector3f::ZAxisVector;
+	};
+
+	struct FMeshReadHelper
+	{
+		explicit FMeshReadHelper(const FNDIDynamicMeshSimCacheFrame* CacheFrame)
+		{
+			NumVertices = CacheFrame->NumVertices;
+			NumTriangles = CacheFrame->NumTriangles;
+			NumTexCoords = CacheFrame->NumTexCoords;
+
+			Indicies = MakeArrayView((const uint32*)CacheFrame->IndexData.GetData(), CacheFrame->IndexData.Num() / sizeof(uint32));
+
+			if (CacheFrame->PositionOffset != INDEX_NONE)
+			{
+				Positions = MakeArrayView((const FVector3f*)(CacheFrame->VertexData.GetData() + CacheFrame->PositionOffset), NumVertices);
+			}
+			if (CacheFrame->TangentBasisOffset != INDEX_NONE)
+			{
+				TangentBasis = MakeArrayView((const FUintVector2*)(CacheFrame->VertexData.GetData() + CacheFrame->TangentBasisOffset), NumVertices);
+			}
+			if (CacheFrame->TexCoordOffset != INDEX_NONE)
+			{
+				TexCoords = MakeArrayView((const FVector2f*)(CacheFrame->VertexData.GetData() + CacheFrame->TexCoordOffset), NumVertices * NumTexCoords);
+			}
+			if (CacheFrame->ColorOffset != INDEX_NONE)
+			{
+				Colors = MakeArrayView((const FColor*)(CacheFrame->VertexData.GetData() + CacheFrame->ColorOffset), NumVertices);
+			}
+		}
+
+		uint32 GetNumVertices() const { return NumVertices; }
+		uint32 GetNumTriangles() const { return NumTriangles; }
+		uint32 GetNumTexCoords() const { return NumTexCoords; }
+
+		FUintVector3 GetTriangleIndices(uint32 Triangle) const { return FUintVector3(Indicies[Triangle * 3 + 0], Indicies[Triangle * 3 + 1], Indicies[Triangle * 3 + 2]); }
+
+		FVector3f GetVertexPosition(uint32 VertexIndex) const { return Positions[VertexIndex]; }
+		FTangentBasis GeVertexTangents(uint32 VertexIndex) const
+		{
+			const uint32 PackedTangentX = TangentBasis[VertexIndex].X;
+			const uint32 PackedTangentZ = TangentBasis[VertexIndex].Y;
+
+			FTangentBasis Tangents;
+			Tangents.XAxis = FVector3f(float((PackedTangentX >> 0) & 0xff), float((PackedTangentX >> 8) & 0xff), float((PackedTangentX >> 16) & 0xff));
+			Tangents.ZAxis = FVector3f(float((PackedTangentZ >> 0) & 0xff), float((PackedTangentZ >> 8) & 0xff), float((PackedTangentZ >> 16) & 0xff));
+			Tangents.XAxis = (Tangents.XAxis / 127.5f) - 1.0f;
+			Tangents.ZAxis = (Tangents.ZAxis / 127.5f) - 1.0f;
+
+			const float TangentSign = (float((PackedTangentZ >> 24) & 0xff) / 127.5f) - 1.0f;
+
+			Tangents.YAxis = FVector3f::CrossProduct(Tangents.ZAxis, Tangents.XAxis) * TangentSign;
+			Tangents.XAxis = FVector3f::CrossProduct(Tangents.YAxis, Tangents.ZAxis) * TangentSign;
+			return Tangents;
+		}
+		FVector2f GetVertexTexCoord(uint32 VertexIndex, uint32 Layer) const { return TexCoords[(VertexIndex * NumTexCoords) + Layer]; }
+		FColor GetVertexColor(uint32 VertexIndex) const { return Colors[VertexIndex]; }
+
+		bool HasVertexPositions() const { return Positions.Num() > 0; }
+		bool HasVertexTangentBasis() const { return TangentBasis.Num() > 0; }
+		bool HasVertexTexCoords() const { return TexCoords.Num() > 0; }
+		bool HasVertexColors() const { return Colors.Num() > 0; }
+
+	private:
+		uint32							NumVertices = 0;
+		uint32							NumTriangles = 0;
+		uint32							NumTexCoords = 0;
+
+		TConstArrayView<uint32>			Indicies;
+
+		TConstArrayView<FVector3f>		Positions;
+		TConstArrayView<FUintVector2>	TangentBasis;
+		TConstArrayView<FVector2f>		TexCoords;
+		TConstArrayView<FColor>			Colors;
+	};
+
+
+	class SSimCacheView : public SCompoundWidget
+	{
+	public:
+		SLATE_BEGIN_ARGS(SSimCacheView)
+		{}
+		SLATE_END_ARGS()
+
+		virtual ~SSimCacheView() override
+		{
+			ViewModel->OnViewDataChanged().RemoveAll(this);
+		}
+
+		void Construct(const FArguments& InArgs, TSharedPtr<FNiagaraSimCacheViewModel> InViewModel, const UNiagaraDataInterfaceDynamicMeshSimCacheData* InCacheData)
+		{
+			ViewModel = InViewModel;
+			CacheData.Reset(InCacheData);
+
+			ViewModel->OnViewDataChanged().AddSP(this, &SSimCacheView::RefreshContents);
+
+			{
+				FPropertyEditorModule& PropertyModule = FModuleManager::LoadModuleChecked<FPropertyEditorModule>("PropertyEditor");
+
+				FDetailsViewArgs DetailsViewArgs;
+				DetailsViewArgs.bHideSelectionTip = true;
+				DetailsViewArgs.bAllowSearch = false;
+
+				FStructureDetailsViewArgs StructureViewArgs;
+				StructureViewArgs.bShowObjects = true;
+				StructureViewArgs.bShowAssets = true;
+				StructureViewArgs.bShowClasses = true;
+				StructureViewArgs.bShowInterfaces = true;
+
+				StructureDetailsView = PropertyModule.CreateStructureDetailView(DetailsViewArgs, StructureViewArgs, nullptr);
+				StructureDetailsView->GetDetailsView()->SetIsPropertyReadOnlyDelegate(FIsPropertyReadOnly::CreateLambda([](const FPropertyAndParent&) { return true; }));
+			}
+
+			RefreshContents();
+
+			ChildSlot
+			[
+				SNew(SHorizontalBox)
+				+ SHorizontalBox::Slot()
+				[
+					SNew(SVerticalBox)
+					+ SVerticalBox::Slot()
+					.AutoHeight()
+					[
+						MakeToolbar()
+					]
+					+ SVerticalBox::Slot()
+					.AutoHeight()
+					.Padding(2.0)
+					[
+						SNew(STextBlock)
+						.Text(this, &SSimCacheView::GetMeshOverviewText)
+					]
+					+ SVerticalBox::Slot()
+					.AutoHeight()
+					[
+						StructureDetailsView->GetWidget().ToSharedRef()
+					]
+				]
+			];
+		}
+
+		const FNDIDynamicMeshSimCacheFrame* GetCurrentMeshFrame() const
+		{
+			return CacheData->GetFrameData(ViewModel->GetFrameIndex(), CurrentSimTarget);
+		}
+
+		bool HasValidMeshData() const
+		{
+			return ViewingMeshFrameValid;
+		}
+
+		void RefreshContents(bool bForceRefesh = true)
+		{
+			if (const FNDIDynamicMeshSimCacheFrame* CacheFrame = GetCurrentMeshFrame())
+			{
+				ViewingMeshFrameValid = true;
+				ViewingMeshFrame = *CacheFrame;
+				ViewingMeshFrame.IndexData.Empty();
+				ViewingMeshFrame.VertexData.Empty();
+
+				TSharedPtr<FStructOnScope> StructOnScope = MakeShared<FStructOnScope>(FNDIDynamicMeshSimCacheFrame::StaticStruct(), reinterpret_cast<uint8*>(&ViewingMeshFrame));
+				StructureDetailsView->SetStructureData(StructOnScope);
+			}
+			else
+			{
+				ViewingMeshFrameValid = false;
+				StructureDetailsView->SetStructureData(nullptr);
+			}
+		}
+
+		FText GetMeshOverviewText() const
+		{
+			if (const FNDIDynamicMeshSimCacheFrame* CacheFrame = GetCurrentMeshFrame())
+			{
+				return FText::Format(
+					LOCTEXT("MeshDataOverviewFmt", "Mesh has {0} Vertices, {1} Triangles and is approximately {2}kb."),
+					FText::AsNumber(CacheFrame->NumVertices),
+					FText::AsNumber(CacheFrame->NumTriangles),
+					FText::AsNumber(FMath::DivideAndRoundUp(CacheFrame->VertexData.Num()  + CacheFrame->IndexData.Num(), 1024))
+				);
+			}
+			return LOCTEXT("NoMeshData", "No mesh data is stored for this simulation type on this frame.");
+		}
+
+		TSharedRef<SWidget> MakeToolbar()
+		{
+			FSlimHorizontalToolBarBuilder ToolbarBuilder(MakeShareable(new FUICommandList), FMultiBoxCustomization::None);
+
+			ToolbarBuilder.BeginSection("SimMode");
+			{
+				UEnum* SimTargetEnum = StaticEnum<ENiagaraSimTarget>();
+				for (ENiagaraSimTarget SimTarget : {ENiagaraSimTarget::CPUSim, ENiagaraSimTarget::GPUComputeSim})
+				{
+					const FName IconName(SimTarget == ENiagaraSimTarget::CPUSim ? "NiagaraEditor.Stack.CPUIcon" : "NiagaraEditor.Stack.GPUIcon");
+
+					ToolbarBuilder.AddToolBarButton(
+						FUIAction(
+							FExecuteAction::CreateLambda([this, SimTarget]() { CurrentSimTarget = SimTarget; RefreshContents(); }),
+							FCanExecuteAction(),
+							FIsActionChecked::CreateLambda([this, SimTarget]() { return CurrentSimTarget == SimTarget; })
+						),
+						NAME_None,
+						SimTargetEnum->GetDisplayNameTextByValue(int(SimTarget)),
+						SimTargetEnum->GetToolTipTextByIndex(SimTargetEnum->GetIndexByValue(int(SimTarget))),
+						FSlateIcon(FNiagaraEditorStyle::Get().GetStyleSetName(), IconName),
+						EUserInterfaceActionType::ToggleButton
+					);
+				}
+			}
+			ToolbarBuilder.EndSection();
+
+			ToolbarBuilder.BeginSection("Utils");
+			{
+				ToolbarBuilder.AddToolBarButton(
+					FUIAction(
+						FExecuteAction::CreateSP(this, &SSimCacheView::CopyTrianglesToClipboard),
+						FCanExecuteAction::CreateSP(this, &SSimCacheView::HasValidMeshData)
+					),
+					NAME_None,
+					LOCTEXT("TrianglesToClipboard", "Triangles to Clipboard"),
+					LOCTEXT("TrianglesToClipboardTooltip", "Copies the mesh triangles indicies into the clipboard in CSV format"),
+					FSlateIcon()
+				);
+				ToolbarBuilder.AddToolBarButton(
+					FUIAction(
+						FExecuteAction::CreateSP(this, &SSimCacheView::CopyVerticesToClipboard),
+						FCanExecuteAction::CreateSP(this, &SSimCacheView::HasValidMeshData)
+					),
+					NAME_None,
+					LOCTEXT("VerticesToClipboard", "Vertices to Clipboard"),
+					LOCTEXT("VerticesToClipboardTooltip", "Copies the mesh vertices into the clipboard in CSV format"),
+					FSlateIcon()
+				);
+				ToolbarBuilder.AddToolBarButton(
+					FUIAction(
+						FExecuteAction::CreateSP(this, &SSimCacheView::ExportAsStaticMesh),
+						FCanExecuteAction::CreateSP(this, &SSimCacheView::HasValidMeshData)
+					),
+					NAME_None,
+					LOCTEXT("CreateStaticMesh", "Create Static Mesh"),
+					LOCTEXT("CreateStaticMeshTooltip", "Creates a Static Mesh asset from the current mesh data"),
+					FSlateIcon()
+				);
+			}
+			ToolbarBuilder.EndSection();
+
+			return ToolbarBuilder.MakeWidget();
+		}
+
+		void CopyTrianglesToClipboard()
+		{
+			FString ClipboardString;
+			ClipboardString.Append(TEXT("Section,Triangle,Index0,Index1,Index2\n"));
+			if (const FNDIDynamicMeshSimCacheFrame* CacheFrame = GetCurrentMeshFrame())
+			{
+				TConstArrayView<uint32> Indices = MakeArrayView((uint32*)CacheFrame->IndexData.GetData(), CacheFrame->IndexData.Num() / sizeof(uint32));
+				for (int32 iSection = 0; iSection < CacheFrame->Sections.Num(); ++iSection)
+				{
+					const FNDIDynamicMeshSimCacheSection& Section = CacheFrame->Sections[iSection];
+					for (uint32 iTriangle=0; iTriangle < Section.AllocatedTriangles; ++iTriangle)
+					{
+						const uint32 IndexOffset = (Section.TriangleOffset + iTriangle) * 3;
+						ClipboardString.Appendf(TEXT("%u,%u,%u,%u,%u\n"), iSection, iTriangle, Indices[IndexOffset + 0], Indices[IndexOffset + 1], Indices[IndexOffset + 2]);
+					}
+				}
+			}
+			FPlatformApplicationMisc::ClipboardCopy(*ClipboardString);
+		}
+
+		void CopyVerticesToClipboard()
+		{
+			const FNDIDynamicMeshSimCacheFrame* CacheFrame = GetCurrentMeshFrame();
+			if (CacheFrame == nullptr)
+			{
+				return;
+			}
+
+			const FMeshReadHelper MeshReadHelper(CacheFrame);
+
+			// Build used or not information
+			TBitArray<TInlineAllocator<1>> VertexUsed;
+			{
+				VertexUsed.Add(false, CacheFrame->NumVertices);
+
+				for (int32 iSection = 0; iSection < CacheFrame->Sections.Num(); ++iSection)
+				{
+					const FNDIDynamicMeshSimCacheSection& Section = CacheFrame->Sections[iSection];
+					for (uint32 iTriangle = 0; iTriangle < Section.AllocatedTriangles; ++iTriangle)
+					{
+						const FUintVector3 Indices = MeshReadHelper.GetTriangleIndices(Section.TriangleOffset + iTriangle);
+						VertexUsed[Indices.X] = true;
+						VertexUsed[Indices.Y] = true;
+						VertexUsed[Indices.Z] = true;
+					}
+				}
+			}
+
+			// Build Header String
+			FString ClipboardString;
+			ClipboardString.Append(TEXT("Vertex,"));
+			if (MeshReadHelper.HasVertexPositions())
+			{
+				ClipboardString.Append(TEXT("Position.X,Position.Y,Position.Z,"));
+			}
+			if (MeshReadHelper.HasVertexTangentBasis())
+			{
+				ClipboardString.Append(TEXT("TangentX.X,TangentX.Y,TangentX.Z,"));
+				ClipboardString.Append(TEXT("TangentY.X,TangentY.Y,TangentY.Z,"));
+				ClipboardString.Append(TEXT("TangentZ.X,TangentZ.Y,TangentZ.Z,"));
+			}
+			if (MeshReadHelper.HasVertexTexCoords())
+			{
+				for (uint32 iTexCoord = 0; iTexCoord < MeshReadHelper.GetNumTexCoords(); ++iTexCoord)
+				{
+					ClipboardString.Appendf(TEXT("TexCoord[%d].X,TexCoord[%d].Y,"), iTexCoord, iTexCoord);
+				}
+			}
+			if (MeshReadHelper.HasVertexColors())
+			{
+				ClipboardString.Append(TEXT("Color.R,Color.G,Color.B,Color.A,"));
+			}
+
+			ClipboardString[ClipboardString.Len() - 1] = '\n';
+
+			// Build Vertices
+			for (uint32 iVertex=0; iVertex < MeshReadHelper.GetNumVertices(); ++iVertex)
+			{
+				if (VertexUsed[iVertex] == false)
+				{
+					continue;
+				}
+
+				ClipboardString.Appendf(TEXT("%u,"), iVertex);
+				if (MeshReadHelper.HasVertexPositions())
+				{
+					const FVector3f Position = MeshReadHelper.GetVertexPosition(iVertex);
+					ClipboardString.Appendf(TEXT("%f,%f,%f,"), Position.X, Position.Y, Position.Z);
+				}
+				if (MeshReadHelper.HasVertexTangentBasis())
+				{
+					const FTangentBasis Tangents = MeshReadHelper.GeVertexTangents(iVertex);
+					ClipboardString.Appendf(TEXT("%f,%f,%f,"), Tangents.XAxis.X, Tangents.XAxis.Y, Tangents.XAxis.Z);
+					ClipboardString.Appendf(TEXT("%f,%f,%f,"), Tangents.YAxis.X, Tangents.YAxis.Y, Tangents.YAxis.Z);
+					ClipboardString.Appendf(TEXT("%f,%f,%f,"), Tangents.ZAxis.X, Tangents.ZAxis.Y, Tangents.ZAxis.Z);
+				}
+				if (MeshReadHelper.HasVertexTexCoords())
+				{
+					for (uint32 iTexCoord = 0; iTexCoord < MeshReadHelper.GetNumTexCoords(); ++iTexCoord)
+					{
+						const FVector2f TexCoord = MeshReadHelper.GetVertexTexCoord(iVertex, iTexCoord);
+						ClipboardString.Appendf(TEXT("%f,%f,"), TexCoord.X, TexCoord.Y);
+					}
+				}
+				if (MeshReadHelper.HasVertexColors())
+				{
+					const FColor Color = MeshReadHelper.GetVertexColor(iVertex);
+					ClipboardString.Appendf(TEXT("%d,%d,%d,%d,"), Color.R, Color.G, Color.B, Color.A);
+				}
+				ClipboardString[ClipboardString.Len() - 1] = '\n';
+			}
+
+			FPlatformApplicationMisc::ClipboardCopy(*ClipboardString);
+		}
+
+		void ExportAsStaticMesh()
+		{
+			const FNDIDynamicMeshSimCacheFrame* CacheFrame = GetCurrentMeshFrame();
+			if (CacheFrame == nullptr)
+			{
+				return;
+			}
+
+			const FMeshReadHelper MeshReadHelper(CacheFrame);
+
+		//	const uint32 NumTexCoords = ReadbackResult.VertexTexCoordNum;
+		//	const uint32 NumTriangles = ReadbackResult.NumVertices / 3;
+
+			// Create Mesh Description
+			FMeshDescription MeshDescription;
+			TVertexAttributesRef<FVector3f> VertexPositions = MeshDescription.GetVertexPositions();
+			TVertexInstanceAttributesRef<FVector3f> VertexInstanceNormals = MeshDescription.VertexInstanceAttributes().RegisterAttribute<FVector3f>(MeshAttribute::VertexInstance::Normal, 1, FVector3f::ZeroVector, EMeshAttributeFlags::Mandatory);
+			TVertexInstanceAttributesRef<FVector3f> VertexInstanceTangents = MeshDescription.VertexInstanceAttributes().RegisterAttribute<FVector3f>(MeshAttribute::VertexInstance::Tangent, 1, FVector3f::ZeroVector, EMeshAttributeFlags::Mandatory);
+			TVertexInstanceAttributesRef<float> VertexInstanceBinormalSigns = MeshDescription.VertexInstanceAttributes().RegisterAttribute<float>(MeshAttribute::VertexInstance::BinormalSign, 1, 1.0f, EMeshAttributeFlags::Mandatory);
+			TVertexInstanceAttributesRef<FVector4f> VertexInstanceColors = MeshDescription.VertexInstanceAttributes().RegisterAttribute<FVector4f>(MeshAttribute::VertexInstance::Color, 1, FVector4f(1.0f, 1.0f, 1.0f, 1.0f), EMeshAttributeFlags::Lerpable | EMeshAttributeFlags::Mandatory);
+			TVertexInstanceAttributesRef<FVector2f> VertexInstanceUVs = MeshDescription.VertexInstanceAttributes().RegisterAttribute<FVector2f>(MeshAttribute::VertexInstance::TextureCoordinate, MeshReadHelper.GetNumTexCoords(), FVector2f::ZeroVector, EMeshAttributeFlags::Lerpable | EMeshAttributeFlags::Mandatory);
+
+			MeshDescription.EdgeAttributes().RegisterAttribute<bool>(MeshAttribute::Edge::IsHard, 1, false, EMeshAttributeFlags::Mandatory);
+			TPolygonGroupAttributesRef<FName> PolygonGroupSlotNames = MeshDescription.PolygonGroupAttributes().RegisterAttribute<FName>(MeshAttribute::PolygonGroup::ImportedMaterialSlotName, 1, NAME_None, EMeshAttributeFlags::Mandatory); //The unique key to match the mesh material slot
+
+			// Reserve space
+			MeshDescription.ReserveNewVertices(MeshReadHelper.GetNumVertices());
+			MeshDescription.ReserveNewVertexInstances(MeshReadHelper.GetNumTriangles());
+			MeshDescription.ReserveNewEdges(MeshReadHelper.GetNumTriangles());
+			MeshDescription.ReserveNewPolygons(MeshReadHelper.GetNumTriangles());
+			MeshDescription.ReserveNewPolygonGroups(CacheFrame->Sections.Num());
+
+			// Build vertices
+			TArray<FVertexInstanceID> VertexInstances;
+			VertexInstances.Reserve(MeshReadHelper.GetNumVertices());
+			for (uint32 iVertex=0; iVertex < MeshReadHelper.GetNumVertices(); ++iVertex)
+			{
+				const FVertexID VertexID = MeshDescription.CreateVertex();
+				check(VertexID.GetValue() == iVertex);
+				VertexPositions[VertexID] = MeshReadHelper.HasVertexPositions() ? MeshReadHelper.GetVertexPosition(iVertex) : FVector3f::ZeroVector;
+
+				FVertexInstanceID VertexInstanceID = MeshDescription.CreateVertexInstance(VertexID);
+				VertexInstances.Add(VertexInstanceID);
+
+				const FTangentBasis TangentBasis = MeshReadHelper.HasVertexTangentBasis() ? MeshReadHelper.GeVertexTangents(iVertex) : FTangentBasis();
+				const float TangentSign = FVector3f::DotProduct(FVector3f::CrossProduct(TangentBasis.XAxis, TangentBasis.ZAxis), TangentBasis.YAxis) < 0.0f ? -1.0f : 1.0f;
+				VertexInstanceNormals.Set(VertexInstanceID, TangentBasis.XAxis);
+				VertexInstanceTangents.Set(VertexInstanceID, TangentBasis.ZAxis);
+				VertexInstanceBinormalSigns.Set(VertexInstanceID, TangentSign);
+
+				VertexInstanceColors.Set(VertexInstanceID, MeshReadHelper.HasVertexColors() ? MeshReadHelper.GetVertexColor(iVertex).ReinterpretAsLinear() : FLinearColor::White);
+				for (uint32 iTexCoord=0; iTexCoord < MeshReadHelper.GetNumTexCoords(); ++iTexCoord)
+				{
+					VertexInstanceUVs.Set(VertexInstanceID, iTexCoord, MeshReadHelper.GetVertexTexCoord(iVertex, iTexCoord));
+				}
+			}
+
+			// Build sections / triangles
+			TArray<FStaticMaterial> StaticMaterials;
+			for (int32 iSection=0; iSection < CacheFrame->Sections.Num(); ++iSection)
+			{
+				const FNDIDynamicMeshSimCacheSection& Section = CacheFrame->Sections[iSection];
+
+				const FName MaterialSlotName(*FString::Printf(TEXT("Section%d"), iSection));
+				const FPolygonGroupID PolyGroupId = MeshDescription.CreatePolygonGroup();
+				PolygonGroupSlotNames[PolyGroupId] = MaterialSlotName;
+
+				FStaticMaterial& StaticMaterial			= StaticMaterials.AddDefaulted_GetRef();
+				StaticMaterial.MaterialInterface		= Section.Material;
+				StaticMaterial.MaterialSlotName			= MaterialSlotName;
+				StaticMaterial.ImportedMaterialSlotName	= MaterialSlotName;
+
+				for (uint32 iTriangle=0; iTriangle < Section.AllocatedTriangles; ++iTriangle)
+				{
+					const FUintVector3 Indices = MeshReadHelper.GetTriangleIndices(Section.TriangleOffset + iTriangle);
+					const FPolygonID PolygonID = MeshDescription.CreatePolygon(PolyGroupId, { Indices.X, Indices.Y, Indices.Z });
+					MeshDescription.ComputePolygonTriangulation(PolygonID);
+				}
+			}
+
+			UStaticMesh* StaticMesh = NewObject<UStaticMesh>();
+			StaticMesh->SetNumSourceModels(1);
+			{
+				FMeshBuildSettings& MeshBuildSettings = StaticMesh->GetSourceModel(0).BuildSettings;
+				MeshBuildSettings.bRecomputeNormals = false;
+				MeshBuildSettings.bRecomputeTangents = false;
+			}
+
+			StaticMesh->SetStaticMaterials(StaticMaterials);
+
+			// Build Mesh
+			UStaticMesh::FBuildMeshDescriptionsParams Params;
+			Params.bFastBuild = WITH_EDITOR ? false : true;
+			Params.bUseHashAsGuid = true;
+			Params.bMarkPackageDirty = false;
+			Params.bCommitMeshDescription = true;
+			Params.bAllowCpuAccess = false;
+
+			StaticMesh->BuildFromMeshDescriptions({ &MeshDescription }, Params);
+
+			StaticMesh->SetFlags(RF_Public | RF_Standalone);
+
+			TArray<UObject*> AssetsToSave = { StaticMesh };
+			TArray<UObject*> SavedAssets;
+			FEditorFileUtils::SaveAssetsAs(AssetsToSave, SavedAssets);
+
+			if (SavedAssets.Num() > 0)
+			{
+				UAssetEditorSubsystem* AssetEditorSubsystem = GEditor->GetEditorSubsystem<UAssetEditorSubsystem>();
+				AssetEditorSubsystem->OpenEditorForAssets_Advanced(SavedAssets, EToolkitMode::Standalone);
+			}
+
+			StaticMesh->ClearFlags(RF_Public | RF_Standalone);
+		}
+
+	private:
+		TSharedPtr<FNiagaraSimCacheViewModel>									ViewModel;
+		TStrongObjectPtr<const UNiagaraDataInterfaceDynamicMeshSimCacheData>	CacheData;
+
+		bool																	ViewingMeshFrameValid = false;
+		FNDIDynamicMeshSimCacheFrame											ViewingMeshFrame;
+		TSharedPtr<IStructureDetailsView>										StructureDetailsView;
+
+		ENiagaraSimTarget														CurrentSimTarget = ENiagaraSimTarget::CPUSim;
+	};
+} // NDIDynamicMeshSimCacheVisualizer
+
+////////////////////////////////////////////////////////////////////////////////////////////////////////////////////
+
+TSharedPtr<SWidget> FNiagaraDynamicMeshSimCacheVisualizer::CreateWidgetFor(const UObject* InCachedData, TSharedPtr<FNiagaraSimCacheViewModel> ViewModel)
+{
+	using namespace NDIDynamicMeshSimCacheVisualizer;
+
+	const UNiagaraDataInterfaceDynamicMeshSimCacheData* CachedData = CastChecked<const UNiagaraDataInterfaceDynamicMeshSimCacheData>(InCachedData);
+	return SNew(SSimCacheView, ViewModel, CachedData);
+}
+
+#undef LOCTEXT_NAMESPACE
diff --git a/Engine/Plugins/FX/Niagara/Source/NiagaraEditor/Private/Customizations/SimCache/NiagaraDynamicMeshSimCacheVisualizer.h b/Engine/Plugins/FX/Niagara/Source/NiagaraEditor/Private/Customizations/SimCache/NiagaraDynamicMeshSimCacheVisualizer.h
new file mode 100644
--- /dev/null
+++ b/Engine/Plugins/FX/Niagara/Source/NiagaraEditor/Private/Customizations/SimCache/NiagaraDynamicMeshSimCacheVisualizer.h
@@ -0,0 +1,15 @@
+// Copyright Epic Games, Inc. All Rights Reserved.
+
+#pragma once
+
+#include "NiagaraCommon.h"
+#include "Customizations/NiagaraDataInterfaceSimCacheVisualizer.h"
+
+// Visualizer for dynamic mesh sim cache data
+class FNiagaraDynamicMeshSimCacheVisualizer : public INiagaraDataInterfaceSimCacheVisualizer
+{
+public:
+	virtual ~FNiagaraDynamicMeshSimCacheVisualizer() override = default;
+	
+	virtual TSharedPtr<SWidget> CreateWidgetFor(const UObject* CachedData, TSharedPtr<FNiagaraSimCacheViewModel> ViewModel) override;
+};
diff --git a/Engine/Plugins/FX/Niagara/Source/NiagaraEditor/Private/NiagaraEditorModule.cpp b/Engine/Plugins/FX/Niagara/Source/NiagaraEditor/Private/NiagaraEditorModule.cpp
--- a/Engine/Plugins/FX/Niagara/Source/NiagaraEditor/Private/NiagaraEditorModule.cpp
+++ b/Engine/Plugins/FX/Niagara/Source/NiagaraEditor/Private/NiagaraEditorModule.cpp
@@ -59,6 +59,7 @@
 #include "DataInterface/NiagaraDataInterfaceDataChannelRead.h"
 #include "DataInterface/NiagaraDataInterfaceDataChannelWrite.h"
 #include "DataInterface/NiagaraDataInterfaceDataTable.h"
+#include "DataInterface/NiagaraDataInterfaceDynamicMesh.h"
 #include "DataInterface/NiagaraDataInterfaceMemoryBuffer.h"
 #include "DataInterface/NiagaraDataInterfaceSimpleCounter.h"
 #include "NiagaraDataInterfaceRenderTarget2D.h"
@@ -124,6 +125,7 @@
 #include "Customizations/NiagaraDataChannelDetails.h"
 #include "Customizations/SimCache/NiagaraArraySimCacheVisualizer.h"
 #include "Customizations/SimCache/FNiagaraDataChannelSimCacheVisualizer.h"
+#include "Customizations/SimCache/NiagaraDynamicMeshSimCacheVisualizer.h"
 #include "Customizations/SimCache/NiagaraMemoryBufferSimCacheVisualizer.h"
 #include "Customizations/SimCache/NiagaraRenderTargetSimCacheVisualizer.h"
 #include "Customizations/SimCache/NiagaraRenderTargetVolumeSimCacheVisualizer.h"
@@ -1592,6 +1594,7 @@ void FNiagaraEditorModule::StartupModule()
 
 	// Register sim cache visualizers
 	RegisterDataInterfaceCacheVisualizer(UNiagaraDataInterfaceDataChannelWrite::StaticClass(), MakeShared<FNiagaraDataChannelSimCacheVisualizer>());
+	RegisterDataInterfaceCacheVisualizer(UNiagaraDataInterfaceDynamicMesh::StaticClass(), MakeShared<FNiagaraDynamicMeshSimCacheVisualizer>());
 	RegisterDataInterfaceCacheVisualizer(UNiagaraDataInterfaceMemoryBuffer::StaticClass(), MakeShared<FNiagaraMemoryBufferSimCacheVisualizer>());
 	RegisterDataInterfaceCacheVisualizer(UNiagaraDataInterfaceRenderTarget2D::StaticClass(), MakeShared<FNiagaraRenderTargetSimCacheVisualizer>());
 	RegisterDataInterfaceCacheVisualizer(UNiagaraDataInterfaceRenderTargetVolume::StaticClass(), MakeShared<FNiagaraRenderTargetVolumeSimCacheVisualizer>());
