- Use the aliased variable when finding the default mode
- This fixes an issue if the first time a variable was encountered it was in StackContext as we would ignore the binding

diff --git a/Engine/Plugins/FX/Niagara/Source/NiagaraEditor/Private/NiagaraHlslTranslator.cpp b/Engine/Plugins/FX/Niagara/Source/NiagaraEditor/Private/NiagaraHlslTranslator.cpp
--- a/Engine/Plugins/FX/Niagara/Source/NiagaraEditor/Private/NiagaraHlslTranslator.cpp
+++ b/Engine/Plugins/FX/Niagara/Source/NiagaraEditor/Private/NiagaraHlslTranslator.cpp
@@ -1020,7 +1020,8 @@ void TNiagaraHlslTranslator<GraphBridge>::BuildMissingDefaults()
 				DefaultPinNode = GraphBridge::GetOwningNode(DefaultPin);
 				if (const FGraph* DefaultPinGraph = GraphBridge::GetOwningGraph(DefaultPinNode))
 				{
-					DefaultMode = GraphBridge::GetGraphDefaultMode(DefaultPinGraph, Var, DefaultBinding);
+					const FNiagaraVariableBase* AliasedVar = UniqueVarToAliasedVar.Find(Var);
+					DefaultMode = GraphBridge::GetGraphDefaultMode(DefaultPinGraph, AliasedVar ? *AliasedVar : FNiagaraVariableBase(Var), DefaultBinding);
 				}
 			}
 
@@ -5668,6 +5669,7 @@ void TNiagaraHlslTranslator<GraphBridge>::InitializeParameterMapDefaults(int32 P
 	UniqueVarToDefaultPin.Empty();
 	UniqueVarToWriteToParamMap.Empty();
 	UniqueVarToChunk.Empty();
+	UniqueVarToAliasedVar.Empty();
 
 	FNiagaraHlslTranslationStage& ActiveStage = TranslationStages[ActiveStageIdx];
 	// First pass just use the current parameter map.
@@ -5684,6 +5686,7 @@ void TNiagaraHlslTranslator<GraphBridge>::InitializeParameterMapDefaults(int32 P
 					!UniqueVars.Contains(Var))
 				{
 					UniqueVars.Add(Var);
+					UniqueVarToAliasedVar.Add(Var, AliasedVar);
 					const FInputPin* DefaultPin = History.GetDefaultValuePin(i);
 					UniqueVarToDefaultPin.Add(Var, DefaultPin);
 					UniqueVarToWriteToParamMap.Add(Var, true);
@@ -5709,6 +5712,7 @@ void TNiagaraHlslTranslator<GraphBridge>::InitializeParameterMapDefaults(int32 P
 					!UniqueVars.Contains(Var))
 				{
 					UniqueVars.Add(Var);
+					UniqueVarToAliasedVar.Add(Var, AliasedVar);
 					const FInputPin* DefaultPin = History.GetDefaultValuePin(i);
 					UniqueVarToDefaultPin.Add(Var, DefaultPin);
 					UniqueVarToWriteToParamMap.Add(Var, false);
@@ -5733,7 +5737,8 @@ void TNiagaraHlslTranslator<GraphBridge>::InitializeParameterMapDefaults(int32 P
 				DefaultPinNode = GraphBridge::GetOwningNode(DefaultPin);
 				if (const FGraph* DefaultPinGraph = GraphBridge::GetOwningGraph(DefaultPinNode))
 				{
-					DefaultMode = GraphBridge::GetGraphDefaultMode(DefaultPinGraph, Var, DefaultBinding);
+					const FNiagaraVariableBase* AliasedVar = UniqueVarToAliasedVar.Find(Var);
+					DefaultMode = GraphBridge::GetGraphDefaultMode(DefaultPinGraph, AliasedVar ? *AliasedVar : FNiagaraVariableBase(Var), DefaultBinding);
 				}
 			}
 
diff --git a/Engine/Plugins/FX/Niagara/Source/NiagaraEditor/Private/NiagaraHlslTranslator.h b/Engine/Plugins/FX/Niagara/Source/NiagaraEditor/Private/NiagaraHlslTranslator.h
--- a/Engine/Plugins/FX/Niagara/Source/NiagaraEditor/Private/NiagaraHlslTranslator.h
+++ b/Engine/Plugins/FX/Niagara/Source/NiagaraEditor/Private/NiagaraHlslTranslator.h
@@ -550,6 +550,9 @@ protected:
 	// Map ofthe primary output variable description to the actual chunk id that wrote to it.
 	TMap<FNiagaraVariable, int32> UniqueVarToChunk;
 
+	// Map of the output variable to the found aliased variable, used to discover the default binding, etc
+	TMap<FNiagaraVariableBase, FNiagaraVariableBase> UniqueVarToAliasedVar;
+
 	TArray<int32> ParameterMapForIndexStack;
 
 	// Strings to be inserted within the main function
diff --git a/Engine/Source/Runtime/Core/Private/UObject/FortniteMainBranchObjectVersion.cpp b/Engine/Source/Runtime/Core/Private/UObject/FortniteMainBranchObjectVersion.cpp
--- a/Engine/Source/Runtime/Core/Private/UObject/FortniteMainBranchObjectVersion.cpp
+++ b/Engine/Source/Runtime/Core/Private/UObject/FortniteMainBranchObjectVersion.cpp
@@ -11,7 +11,7 @@ TMap<FGuid, FGuid> FFortniteMainBranchObjectVersion::GetSystemGuids()
 	SystemGuids.Add(DevGuids.MATERIALSHADERMAP_DERIVEDDATA_VER, FGuid("04B8C817071A453C9EBEFBFB51D5F785"));
 	SystemGuids.Add(DevGuids.NANITE_DERIVEDDATA_VER, FGuid("FB5EAC6147EF4591B0D5EB7A85062AC9"));
 	SystemGuids.Add(DevGuids.NIAGARASHADERMAP_DERIVEDDATA_VER, FGuid("6360A977062842A29ED30E3A7ACB0E64"));
-	SystemGuids.Add(DevGuids.Niagara_LatestScriptCompileVersion, FGuid("18843C17BE24154D9DD3EA597ADC5FF5"));
+	SystemGuids.Add(DevGuids.Niagara_LatestScriptCompileVersion, FGuid("081BBC4CD785874B8E77E9AEAEBD0537"));
 	SystemGuids.Add(DevGuids.SkeletalMeshDerivedDataVersion, FGuid("716424BCEC314E969ACFE36D2C2ECE49"));
 	SystemGuids.Add(DevGuids.STATICMESH_DERIVEDDATA_VER, FGuid("3ABF32149F9D4D838750368AB95573BE"));
 	SystemGuids.Add(DevGuids.MaterialTranslationDDCVersion, FGuid("84CF0B79D69D4A4BBE8CAFD4B3D48AB7"));
