Enable GPUScene support for Niagara mesh rendering on mobile platforms

diff --git a/Engine/Plugins/FX/Niagara/Shaders/Private/NiagaraMeshVertexFactory.ush b/Engine/Plugins/FX/Niagara/Shaders/Private/NiagaraMeshVertexFactory.ush
--- a/Engine/Plugins/FX/Niagara/Shaders/Private/NiagaraMeshVertexFactory.ush
+++ b/Engine/Plugins/FX/Niagara/Shaders/Private/NiagaraMeshVertexFactory.ush
@@ -860,6 +860,11 @@ uint VertexFactoryGetPrimitiveId(FVertexFactoryInterpolantsVSToPS Interpolants)
 #endif
 }
 
+float4 VertexFactoryGetInstanceHitProxyId(FVertexFactoryInput Input, FVertexFactoryIntermediates Intermediates)
+{
+	return 0;
+}
+
 float4 VertexFactoryGetTranslatedPrimitiveVolumeBounds(FVertexFactoryInterpolantsVSToPS Interpolants)
 {
 	return float4(0,0,0,0);
diff --git a/Engine/Plugins/FX/Niagara/Source/Niagara/Private/NiagaraRendererMeshes.cpp b/Engine/Plugins/FX/Niagara/Source/Niagara/Private/NiagaraRendererMeshes.cpp
--- a/Engine/Plugins/FX/Niagara/Source/Niagara/Private/NiagaraRendererMeshes.cpp
+++ b/Engine/Plugins/FX/Niagara/Source/Niagara/Private/NiagaraRendererMeshes.cpp
@@ -318,8 +318,7 @@ void FNiagaraRendererMeshes::PrepareParticleMeshRenderData(FRHICommandListBase&
 	ParticleMeshRenderData.Collector = &Collector;
 
 #if NIAGARA_ENABLE_GPU_SCENE_MESHES
-	ParticleMeshRenderData.bUseGPUScene = FeatureLevel > ERHIFeatureLevel::ES3_1	
-		&& UseGPUScene(SceneProxy->GetScene().GetShaderPlatform(), FeatureLevel);
+	ParticleMeshRenderData.bUseGPUScene = UseGPUScene(SceneProxy->GetScene().GetShaderPlatform(), FeatureLevel);
 #else
 	ParticleMeshRenderData.bUseGPUScene = false;
 #endif
diff --git a/Engine/Plugins/FX/Niagara/Source/NiagaraVertexFactories/Private/NiagaraMeshVertexFactory.cpp b/Engine/Plugins/FX/Niagara/Source/NiagaraVertexFactories/Private/NiagaraMeshVertexFactory.cpp
--- a/Engine/Plugins/FX/Niagara/Source/NiagaraVertexFactories/Private/NiagaraMeshVertexFactory.cpp
+++ b/Engine/Plugins/FX/Niagara/Source/NiagaraVertexFactories/Private/NiagaraMeshVertexFactory.cpp
@@ -144,7 +144,7 @@ void FNiagaraMeshVertexFactory::InitRHI(FRHICommandListBase& RHICmdList)
 	if (bAddPrimitiveIDElement)
 	{
 		// TODO: Support GPU Scene on mobile? Maybe only for CPU particles?
-		AddPrimitiveIdStreamElement(EVertexInputStreamType::Default, Elements, 13, 0xFF);
+		AddPrimitiveIdStreamElement(EVertexInputStreamType::Default, Elements, 13, 13);
 	}
 #endif
 
@@ -173,13 +173,17 @@ void FNiagaraMeshVertexFactory::ModifyCompilationEnvironment(const FVertexFactor
 
 #if NIAGARA_ENABLE_GPU_SCENE_MESHES
 	const ERHIFeatureLevel::Type MaxSupportedFeatureLevel = GetMaxSupportedFeatureLevel(Parameters.Platform);
-
-	// TODO: Support GPU Scene on mobile?
-	const bool bUseGPUScene = UseGPUScene(Parameters.Platform, MaxSupportedFeatureLevel) && MaxSupportedFeatureLevel > ERHIFeatureLevel::ES3_1;
+	const bool bUseGPUScene = UseGPUScene(Parameters.Platform, MaxSupportedFeatureLevel);
 	const bool bSupportsPrimitiveIdStream = Parameters.VertexFactoryType->SupportsPrimitiveIdStream();
 	
 	OutEnvironment.SetDefine(TEXT("VF_SUPPORTS_PRIMITIVE_SCENE_DATA"), bSupportsPrimitiveIdStream && bUseGPUScene);
 	OutEnvironment.SetDefine(TEXT("VF_REQUIRES_PER_INSTANCE_CUSTOM_DATA"), bSupportsPrimitiveIdStream && bUseGPUScene);
+
+	// Mobile GPUScene implementation relies on USE_INSTANCE_CULLING define for rendering instanced meshes
+	if (bUseGPUScene && MaxSupportedFeatureLevel == ERHIFeatureLevel::ES3_1)
+	{
+		OutEnvironment.SetDefine(TEXT("USE_INSTANCE_CULLING"), 1);
+	}
 #endif
 
 	if (RHISupportsManualVertexFetch(Parameters.Platform))
