Enabled debugger overview by default

diff --git a/Engine/Plugins/FX/Niagara/Source/Niagara/Public/NiagaraDebuggerCommon.h b/Engine/Plugins/FX/Niagara/Source/Niagara/Public/NiagaraDebuggerCommon.h
--- a/Engine/Plugins/FX/Niagara/Source/Niagara/Public/NiagaraDebuggerCommon.h
+++ b/Engine/Plugins/FX/Niagara/Source/Niagara/Public/NiagaraDebuggerCommon.h
@@ -475,7 +475,7 @@ struct FNiagaraDebugHUDSettingsData
 
 	/** When enabled the overview display will be enabled. */
 	UPROPERTY()
-	bool bOverviewEnabled = false;
+	bool bOverviewEnabled = true;
 
 	UPROPERTY(EditAnywhere, Category = "Debug Overview", meta = (DisplayName = "Debug Overview Mode"))
 	ENiagaraDebugHUDOverviewMode OverviewMode = ENiagaraDebugHUDOverviewMode::Overview;
