Undo changelist 41963565, which was causing NiagaraDebugHud to appear by default in staged builds

diff --git a/Engine/Plugins/FX/Niagara/Source/Niagara/Public/NiagaraDebuggerCommon.h b/Engine/Plugins/FX/Niagara/Source/Niagara/Public/NiagaraDebuggerCommon.h
--- a/Engine/Plugins/FX/Niagara/Source/Niagara/Public/NiagaraDebuggerCommon.h
+++ b/Engine/Plugins/FX/Niagara/Source/Niagara/Public/NiagaraDebuggerCommon.h
@@ -475,7 +475,7 @@ struct FNiagaraDebugHUDSettingsData
 
 	/** When enabled the overview display will be enabled. */
 	UPROPERTY()
-	bool bOverviewEnabled = true;
+	bool bOverviewEnabled = false;
 
 	UPROPERTY(EditAnywhere, Category = "Debug Overview", meta = (DisplayName = "Debug Overview Mode"))
 	ENiagaraDebugHUDOverviewMode OverviewMode = ENiagaraDebugHUDOverviewMode::Overview;
