- Add explicit ColorGradient type to distribution
- This fixes various quirks in the UI where hitting undo / redo would change the color view making it confusing
- This does not fix pasting from one distribution to another as the context around the mode is lost

#prefligjht 680c10025a2e482944e41d7b
#rnx

diff --git a/Engine/Plugins/FX/Niagara/Source/Niagara/Internal/Stateless/NiagaraStatelessEmitterDataBuildContext.h b/Engine/Plugins/FX/Niagara/Source/Niagara/Internal/Stateless/NiagaraStatelessEmitterDataBuildContext.h
--- a/Engine/Plugins/FX/Niagara/Source/Niagara/Internal/Stateless/NiagaraStatelessEmitterDataBuildContext.h
+++ b/Engine/Plugins/FX/Niagara/Source/Niagara/Internal/Stateless/NiagaraStatelessEmitterDataBuildContext.h
@@ -94,6 +94,7 @@ public:
 				case ENiagaraDistributionMode::NonUniformRange:		FNiagaraStatelessBuiltDistribution::SetIsRandom(BuiltDistribution); break;
 				case ENiagaraDistributionMode::UniformCurve:		FNiagaraStatelessBuiltDistribution::SetIsUniform(BuiltDistribution); break;
 				case ENiagaraDistributionMode::NonUniformCurve:		break;
+				case ENiagaraDistributionMode::ColorGradient:		break;
 				default:											checkNoEntry(); break;
 			}
 
@@ -139,7 +140,7 @@ public:
 	FNiagaraStatelessBuiltDistributionType AddDistributionAsCurve(const TDistribution& Distribution, TType DefaultValue) const
 	{
 		FNiagaraStatelessBuiltDistributionType BuiltDistribution = FNiagaraStatelessBuiltDistribution::GetDefault();
-		if (ensure(Distribution.Mode == ENiagaraDistributionMode::UniformCurve || Distribution.Mode == ENiagaraDistributionMode::NonUniformCurve) && Distribution.Values.Num() > 1)
+		if (ensure((Distribution.IsCurve() || Distribution.IsGradient()) && Distribution.Values.Num() > 1))
 		{
 			FNiagaraStatelessBuiltDistribution::SetLookupParameters(BuiltDistribution, AddStaticData(Distribution.Values), Distribution.Values.Num(), Distribution.ValuesTimeRange);
 		}
diff --git a/Engine/Plugins/FX/Niagara/Source/Niagara/Private/Stateless/NiagaraStatelessDistribution.cpp b/Engine/Plugins/FX/Niagara/Source/Niagara/Private/Stateless/NiagaraStatelessDistribution.cpp
--- a/Engine/Plugins/FX/Niagara/Source/Niagara/Private/Stateless/NiagaraStatelessDistribution.cpp
+++ b/Engine/Plugins/FX/Niagara/Source/Niagara/Private/Stateless/NiagaraStatelessDistribution.cpp
@@ -280,8 +280,9 @@ namespace NiagaraStatelessDistributionPrivate
 				break;
 			case ENiagaraDistributionMode::UniformCurve:
 			case ENiagaraDistributionMode::NonUniformCurve:
+			case ENiagaraDistributionMode::ColorGradient:
 			{
-				const int32 ExpectedChannels = InMode == ENiagaraDistributionMode::NonUniformCurve ? InChannelCount : 1;
+				const int32 ExpectedChannels = InMode == ENiagaraDistributionMode::UniformCurve ? 1 : InChannelCount;
 				if (InChannelCurves.Num() >= ExpectedChannels)
 				{
 					MaxLutSampleCount = FMath::Max(MaxLutSampleCount, 2);
@@ -293,7 +294,7 @@ namespace NiagaraStatelessDistributionPrivate
 					{
 						for (int32 ChannelIndex=0; ChannelIndex < InChannelCount; ++ChannelIndex)
 						{
-							const int32 LUTChannelIndex = InMode == ENiagaraDistributionMode::NonUniformCurve ? ChannelIndex : 0;
+							const int32 LUTChannelIndex = InMode == ENiagaraDistributionMode::UniformCurve ? 0 : ChannelIndex;
 							InValueAndChannelAccessor(SampleIndex, ChannelIndex) = LUT[(SampleIndex * ExpectedChannels) + LUTChannelIndex];
 						}
 					}
@@ -306,9 +307,9 @@ namespace NiagaraStatelessDistributionPrivate
 	template<typename FValueContainerSetNum, typename FValueAndChannelAccessor>
 	void UpdateDistributionCurveValues(ENiagaraDistributionMode InMode, const TArray<float>& InChannelConstantsAndRanges, const TArray<FRichCurve>& InChannelCurves, int32 InChannelCount, FVector2f& OutTimeRange, FValueContainerSetNum InContainerNum, FValueAndChannelAccessor InValueAndChannelAccessor, int32 MaxLutSampleCount, ENiagaraDistributionCurveLUTMode LUTMode)
 	{
-		check(InMode == ENiagaraDistributionMode::UniformCurve || InMode == ENiagaraDistributionMode::NonUniformCurve);
+		check(InMode == ENiagaraDistributionMode::UniformCurve || InMode == ENiagaraDistributionMode::NonUniformCurve || InMode == ENiagaraDistributionMode::ColorGradient);
 
-		const int32 ExpectedChannels = InMode == ENiagaraDistributionMode::NonUniformCurve ? InChannelCount : 1;
+		const int32 ExpectedChannels = InMode == ENiagaraDistributionMode::UniformCurve ? 1 : InChannelCount;
 		if (InChannelCurves.Num() >= ExpectedChannels)
 		{
 			MaxLutSampleCount = FMath::Max(MaxLutSampleCount, 2);
@@ -320,7 +321,7 @@ namespace NiagaraStatelessDistributionPrivate
 			{
 				for (int32 ChannelIndex=0; ChannelIndex < InChannelCount; ++ChannelIndex)
 				{
-					const int32 LUTChannelIndex = InMode == ENiagaraDistributionMode::NonUniformCurve ? ChannelIndex : 0;
+					const int32 LUTChannelIndex = InMode == ENiagaraDistributionMode::UniformCurve ? 0 : ChannelIndex;
 					InValueAndChannelAccessor(SampleIndex, ChannelIndex) = LUT[(SampleIndex * ExpectedChannels) + LUTChannelIndex];
 				}
 			}
diff --git a/Engine/Plugins/FX/Niagara/Source/Niagara/Public/Stateless/NiagaraStatelessDistribution.h b/Engine/Plugins/FX/Niagara/Source/Niagara/Public/Stateless/NiagaraStatelessDistribution.h
--- a/Engine/Plugins/FX/Niagara/Source/Niagara/Public/Stateless/NiagaraStatelessDistribution.h
+++ b/Engine/Plugins/FX/Niagara/Source/Niagara/Public/Stateless/NiagaraStatelessDistribution.h
@@ -19,7 +19,8 @@ enum class ENiagaraDistributionMode
 	UniformRange,
 	NonUniformRange,
 	UniformCurve,
-	NonUniformCurve
+	NonUniformCurve,
+	ColorGradient,
 };
 
 enum class ENiagaraDistributionCurveLUTMode
@@ -49,6 +50,7 @@ struct FNiagaraDistributionBase
 	bool IsConstant() const { return Mode == ENiagaraDistributionMode::UniformConstant || Mode == ENiagaraDistributionMode::NonUniformConstant; }
 	bool IsUniform() const { return Mode == ENiagaraDistributionMode::UniformConstant || Mode == ENiagaraDistributionMode::UniformRange; }
 	bool IsCurve() const { return Mode == ENiagaraDistributionMode::UniformCurve || Mode == ENiagaraDistributionMode::NonUniformCurve; }
+	bool IsGradient() const { return Mode == ENiagaraDistributionMode::ColorGradient; }
 	bool IsRange() const { return Mode == ENiagaraDistributionMode::UniformRange || Mode == ENiagaraDistributionMode::NonUniformRange; }
 
 #if WITH_EDITORONLY_DATA
@@ -109,6 +111,7 @@ struct FNiagaraDistributionRangeInt
 	bool IsConstant() const { return Mode == ENiagaraDistributionMode::UniformConstant || Mode == ENiagaraDistributionMode::NonUniformConstant; }
 	bool IsUniform() const { return Mode == ENiagaraDistributionMode::UniformConstant || Mode == ENiagaraDistributionMode::UniformRange; }
 	bool IsCurve() const { return Mode == ENiagaraDistributionMode::UniformCurve || Mode == ENiagaraDistributionMode::NonUniformCurve; }
+	bool IsGradient() const { return Mode == ENiagaraDistributionMode::ColorGradient; }
 	bool IsRange() const { return Mode == ENiagaraDistributionMode::UniformRange || Mode == ENiagaraDistributionMode::NonUniformRange; }
 
 #if WITH_EDITORONLY_DATA
diff --git a/Engine/Plugins/FX/Niagara/Source/NiagaraEditor/Private/Stateless/NiagaraDistributionPropertyCustomization.cpp b/Engine/Plugins/FX/Niagara/Source/NiagaraEditor/Private/Stateless/NiagaraDistributionPropertyCustomization.cpp
--- a/Engine/Plugins/FX/Niagara/Source/NiagaraEditor/Private/Stateless/NiagaraDistributionPropertyCustomization.cpp
+++ b/Engine/Plugins/FX/Niagara/Source/NiagaraEditor/Private/Stateless/NiagaraDistributionPropertyCustomization.cpp
@@ -227,11 +227,9 @@ public:
 			case ENiagaraDistributionMode::NonUniformRange:
 				return ENiagaraDistributionEditorMode::ColorRange;
 			case ENiagaraDistributionMode::NonUniformCurve:
-				if (CurrentMode.IsSet())
-				{
-					return CurrentMode.GetValue() == ENiagaraDistributionEditorMode::NonUniformCurve ? ENiagaraDistributionEditorMode::NonUniformCurve : ENiagaraDistributionEditorMode::ColorGradient;
-				}
-				return AreCurveKeyTimesNormalized() ?  ENiagaraDistributionEditorMode::ColorGradient : ENiagaraDistributionEditorMode::NonUniformCurve;
+				return ENiagaraDistributionEditorMode::NonUniformCurve;
+			case ENiagaraDistributionMode::ColorGradient:
+				return ENiagaraDistributionEditorMode::ColorGradient;
 			default:
 				return ENiagaraDistributionEditorMode::ColorConstant;
 			}
@@ -274,6 +272,8 @@ public:
 				return ENiagaraDistributionEditorMode::UniformCurve;
 			case ENiagaraDistributionMode::NonUniformCurve:
 				return ENiagaraDistributionEditorMode::NonUniformCurve;
+			case ENiagaraDistributionMode::ColorGradient:
+				return ENiagaraDistributionEditorMode::ColorGradient;
 			default:
 				return ENiagaraDistributionEditorMode::UniformConstant;
 			}
@@ -304,8 +304,9 @@ public:
 		case ENiagaraDistributionEditorMode::UniformCurve:
 			return ENiagaraDistributionMode::UniformCurve;
 		case ENiagaraDistributionEditorMode::NonUniformCurve:
-		case ENiagaraDistributionEditorMode::ColorGradient:
 			return ENiagaraDistributionMode::NonUniformCurve;
+		case ENiagaraDistributionEditorMode::ColorGradient:
+			return ENiagaraDistributionMode::ColorGradient;
 		default:
 			return ENiagaraDistributionMode::UniformConstant;
 		}
@@ -330,8 +331,7 @@ public:
 
 		// If the distribution mode is not changing and we are switching to be a color gradient we need to make sure the keys are all normalized (i.e. 0-1 time)
 		// if they are not we need to perform a destructive action and trim the keys to 0-1 time range or the widget will not display the keys correctly.
-		bool bNeedsUpdateForColorGradient = false;
-		if (NewSourceMode == SourceDistribution->Mode && InEditorMode == ENiagaraDistributionEditorMode::ColorGradient && !AreCurveKeyTimesNormalized())
+		if (InEditorMode == ENiagaraDistributionEditorMode::ColorGradient && !AreCurveKeyTimesNormalized())
 		{
 			const EAppReturnType::Type DialogResponse = FMessageDialog::Open(
 				EAppMsgType::OkCancel,
@@ -343,26 +343,18 @@ public:
 			{
 				return;
 			}
-			bNeedsUpdateForColorGradient = true;
 		}
 
 		TGuardValue<bool> UpdateGuard(bUpdatingDistributionModeInternally, true);
 		const int32 NewEditorNumChannels = FNiagaraDistributionEditorUtilities::IsUniform(InEditorMode) ? 1 : SourceNumChannels;
 
-		if (SourceDistribution->Mode != NewSourceMode || bNeedsUpdateForColorGradient)
+		if (SourceDistribution->Mode != NewSourceMode)
 		{
 			const FScopedTransaction Transaction(LOCTEXT("SetDistributionModeTransaction", "Set distribution mode"));
 			OwnerObject->Modify();
 			PropertyHandle->NotifyPreChange();
 
-			if (bNeedsUpdateForColorGradient)
-			{
-				MigrateDataToColorGradient();
-			}
-			else
-			{
-				MigrateDataFromModeChange(NewSourceMode, NewEditorNumChannels);
-			}
+			MigrateDataFromModeChange(NewSourceMode, NewEditorNumChannels);
 			SourceDistribution->Mode = NewSourceMode;
 
 			PropertyHandle->NotifyPostChange(EPropertyChangeType::ValueSet);
@@ -712,6 +704,10 @@ private:
 		{
 			MigrateDataToCurves(NewEditorNumChannels);
 		}
+		else if (NewMode == ENiagaraDistributionMode::ColorGradient)
+		{
+			MigrateDataToColorGradient();
+		}
 	}
 
 	void MigrateDataToConstants(int32 NewEditorNumChannels)
@@ -775,6 +771,7 @@ private:
 			break;
 		case ENiagaraDistributionMode::UniformCurve:
 		case ENiagaraDistributionMode::NonUniformCurve:
+		case ENiagaraDistributionMode::ColorGradient:
 			for (int32 ChannelIndex = 0; ChannelIndex < EditorNumChannels; ChannelIndex++)
 			{
 				const FRichCurve* Curve = GetCurveValue(ChannelIndex);
@@ -832,6 +829,7 @@ private:
 			break;
 		case ENiagaraDistributionMode::UniformCurve:
 		case ENiagaraDistributionMode::NonUniformCurve:
+		case ENiagaraDistributionMode::ColorGradient:
 			for (int32 ChannelIndex = 0; ChannelIndex < EditorNumChannels; ChannelIndex++)
 			{
 				const FRichCurve* Curve = GetCurveValue(ChannelIndex);
@@ -863,10 +861,7 @@ private:
 
 	void MigrateDataToColorGradient()
 	{
-		if ( SourceDistribution->Mode != ENiagaraDistributionMode::NonUniformCurve || AreCurveKeyTimesNormalized() )
-		{
-			return;
-		}
+		MigrateDataToCurves(4);
 
 		for (FRichCurve& ChannelCurve : SourceDistribution->ChannelCurves)
 		{
diff --git a/Engine/Plugins/FX/Niagara/Source/NiagaraEditor/Private/TypeEditorUtilities/NiagaraDistributionPropertyEditorUtilities.cpp b/Engine/Plugins/FX/Niagara/Source/NiagaraEditor/Private/TypeEditorUtilities/NiagaraDistributionPropertyEditorUtilities.cpp
--- a/Engine/Plugins/FX/Niagara/Source/NiagaraEditor/Private/TypeEditorUtilities/NiagaraDistributionPropertyEditorUtilities.cpp
+++ b/Engine/Plugins/FX/Niagara/Source/NiagaraEditor/Private/TypeEditorUtilities/NiagaraDistributionPropertyEditorUtilities.cpp
@@ -64,7 +64,7 @@ bool FNiagaraDistributionPropertyEditorUtilities::TryUpdateClipboardPortableValu
 				}
 				}
 			}
-			else if (Distribution->IsCurve())
+			else if (Distribution->IsCurve() || Distribution->IsGradient())
 			{
 				UScriptStruct* CurveCollectionStruct = FNiagaraClipboardCurveCollection::StaticStruct();
 				FNiagaraClipboardCurveCollection CurveCollectionValue;
