- Sim Cache Reader Updates
- Add method to get number of instances
- Add missing VM function for getting emitter index

#rnx

diff --git a/Engine/Plugins/FX/Niagara/Shaders/Private/NiagaraDataInterfaceSimCacheReaderTemplate.ush b/Engine/Plugins/FX/Niagara/Shaders/Private/NiagaraDataInterfaceSimCacheReaderTemplate.ush
--- a/Engine/Plugins/FX/Niagara/Shaders/Private/NiagaraDataInterfaceSimCacheReaderTemplate.ush
+++ b/Engine/Plugins/FX/Niagara/Shaders/Private/NiagaraDataInterfaceSimCacheReaderTemplate.ush
@@ -47,6 +47,12 @@ void GetEmitterIndex_{ParameterName}(out int EmitterIndex)
 	return FrameData;
 }
 
+void GetNumInstances_{ParameterName}(in int FrameIndex, out int NumInstances)
+{
+	{ParameterName}_EmitterFrameData FrameData = GetEmitterFrameData_{ParameterName}(FrameIndex, {ParameterName}_EmitterIndex);
+	NumInstances = FrameData.NumInstances;
+}
+
 void GetInstanceInt_{ParameterName}(int FrameIndex, int EmitterIndex, int AttributeIndex, int InstanceIndex, out bool bValid, out int Value)
 {
 	bValid = false;
diff --git a/Engine/Plugins/FX/Niagara/Source/Niagara/Private/DataInterface/NiagaraDataInterfaceSimCacheReader.cpp b/Engine/Plugins/FX/Niagara/Source/Niagara/Private/DataInterface/NiagaraDataInterfaceSimCacheReader.cpp
--- a/Engine/Plugins/FX/Niagara/Source/Niagara/Private/DataInterface/NiagaraDataInterfaceSimCacheReader.cpp
+++ b/Engine/Plugins/FX/Niagara/Source/Niagara/Private/DataInterface/NiagaraDataInterfaceSimCacheReader.cpp
@@ -24,6 +24,7 @@ namespace NDISimCacheReaderLocal
 	static const FName GetNumFrames("GetNumFrames");
 	static const FName GetNumEmitters("GetNumEmitters");
 	static const FName GetEmitterIndex("GetEmitterIndex");
+	static const FName GetNumInstances("GetNumInstances");
 	static const FName GetInstanceInt("GetInstanceInt");
 	static const FName GetInstanceFloat("GetInstanceFloat");
 	static const FName GetInstanceVector2("GetInstanceVector2");
@@ -368,6 +369,13 @@ void UNiagaraDataInterfaceSimCacheReader::GetFunctionsInternal(TArray<FNiagaraFu
 		Sig.Outputs.Emplace(FNiagaraTypeDefinition::GetIntDef(), TEXT("EmitterIndex"));
 		Sig.SetDescription(LOCTEXT("GetEmitterIndexDesc", "Returns the emitter index the data interface is bound to"));		
 	}
+	{
+		FNiagaraFunctionSignature& Sig = OutFunctions.Add_GetRef(DefaultSignature);
+		Sig.Name = GetNumInstances;
+		Sig.Inputs.Emplace(FNiagaraTypeDefinition::GetIntDef(), TEXT("FrameIndex"));
+		Sig.Outputs.Emplace(FNiagaraTypeDefinition::GetIntDef(), TEXT("InstanceCount"));
+		Sig.SetDescription(LOCTEXT("GetNumInstancesDesc", "Returns the number of instances for the provided frame."));
+	}
 
 	// Attribute Reading functions
 	FNiagaraFunctionSignature AttributeFunctionSignature = DefaultSignature;
@@ -400,6 +408,14 @@ void UNiagaraDataInterfaceSimCacheReader::GetVMExternalFunction(const FVMExterna
 	{
 		OutFunc = FVMExternalFunction::CreateUObject(this, &UNiagaraDataInterfaceSimCacheReader::VMGetNumEmitters);
 	}
+	else if (BindingInfo.Name == GetEmitterIndex)
+	{
+		OutFunc = FVMExternalFunction::CreateUObject(this, &UNiagaraDataInterfaceSimCacheReader::VMGetEmitterIndex);
+	}
+	else if (BindingInfo.Name == GetNumInstances)
+	{
+		OutFunc = FVMExternalFunction::CreateUObject(this, &UNiagaraDataInterfaceSimCacheReader::VMGetNumInstances);
+	}
 	else if (BindingInfo.Name == GetInstanceInt)
 	{
 		const int32 AttributeIndex = AttributeIndexFromVariable(InstanceSimCache, BindingInfo.Name, BindingInfo.FunctionSpecifiers[0].Value, InstanceData->EmitterIndex);
@@ -557,6 +573,54 @@ void UNiagaraDataInterfaceSimCacheReader::VMGetNumEmitters(FVectorVMExternalFunc
 	}
 }
 
+void UNiagaraDataInterfaceSimCacheReader::VMGetEmitterIndex(FVectorVMExternalFunctionContext& Context)
+{
+	using namespace NDISimCacheReaderLocal;
+
+	VectorVM::FUserPtrHandler<FInstanceData_GameThread> InstanceData_GT(Context);
+	FNDIOutputParam<int32> OutEmitterIndex(Context);
+
+	for (int32 i = 0; i < Context.GetNumInstances(); ++i)
+	{
+		OutEmitterIndex.SetAndAdvance(InstanceData_GT->EmitterIndex);
+	}
+}
+
+void UNiagaraDataInterfaceSimCacheReader::VMGetNumInstances(FVectorVMExternalFunctionContext& Context)
+{
+	using namespace NDISimCacheReaderLocal;
+
+	VectorVM::FUserPtrHandler<FInstanceData_GameThread> InstanceData_GT(Context);
+	FNDIInputParam<int32>	InFrameIndex(Context);
+	FNDIOutputParam<int32>	OutNumInstances(Context);
+
+	UNiagaraSimCache* InstanceSimCache = InstanceData_GT->WeakSimCache.Get();
+	int32 LastFrameIndex = INDEX_NONE;
+	int32 LastNumInstances = 0;
+
+	for (int32 i = 0; i < Context.GetNumInstances(); ++i)
+	{
+		const int32 FrameIndex = InFrameIndex.GetAndAdvance();
+		if (FrameIndex != LastFrameIndex)
+		{
+			LastFrameIndex = FrameIndex;
+
+			if (InstanceSimCache)
+			{
+				FNiagaraSimCacheAttributeReaderHelper AttributeReader;
+				AttributeReader.Initialize(InstanceData_GT->WeakSimCache.Get(), InstanceData_GT->EmitterIndex, 0, FrameIndex);
+				LastNumInstances = AttributeReader.IsValid() ? AttributeReader.GetNumInstances() : 0;
+			}
+			else
+			{
+				LastNumInstances = 0;
+			}
+		}
+
+		OutNumInstances.SetAndAdvance(LastNumInstances);
+	}
+}
+
 void UNiagaraDataInterfaceSimCacheReader::VMGetInstanceInt(FVectorVMExternalFunctionContext& Context, int AttributeIndex)
 {
 	using namespace NDISimCacheReaderLocal;
diff --git a/Engine/Plugins/FX/Niagara/Source/Niagara/Private/DataInterface/NiagaraDataInterfaceSimCacheReader.h b/Engine/Plugins/FX/Niagara/Source/Niagara/Private/DataInterface/NiagaraDataInterfaceSimCacheReader.h
--- a/Engine/Plugins/FX/Niagara/Source/Niagara/Private/DataInterface/NiagaraDataInterfaceSimCacheReader.h
+++ b/Engine/Plugins/FX/Niagara/Source/Niagara/Private/DataInterface/NiagaraDataInterfaceSimCacheReader.h
@@ -74,6 +74,7 @@ private:
 	void VMGetNumFrames(FVectorVMExternalFunctionContext& Context);
 	void VMGetNumEmitters(FVectorVMExternalFunctionContext& Context);
 	void VMGetEmitterIndex(FVectorVMExternalFunctionContext& Context);
+	void VMGetNumInstances(FVectorVMExternalFunctionContext& Context);
 	void VMGetInstanceInt(FVectorVMExternalFunctionContext& Context, int AttributeIndex);
 	void VMGetInstanceFloat(FVectorVMExternalFunctionContext& Context, int AttributeIndex);
 	void VMGetInstanceVector2(FVectorVMExternalFunctionContext& Context, int AttributeIndex);
