- Remove early out as it will cause a mismatch in the renderers

#rnx
#tests Disable all of stateless (CPU & GPU) in a cooked build and spawn a system which is mixed

diff --git a/Engine/Plugins/FX/Niagara/Source/Niagara/Private/Stateless/NiagaraStatelessEmitter.cpp b/Engine/Plugins/FX/Niagara/Source/Niagara/Private/Stateless/NiagaraStatelessEmitter.cpp
--- a/Engine/Plugins/FX/Niagara/Source/Niagara/Private/Stateless/NiagaraStatelessEmitter.cpp
+++ b/Engine/Plugins/FX/Niagara/Source/Niagara/Private/Stateless/NiagaraStatelessEmitter.cpp
@@ -323,7 +323,6 @@ void UNiagaraStatelessEmitter::CacheFromCompiledData()
 	{
 		StatelessEmitterData->bCanEverExecute = false;
 		UE_LOG(LogNiagara, Log, TEXT("Stateless Emitter (%s) can not execute on any available path and will be disabled."), *GetFullName());
-		return;
 	}
 
 	// Build data set
