Update NiagaraSettings to make the new compilation mode (AsyncTasks) the default for NiagaraSystems.

diff --git a/Engine/Plugins/FX/Niagara/Source/Niagara/Public/NiagaraSettings.h b/Engine/Plugins/FX/Niagara/Source/Niagara/Public/NiagaraSettings.h
--- a/Engine/Plugins/FX/Niagara/Source/Niagara/Public/NiagaraSettings.h
+++ b/Engine/Plugins/FX/Niagara/Source/Niagara/Public/NiagaraSettings.h
@@ -111,8 +111,8 @@ enum class ENiagaraStripScriptByteCodeOption : uint8
 UENUM()
 enum class ENiagaraCompilationMode : int32
 {
-	Original = 0	UMETA(DisplayName = "Standard Compilation"),
-	AsyncTasks = 1	UMETA(DisplayName = "Experimental - Async Compilation"),
+	Original = 0	UMETA(DisplayName = "Serial Compilation (Deprecated)"),
+	AsyncTasks = 1	UMETA(DisplayName = "Async System Compilation"),
 	Verify = 2      UMETA(Hidden)
 };
 #endif
@@ -200,7 +200,7 @@ class UNiagaraSettings : public UDeveloperSettings
 
 #if WITH_EDITORONLY_DATA
 	UPROPERTY(config, EditAnywhere, Category = Niagara, meta = (DisplayName = "Option for how to compile Niagara scripts"))
-	ENiagaraCompilationMode CompilationMode = ENiagaraCompilationMode::Original;
+	ENiagaraCompilationMode CompilationMode = ENiagaraCompilationMode::AsyncTasks;
 #endif
 
 	/** The quality levels Niagara uses. */
