Fix for traversal of missing defaults in system script being unable to properly resolve RI parameters
#tests Toggled through enum for the type used in UpdateValueStateChange module that was in the system spawn script

diff --git a/Engine/Plugins/FX/Niagara/Source/NiagaraEditor/Private/NiagaraHlslTranslator.cpp b/Engine/Plugins/FX/Niagara/Source/NiagaraEditor/Private/NiagaraHlslTranslator.cpp
--- a/Engine/Plugins/FX/Niagara/Source/NiagaraEditor/Private/NiagaraHlslTranslator.cpp
+++ b/Engine/Plugins/FX/Niagara/Source/NiagaraEditor/Private/NiagaraHlslTranslator.cpp
@@ -2161,7 +2161,10 @@ FNiagaraTranslateResults TNiagaraHlslTranslator<GraphBridge>::Translate(const FN
 		bool bIsEventSpawn = IsEventSpawnScript();
 		if (IsSpawnScript() || TranslationStages[0].bShouldUpdateInitialAttributeValues)
 		{
+			// be sure to reset the usage as it might have changed during the above Compile and BuildMissingDefaults does a traversal
+			ActiveHistoryForFunctionCalls.BeginUsage(TranslationStages[0].ScriptUsage, IterSource);
 			BuildMissingDefaults();
+			ActiveHistoryForFunctionCalls.EndUsage();
 		}
 
 		if (UNiagaraScript::IsParticleEventScript(TranslationStages[0].ScriptUsage))
