- When setting up previous data use current if previous does not exist
- Current changing to a new position binding which does not have a previous value generated it would use the default position which is 0,0,0.  Using the current value won't produce corrupt motion vectors, but will also not produce correct motion vectors.

diff --git a/Engine/Plugins/FX/Niagara/Source/Niagara/Private/NiagaraRendererMeshes.cpp b/Engine/Plugins/FX/Niagara/Source/Niagara/Private/NiagaraRendererMeshes.cpp
--- a/Engine/Plugins/FX/Niagara/Source/Niagara/Private/NiagaraRendererMeshes.cpp
+++ b/Engine/Plugins/FX/Niagara/Source/Niagara/Private/NiagaraRendererMeshes.cpp
@@ -893,11 +893,11 @@ FNiagaraMeshCommonParameters FNiagaraRendererMeshes::CreateCommonShaderParams(co
 		
 		if (bAccurateMotionVectors)
 		{
-			Params.PrevScaleDataOffset			= VFVariables[ENiagaraMeshVFLayout::PrevScale].GetGPUOffset();
-			Params.PrevRotationDataOffset 		= VFVariables[ENiagaraMeshVFLayout::PrevRotation].GetGPUOffset();
-			Params.PrevPositionDataOffset 		= VFVariables[ENiagaraMeshVFLayout::PrevPosition].GetGPUOffset();
-			Params.PrevVelocityDataOffset 		= VFVariables[ENiagaraMeshVFLayout::PrevVelocity].GetGPUOffset();
-			Params.PrevCameraOffsetDataOffset 	= VFVariables[ENiagaraMeshVFLayout::PrevCameraOffset].GetGPUOffset();
+			Params.PrevScaleDataOffset			= SelectVFComponent(VFVariables, ENiagaraMeshVFLayout::PrevScale, ENiagaraMeshVFLayout::Scale);
+			Params.PrevRotationDataOffset 		= SelectVFComponent(VFVariables, ENiagaraMeshVFLayout::PrevRotation, ENiagaraMeshVFLayout::Rotation);
+			Params.PrevPositionDataOffset 		= SelectVFComponent(VFVariables, ENiagaraMeshVFLayout::PrevPosition, ENiagaraMeshVFLayout::Position);
+			Params.PrevVelocityDataOffset 		= SelectVFComponent(VFVariables, ENiagaraMeshVFLayout::PrevVelocity, ENiagaraMeshVFLayout::Velocity);
+			Params.PrevCameraOffsetDataOffset 	= SelectVFComponent(VFVariables, ENiagaraMeshVFLayout::PrevCameraOffset, ENiagaraMeshVFLayout::CameraOffset);
 		}
 	}
 	else if (SourceMode == ENiagaraRendererSourceDataMode::Emitter) // Clear all these out because we will be using the defaults to specify them
diff --git a/Engine/Plugins/FX/Niagara/Source/Niagara/Private/NiagaraRendererRibbons.cpp b/Engine/Plugins/FX/Niagara/Source/Niagara/Private/NiagaraRendererRibbons.cpp
--- a/Engine/Plugins/FX/Niagara/Source/Niagara/Private/NiagaraRendererRibbons.cpp
+++ b/Engine/Plugins/FX/Niagara/Source/Niagara/Private/NiagaraRendererRibbons.cpp
@@ -2163,13 +2163,13 @@ void FNiagaraRendererRibbons::SetupPerViewUniformBuffer(FNiagaraIndexGenerationI
 
 	TConstArrayView<FNiagaraRendererVariableInfo> VFVariables = RendererLayout->GetVFVariables_RenderThread();
 	PerViewUniformParameters.PositionDataOffset = VFVariables[ENiagaraRibbonVFLayout::Position].GetGPUOffset();
-	PerViewUniformParameters.PrevPositionDataOffset = VFVariables[ENiagaraRibbonVFLayout::PrevPosition].GetGPUOffset();
+	PerViewUniformParameters.PrevPositionDataOffset = SelectVFComponent(VFVariables, ENiagaraRibbonVFLayout::PrevPosition, ENiagaraRibbonVFLayout::Position);
 	PerViewUniformParameters.VelocityDataOffset = VFVariables[ENiagaraRibbonVFLayout::Velocity].GetGPUOffset();
 	PerViewUniformParameters.ColorDataOffset = VFVariables[ENiagaraRibbonVFLayout::Color].GetGPUOffset();
 	PerViewUniformParameters.WidthDataOffset = VFVariables[ENiagaraRibbonVFLayout::Width].GetGPUOffset();
-	PerViewUniformParameters.PrevWidthDataOffset = VFVariables[ENiagaraRibbonVFLayout::PrevRibbonWidth].GetGPUOffset();
+	PerViewUniformParameters.PrevWidthDataOffset = SelectVFComponent(VFVariables, ENiagaraRibbonVFLayout::PrevRibbonWidth, ENiagaraRibbonVFLayout::Width);
 	PerViewUniformParameters.TwistDataOffset = VFVariables[ENiagaraRibbonVFLayout::Twist].GetGPUOffset();
-	PerViewUniformParameters.PrevTwistDataOffset = VFVariables[ENiagaraRibbonVFLayout::PrevRibbonTwist].GetGPUOffset();
+	PerViewUniformParameters.PrevTwistDataOffset = SelectVFComponent(VFVariables, ENiagaraRibbonVFLayout::PrevRibbonTwist, ENiagaraRibbonVFLayout::PrevRibbonTwist);
 	PerViewUniformParameters.NormalizedAgeDataOffset = VFVariables[ENiagaraRibbonVFLayout::NormalizedAge].GetGPUOffset();
 	PerViewUniformParameters.MaterialRandomDataOffset = VFVariables[ENiagaraRibbonVFLayout::MaterialRandom].GetGPUOffset();
 	PerViewUniformParameters.MaterialParamDataOffset = VFVariables[ENiagaraRibbonVFLayout::MaterialParam0].GetGPUOffset();
@@ -2189,7 +2189,7 @@ void FNiagaraRendererRibbons::SetupPerViewUniformBuffer(FNiagaraIndexGenerationI
 
 	bool bShouldDoFacing = FacingMode == ENiagaraRibbonFacingMode::Custom || FacingMode == ENiagaraRibbonFacingMode::CustomSideVector;
 	PerViewUniformParameters.FacingDataOffset = bShouldDoFacing ? VFVariables[ENiagaraRibbonVFLayout::Facing].GetGPUOffset() : -1;
-	PerViewUniformParameters.PrevFacingDataOffset = bShouldDoFacing ? VFVariables[ENiagaraRibbonVFLayout::PrevRibbonFacing].GetGPUOffset() : -1;
+	PerViewUniformParameters.PrevFacingDataOffset = bShouldDoFacing ? SelectVFComponent(VFVariables, ENiagaraRibbonVFLayout::PrevRibbonFacing, ENiagaraRibbonVFLayout::Facing) : -1;
 
 	PerViewUniformParameters.U0DistributionMode = static_cast<int32>(UV0Settings.DistributionMode);
 	PerViewUniformParameters.U1DistributionMode = static_cast<int32>(UV1Settings.DistributionMode);
diff --git a/Engine/Plugins/FX/Niagara/Source/Niagara/Private/NiagaraRendererSprites.cpp b/Engine/Plugins/FX/Niagara/Source/Niagara/Private/NiagaraRendererSprites.cpp
--- a/Engine/Plugins/FX/Niagara/Source/Niagara/Private/NiagaraRendererSprites.cpp
+++ b/Engine/Plugins/FX/Niagara/Source/Niagara/Private/NiagaraRendererSprites.cpp
@@ -656,14 +656,14 @@ FNiagaraSpriteUniformBufferRef FNiagaraRendererSprites::CreateViewUniformBuffer(
 		if (bAccurateMotionVectors)
 		{
 			PerViewUniformParameters.AccurateMotionVectors = true;
-			PerViewUniformParameters.PrevPositionDataOffset = VFVariables[ENiagaraSpriteVFLayout::PrevPosition].GetGPUOffset();
-			PerViewUniformParameters.PrevVelocityDataOffset = VFVariables[ENiagaraSpriteVFLayout::PrevVelocity].GetGPUOffset();
-			PerViewUniformParameters.PrevRotationDataOffset = VFVariables[ENiagaraSpriteVFLayout::PrevRotation].GetGPUOffset();
-			PerViewUniformParameters.PrevSizeDataOffset = VFVariables[ENiagaraSpriteVFLayout::PrevSize].GetGPUOffset();
-			PerViewUniformParameters.PrevFacingDataOffset = VFVariables[ENiagaraSpriteVFLayout::PrevFacing].GetGPUOffset();
-			PerViewUniformParameters.PrevAlignmentDataOffset = VFVariables[ENiagaraSpriteVFLayout::PrevAlignment].GetGPUOffset();
-			PerViewUniformParameters.PrevCameraOffsetDataOffset = VFVariables[ENiagaraSpriteVFLayout::PrevCameraOffset].GetGPUOffset();
-			PerViewUniformParameters.PrevPivotOffsetDataOffset = VFVariables[ENiagaraSpriteVFLayout::PrevPivotOffset].GetGPUOffset();
+			PerViewUniformParameters.PrevPositionDataOffset = SelectVFComponent(VFVariables, ENiagaraSpriteVFLayout::PrevPosition, ENiagaraSpriteVFLayout::Position);
+			PerViewUniformParameters.PrevVelocityDataOffset = SelectVFComponent(VFVariables, ENiagaraSpriteVFLayout::PrevVelocity, ENiagaraSpriteVFLayout::Velocity);
+			PerViewUniformParameters.PrevRotationDataOffset = SelectVFComponent(VFVariables, ENiagaraSpriteVFLayout::PrevRotation, ENiagaraSpriteVFLayout::Rotation);
+			PerViewUniformParameters.PrevSizeDataOffset = SelectVFComponent(VFVariables, ENiagaraSpriteVFLayout::PrevSize, ENiagaraSpriteVFLayout::Size);
+			PerViewUniformParameters.PrevFacingDataOffset = SelectVFComponent(VFVariables, ENiagaraSpriteVFLayout::PrevFacing, ENiagaraSpriteVFLayout::Facing);
+			PerViewUniformParameters.PrevAlignmentDataOffset = SelectVFComponent(VFVariables, ENiagaraSpriteVFLayout::PrevAlignment, ENiagaraSpriteVFLayout::Alignment);
+			PerViewUniformParameters.PrevCameraOffsetDataOffset = SelectVFComponent(VFVariables, ENiagaraSpriteVFLayout::PrevCameraOffset, ENiagaraSpriteVFLayout::CameraOffset);
+			PerViewUniformParameters.PrevPivotOffsetDataOffset = SelectVFComponent(VFVariables, ENiagaraSpriteVFLayout::PrevPivotOffset, ENiagaraSpriteVFLayout::PivotOffset);
 		}
 	}
 	else if (SourceMode == ENiagaraRendererSourceDataMode::Emitter) // Clear all these out because we will be using the defaults to specify them
diff --git a/Engine/Plugins/FX/Niagara/Source/Niagara/Private/NiagaraRibbonRendererProperties.cpp b/Engine/Plugins/FX/Niagara/Source/Niagara/Private/NiagaraRibbonRendererProperties.cpp
--- a/Engine/Plugins/FX/Niagara/Source/Niagara/Private/NiagaraRibbonRendererProperties.cpp
+++ b/Engine/Plugins/FX/Niagara/Source/Niagara/Private/NiagaraRibbonRendererProperties.cpp
@@ -170,6 +170,9 @@ void UNiagaraRibbonRendererProperties::PostLoad()
 	
 	PostLoadBindings(ENiagaraRendererSourceDataMode::Particles);
 
+	// Fix up these bindings from their loaded source bindings
+	SetPreviousBindings(FVersionedNiagaraEmitter());
+
 	if ( Material )
 	{
 		Material->ConditionalPostLoad();
diff --git a/Engine/Plugins/FX/Niagara/Source/Niagara/Public/NiagaraRenderer.h b/Engine/Plugins/FX/Niagara/Source/Niagara/Public/NiagaraRenderer.h
--- a/Engine/Plugins/FX/Niagara/Source/Niagara/Public/NiagaraRenderer.h
+++ b/Engine/Plugins/FX/Niagara/Source/Niagara/Public/NiagaraRenderer.h
@@ -141,6 +141,12 @@ public:
 	static NIAGARA_API FRHIShaderResourceView* GetDummyUInt4Buffer();
 	static NIAGARA_API FRHIShaderResourceView* GetDummyHalfBuffer();
 
+	static inline int SelectVFComponent(const TConstArrayView<FNiagaraRendererVariableInfo>& Variables, int Variable, int Fallback)
+	{
+		const int GPUOffset = Variables[Variable].GetGPUOffset();
+		return GPUOffset == INDEX_NONE ? Variables[Fallback].GetGPUOffset() : GPUOffset;
+	}
+
 	FORCEINLINE ENiagaraSimTarget GetSimTarget() const { return SimTarget; }
 
 	virtual void GetUsedMaterials(TArray<UMaterialInterface*>& UsedMaterials, bool bGetDebugMaterials) { UsedMaterials.Append(BaseMaterials_GT); }
