- Add column sorting to sim cache viewer
- Also clean up how we build the table data to make it less confusing

diff --git a/Engine/Plugins/FX/Niagara/Source/NiagaraEditor/Private/ViewModels/NiagaraSimCacheViewModel.cpp b/Engine/Plugins/FX/Niagara/Source/NiagaraEditor/Private/ViewModels/NiagaraSimCacheViewModel.cpp
--- a/Engine/Plugins/FX/Niagara/Source/NiagaraEditor/Private/ViewModels/NiagaraSimCacheViewModel.cpp
+++ b/Engine/Plugins/FX/Niagara/Source/NiagaraEditor/Private/ViewModels/NiagaraSimCacheViewModel.cpp
@@ -146,6 +146,56 @@ FText FNiagaraSimCacheViewModel::GetComponentText(const FName ComponentName, con
 	return LOCTEXT("Error", "Error");
 }
 
+bool FNiagaraSimCacheViewModel::CompareComponent(int32 ComponentIndex, int32 LhsInstance, int32 RhsInstance, bool bAscending) const
+{
+	bool bResult = bAscending ? LhsInstance < RhsInstance : LhsInstance > RhsInstance;
+
+	TConstArrayView<FNiagaraSimCacheViewModel::FComponentInfo> ComponentInfos = GetSelectedComponentInfos();
+	if (ComponentInfos.IsValidIndex(ComponentIndex) && LhsInstance >= 0 && LhsInstance < NumInstances && RhsInstance >= 0 && RhsInstance < NumInstances)
+	{
+		const FComponentInfo* ComponentInfo = &ComponentInfos[ComponentIndex];
+		if (ComponentInfo->bIsFloat)
+		{
+			const float LhsValue = FloatComponents[(ComponentInfo->ComponentOffset * NumInstances) + LhsInstance];
+			const float RhsValue = FloatComponents[(ComponentInfo->ComponentOffset * NumInstances) + RhsInstance];
+			if ( LhsValue != RhsValue )
+			{
+				bResult = bAscending ? LhsValue < RhsValue : LhsValue > RhsValue;
+			}
+		}
+		else if (ComponentInfo->bIsHalf)
+		{
+			const float LhsValue = HalfComponents[(ComponentInfo->ComponentOffset * NumInstances) + LhsInstance].GetFloat();
+			const float RhsValue = HalfComponents[(ComponentInfo->ComponentOffset * NumInstances) + RhsInstance].GetFloat();
+			if (LhsValue != RhsValue)
+			{
+				bResult = bAscending ? LhsValue < RhsValue : LhsValue > RhsValue;
+			}
+		}
+		else if (ComponentInfo->bIsInt32)
+		{
+			const int32 LhsValue = Int32Components[(ComponentInfo->ComponentOffset * NumInstances) + LhsInstance];
+			const int32 RhsValue = Int32Components[(ComponentInfo->ComponentOffset * NumInstances) + RhsInstance];
+			if (ComponentInfo->bShowAsBool)
+			{
+				const int32 bLhsValue = LhsValue != 0;
+				const int32 bRhsValue = RhsValue != 0;
+				if (bLhsValue != bRhsValue)
+				{
+					bResult = bAscending ? bLhsValue < bRhsValue : bLhsValue > bRhsValue;
+				}
+			}
+			else
+			{
+				if (LhsValue != RhsValue)
+				{
+					bResult = bAscending ? LhsValue < RhsValue : LhsValue > RhsValue;
+				}
+			}
+		}
+	}
+	return bResult;
+}
 
 int32 FNiagaraSimCacheViewModel::GetNumFrames() const
 {
diff --git a/Engine/Plugins/FX/Niagara/Source/NiagaraEditor/Private/Widgets/SNiagaraSimCacheView.cpp b/Engine/Plugins/FX/Niagara/Source/NiagaraEditor/Private/Widgets/SNiagaraSimCacheView.cpp
--- a/Engine/Plugins/FX/Niagara/Source/NiagaraEditor/Private/Widgets/SNiagaraSimCacheView.cpp
+++ b/Engine/Plugins/FX/Niagara/Source/NiagaraEditor/Private/Widgets/SNiagaraSimCacheView.cpp
@@ -75,8 +75,7 @@ void SNiagaraSimCacheView::Construct(const FArguments& InArgs)
 
 	HeaderRowWidget = SNew(SHeaderRow);
 
-	UpdateColumns(true);
-	UpdateRows(false);
+	UpdateListView();
 
 	TSharedRef<SScrollBar> HorizontalScrollBar =
 		SNew(SScrollBar)
@@ -193,48 +192,6 @@ TSharedRef<ITableRow> SNiagaraSimCacheView::MakeRowWidget(const TSharedPtr<int32
 		.SimCacheViewModel(SimCacheViewModel);
 }
 
-void SNiagaraSimCacheView::GenerateColumns()
-{
-	//  Give columns a width to prevent them from being shrunk when filtering. 
-	constexpr float ManualWidth = 125.0f;
-	HeaderRowWidget->ClearColumns();
-
-	if(!SimCacheViewModel->IsCacheValid())
-	{
-		return;
-	}
-	
-	// Generate instance count column
-	HeaderRowWidget->AddColumn(
-		SHeaderRow::Column(NAME_Instance)
-		.DefaultLabel(FText::FromName(NAME_Instance))
-		.HAlignHeader(EHorizontalAlignment::HAlign_Center)
-		.VAlignHeader(EVerticalAlignment::VAlign_Fill)
-		.HAlignCell(EHorizontalAlignment::HAlign_Center)
-		.VAlignCell(EVerticalAlignment::VAlign_Fill)
-		.ManualWidth(ManualWidth)
-		.SortMode(EColumnSortMode::None)
-	);
-	
-	// Generate a column for each component
-	for (const FNiagaraSimCacheViewModel::FComponentInfo& ComponentInfo : SimCacheViewModel->GetSelectedComponentInfos())
-	{
-		HeaderRowWidget->AddColumn(
-			SHeaderRow::Column(ComponentInfo.Name)
-			.DefaultLabel(FText::FromName(ComponentInfo.Name))
-			.HAlignHeader(EHorizontalAlignment::HAlign_Center)
-			.VAlignHeader(EVerticalAlignment::VAlign_Fill)
-			.HAlignCell(EHorizontalAlignment::HAlign_Center)
-			.VAlignCell(EVerticalAlignment::VAlign_Fill)
-			.FillWidth(1.0f)
-			.ManualWidth(ManualWidth)
-			.ShouldGenerateWidget(TAttribute<bool>::Create(TAttribute<bool>::FGetter::CreateSP(this, &SNiagaraSimCacheView::GetShouldGenerateWidget, ComponentInfo.Name)))
-			.SortMode(EColumnSortMode::None)
-		);
-	}
-
-}
-
 void SNiagaraSimCacheView::UpdateCustomDisplayWidget()
 {
 	for (TSharedPtr<SWidget> Widget : CustomDisplayWidgets)
@@ -319,51 +276,157 @@ void SNiagaraSimCacheView::UpdateCustomDisplayWidget()
 	);
 }
 
-void SNiagaraSimCacheView::UpdateColumns(const bool bReset)
+void SNiagaraSimCacheView::UpdateListView()
 {
-	if(bReset)
+	GenerateColumns();
+	GenerateRows();
+	SortRows();
+
+	if (ListViewWidget)
 	{
-		GenerateColumns();
+		ListViewWidget->RequestListRefresh();
 	}
-	
+}
+
+void SNiagaraSimCacheView::GenerateColumns()
+{
+	// Invalid early out
+	if (SimCacheViewModel->IsCacheValid() == false)
+	{
+		HeaderRowWidget->ClearColumns();
+		HeaderRowWidget->RefreshColumns();
+		return;
+	}
+
+	// Do we need to update our columns?
+	TConstArrayView<FNiagaraSimCacheViewModel::FComponentInfo> SelectedComponents = SimCacheViewModel->GetSelectedComponentInfos();
+	if (HeaderRowWidget->GetColumns().Num() == SelectedComponents.Num() + 1)
+	{
+		bool bNeedsRefresh = false;
+		for (int32 i = 0; i < SelectedComponents.Num(); ++i)
+		{
+			if (HeaderRowWidget->GetColumns()[i + 1].ColumnId != SelectedComponents[i].Name)
+			{
+				bNeedsRefresh = true;
+				break;
+			}
+		}
+
+		if (bNeedsRefresh == false)
+		{
+			return;
+		}
+	}
+
+	//  Give columns a width to prevent them from being shrunk when filtering. 
+	constexpr float ManualWidth = 125.0f;
+	HeaderRowWidget->ClearColumns();
+
+	if (SimCacheViewModel->IsCacheValid())
+	{
+		// Generate instance count column
+		HeaderRowWidget->AddColumn(
+			SHeaderRow::Column(NAME_Instance)
+			.DefaultLabel(FText::FromName(NAME_Instance))
+			.HAlignHeader(EHorizontalAlignment::HAlign_Center)
+			.VAlignHeader(EVerticalAlignment::VAlign_Fill)
+			.HAlignCell(EHorizontalAlignment::HAlign_Center)
+			.VAlignCell(EVerticalAlignment::VAlign_Fill)
+			.ManualWidth(ManualWidth)
+			.SortMode(this, &SNiagaraSimCacheView::GetColumnSortMode, NAME_Instance)
+			.OnSort(this, &SNiagaraSimCacheView::OnColumnNameSortModeChanged)
+		);
+
+		// Generate a column for each component
+		for (const FNiagaraSimCacheViewModel::FComponentInfo& ComponentInfo : SelectedComponents)
+		{
+			HeaderRowWidget->AddColumn(
+				SHeaderRow::Column(ComponentInfo.Name)
+				.DefaultLabel(FText::FromName(ComponentInfo.Name))
+				.HAlignHeader(EHorizontalAlignment::HAlign_Center)
+				.VAlignHeader(EVerticalAlignment::VAlign_Fill)
+				.HAlignCell(EHorizontalAlignment::HAlign_Center)
+				.VAlignCell(EVerticalAlignment::VAlign_Fill)
+				.FillWidth(1.0f)
+				.ManualWidth(ManualWidth)
+				.ShouldGenerateWidget(TAttribute<bool>::Create(TAttribute<bool>::FGetter::CreateSP(this, &SNiagaraSimCacheView::GetShouldGenerateWidget, ComponentInfo.Name)))
+				.SortMode(this, &SNiagaraSimCacheView::GetColumnSortMode, ComponentInfo.Name)
+				.OnSort(this, &SNiagaraSimCacheView::OnColumnNameSortModeChanged)
+			);
+		}
+	}
+
 	HeaderRowWidget->RefreshColumns();
-	
 }
 
-void SNiagaraSimCacheView::UpdateRows(const bool bRefresh)
+void SNiagaraSimCacheView::GenerateRows()
 {
+	if (SimCacheViewModel->IsCacheValid() == false )
+	{
+		RowItems.Empty();
+		return;
+	}
+	
 	RowItems.Reset(SimCacheViewModel->GetNumInstances());
 	for (int32 i = 0; i < SimCacheViewModel->GetNumInstances(); ++i)
 	{
 		RowItems.Emplace(MakeShared<int32>(i));
 	}
+}
+
+void SNiagaraSimCacheView::SortRows()
+{
+	if (RowItems.Num() == 0)
+	{
+		return;
+	}
 
-	if (bRefresh && ListViewWidget)
+	// Is the column name valid?
+	const int32 ColumnIndex = SimCacheViewModel->GetSelectedComponentInfos().IndexOfByPredicate([ColumnName=SortColumnName](const FNiagaraSimCacheViewModel::FComponentInfo& ComponentInfo) { return ComponentInfo.Name == ColumnName; });
+	if (ColumnIndex == INDEX_NONE)
 	{
-		ListViewWidget->RequestListRefresh();
+		SortColumnName = NAME_Instance;
+	}
+
+	// Simple sort
+	if (SortColumnName == NAME_Instance)
+	{
+		if ( SortMode == EColumnSortMode::Ascending )
+		{
+			RowItems.Sort([](const TSharedPtr<int32>& Lhs, const TSharedPtr<int32>& Rhs) { return *Lhs <= *Rhs; });
+		}
+		else
+		{
+			RowItems.Sort([](const TSharedPtr<int32>& Lhs, const TSharedPtr<int32>& Rhs) { return *Lhs  > *Rhs; });
+		}
+	}
+	// Complex sort
+	else
+	{
+		const bool bAscending = SortMode == EColumnSortMode::Ascending;
+		RowItems.Sort(
+			[ViewModel=SimCacheViewModel.Get(), bAscending, ColumnIndex](const TSharedPtr<int32>& Lhs, const TSharedPtr<int32>& Rhs)
+			{
+				return ViewModel->CompareComponent(ColumnIndex, *Lhs, *Rhs, bAscending);
+			}
+		);
 	}
 }
 
 void SNiagaraSimCacheView::OnSimCacheChanged()
 {
-	UpdateRows(true);
-	UpdateColumns(true);
+	UpdateListView();
 	UpdateCustomDisplayWidget();
 }
 
 void SNiagaraSimCacheView::OnViewDataChanged(const bool bFullRefresh)
 {
-	UpdateRows(true);
-	if (bFullRefresh)
-	{
-		UpdateColumns(false);
-	}
+	UpdateListView();
 }
 
 void SNiagaraSimCacheView::OnBufferChanged()
 {
-	UpdateRows(true);
-	UpdateColumns(true);
+	UpdateListView();
 	UpdateCustomDisplayWidget();
 }
 
@@ -377,4 +440,21 @@ bool SNiagaraSimCacheView::GetShouldGenerateWidget(FName Name)
 	return SimCacheViewModel->IsComponentFiltered(Name);
 }
 
+EColumnSortMode::Type SNiagaraSimCacheView::GetColumnSortMode(FName ColumnName) const
+{
+	return ColumnName == SortColumnName ? SortMode : EColumnSortMode::None;
+}
+
+void SNiagaraSimCacheView::OnColumnNameSortModeChanged(const EColumnSortPriority::Type SortPriority, const FName& ColumnId, const EColumnSortMode::Type InSortMode)
+{
+	SortMode		= InSortMode;
+	SortColumnName	= ColumnId;
+
+	SortRows();
+	if (ListViewWidget)
+	{
+		ListViewWidget->RequestListRefresh();
+	}
+}
+
 #undef LOCTEXT_NAMESPACE
diff --git a/Engine/Plugins/FX/Niagara/Source/NiagaraEditor/Private/Widgets/SNiagaraSimCacheView.h b/Engine/Plugins/FX/Niagara/Source/NiagaraEditor/Private/Widgets/SNiagaraSimCacheView.h
--- a/Engine/Plugins/FX/Niagara/Source/NiagaraEditor/Private/Widgets/SNiagaraSimCacheView.h
+++ b/Engine/Plugins/FX/Niagara/Source/NiagaraEditor/Private/Widgets/SNiagaraSimCacheView.h
@@ -25,22 +25,24 @@ public:
 
 	TSharedRef<ITableRow> MakeRowWidget(const TSharedPtr<int32> RowIndexPtr, const TSharedRef<STableViewBase>& OwnerTable) const;
 
-	// Update existing or generate new columns and apply any filters.
-	// @param bReset Weather the columns need a full reset. If true this will destroy existing columns and rebuild new ones. 
-	void UpdateColumns(const bool bReset);
+private:
+	void UpdateListView();
+	void GenerateColumns();
+	void GenerateRows();
+	void SortRows();
 
-	void UpdateRows(const bool bRefresh);
-	
 	void OnSimCacheChanged();
 
 	void OnViewDataChanged(const bool bFullRefresh);
 	void OnBufferChanged();
 
-private:
 	bool GetShouldGenerateWidget(FName Name);
-	void GenerateColumns();
 	void UpdateCustomDisplayWidget();
 
+	EColumnSortMode::Type GetColumnSortMode(FName ColumnName) const;
+	void OnColumnNameSortModeChanged(const EColumnSortPriority::Type SortPriority, const FName& ColumnId, const EColumnSortMode::Type InSortMode);
+
+private:
 	TArray<TSharedPtr<int32>>					RowItems;
 	TSharedPtr<FNiagaraSimCacheViewModel>		SimCacheViewModel;
 	
@@ -49,4 +51,7 @@ private:
 	TSharedPtr<SWidgetSwitcher>					SwitchWidget;
 	TArray<TSharedPtr<SWidget>>					CustomDisplayWidgets;
 	TSharedPtr<SScrollBar>						CustomDisplayScrollBar;
-};
\ No newline at end of file
+
+	EColumnSortMode::Type						SortMode = EColumnSortMode::Ascending;
+	FName										SortColumnName;
+};
diff --git a/Engine/Plugins/FX/Niagara/Source/NiagaraEditor/Public/ViewModels/NiagaraSimCacheViewModel.h b/Engine/Plugins/FX/Niagara/Source/NiagaraEditor/Public/ViewModels/NiagaraSimCacheViewModel.h
--- a/Engine/Plugins/FX/Niagara/Source/NiagaraEditor/Public/ViewModels/NiagaraSimCacheViewModel.h
+++ b/Engine/Plugins/FX/Niagara/Source/NiagaraEditor/Public/ViewModels/NiagaraSimCacheViewModel.h
@@ -85,6 +85,7 @@ private:
 
 public:
 	FText GetComponentText(FName ComponentName, int32 InstanceIndex) const;
+	bool CompareComponent(int32 ComponentIndex, int32 LhsInstance, int32 RhsInstance, bool bAscending) const;
 	bool IsCacheValid() const;
 	int32 GetNumEmitterLayouts() const;
 	FName GetEmitterLayoutName(int32 Index) const;
