- Add information about who is using sorting
- Adds a marker to sort key generation
- Adds console variable fx.Niagara.Batcher.DumpSortedSimulations which can be used to dump to the log

diff --git a/Engine/Plugins/FX/Niagara/Source/Niagara/Classes/NiagaraGPUSortInfo.h b/Engine/Plugins/FX/Niagara/Source/Niagara/Classes/NiagaraGPUSortInfo.h
--- a/Engine/Plugins/FX/Niagara/Source/Niagara/Classes/NiagaraGPUSortInfo.h
+++ b/Engine/Plugins/FX/Niagara/Source/Niagara/Classes/NiagaraGPUSortInfo.h
@@ -37,6 +37,10 @@ struct FNiagaraGPUSortInfo
 {
 	static constexpr uint32 MaxCullPlanes = 12;
 
+#if STATS
+	// Owner name, used to find who instigated sorting
+	FName OwnerName;
+#endif
 	// The number of particles in the system.
 	int32 ParticleCount = 0;
 	// How the particles should be sorted.
diff --git a/Engine/Plugins/FX/Niagara/Source/Niagara/Private/NiagaraGpuComputeDispatch.cpp b/Engine/Plugins/FX/Niagara/Source/Niagara/Private/NiagaraGpuComputeDispatch.cpp
--- a/Engine/Plugins/FX/Niagara/Source/Niagara/Private/NiagaraGpuComputeDispatch.cpp
+++ b/Engine/Plugins/FX/Niagara/Source/Niagara/Private/NiagaraGpuComputeDispatch.cpp
@@ -141,6 +141,16 @@ namespace FNiagaraGpuComputeDispatchLocal
 		);
 	#endif
 
+	#if STATS
+		static int32 GDumpSortedSimulations = 0;
+		static FAutoConsoleVariableRef CVarDumpSortedSimulations(
+			TEXT("fx.Niagara.Batcher.DumpSortedSimulations"),
+			GDumpSortedSimulations,
+			TEXT("Will dump sorted simulations to the log for the provided number of frames."),
+			ECVF_Default
+		);
+	#endif
+
 	template<typename TTransitionArrayType>
 	static void AddDataBufferTransitions(TTransitionArrayType& BeforeTransitionArray, TTransitionArrayType& AfterTransitionArray, FNiagaraDataBuffer* DestinationData, ERHIAccess BeforeState = ERHIAccess::SRVMask, ERHIAccess AfterState = ERHIAccess::UAVCompute)
 	{
@@ -1966,6 +1976,40 @@ void FNiagaraGpuComputeDispatch::PreInitViews(FRDGBuilder& GraphBuilder, bool bA
 	GPUProfilerPtr->BeginFrame(GraphBuilder.RHICmdList);
 #endif
 
+#if STATS
+	if (FNiagaraGpuComputeDispatchLocal::GDumpSortedSimulations > 0)
+	{
+		--FNiagaraGpuComputeDispatchLocal::GDumpSortedSimulations;
+
+		UE_LOG(LogNiagara, Log, TEXT("Niagara SimulationsToSort:"));
+		for (const FNiagaraGPUSortInfo& SortInfo : SimulationsToSort)
+		{
+			TStringBuilder<256> EnabledFeatures;
+			if (SortInfo.SortMode != ENiagaraSortMode::None && SortInfo.SortAttributeOffset != INDEX_NONE)
+			{
+				EnabledFeatures.Append(TEXT("SortingEnabled "));
+			}
+			if (SortInfo.bEnableCulling)
+			{
+				EnabledFeatures.Append(TEXT("CullingEnabled "));
+				if (SortInfo.CullPlanes.Num() > 0)
+				{
+					EnabledFeatures.Append(TEXT("FrustumCulling "));
+				}
+				if (SortInfo.RendererVisTagAttributeOffset != INDEX_NONE)
+				{
+					EnabledFeatures.Append(TEXT("RendererVisEnabled "));
+				}
+				if (SortInfo.MeshIndexAttributeOffset != INDEX_NONE)
+				{
+					EnabledFeatures.Append(TEXT("MeshIndexEnabled "));
+				}
+			}
+			UE_LOG(LogNiagara, Log, TEXT(" - Owner(%s) ParticleCount(%d) %s"), *SortInfo.OwnerName.ToString(), SortInfo.ParticleCount, EnabledFeatures.ToString());
+		}
+	}
+#endif
+
 	// Reset the list of GPUSort tasks and release any resources they hold on to.
 	// It might be worth considering doing so at the end of the render to free the resources immediately.
 	// (note that currently there are no callback appropriate to do it)
@@ -2362,6 +2406,10 @@ void FNiagaraGpuComputeDispatch::GenerateSortKeys(FRHICommandListImmediate& RHIC
 	{
 		if (SortInfo.AllocationInfo.SortBatchId == BatchId)
 		{
+		#if STATS
+			SCOPED_DRAW_EVENTF(RHICmdList, GenerateSortKeys, TEXT("Owner:%s"), SortInfo.OwnerName.ToString());
+		#endif
+
 			Params.NiagaraParticleDataFloat = SortInfo.ParticleDataFloatSRV;
 			Params.NiagaraParticleDataHalf = SortInfo.ParticleDataHalfSRV;
 			Params.NiagaraParticleDataInt = SortInfo.ParticleDataIntSRV;
diff --git a/Engine/Plugins/FX/Niagara/Source/Niagara/Private/NiagaraRendererMeshes.cpp b/Engine/Plugins/FX/Niagara/Source/Niagara/Private/NiagaraRendererMeshes.cpp
--- a/Engine/Plugins/FX/Niagara/Source/Niagara/Private/NiagaraRendererMeshes.cpp
+++ b/Engine/Plugins/FX/Niagara/Source/Niagara/Private/NiagaraRendererMeshes.cpp
@@ -599,6 +599,9 @@ void FNiagaraRendererMeshes::InitializeSortInfo(const FParticleMeshRenderData& P
 {
 	TConstArrayView<FNiagaraRendererVariableInfo> VFVariables = ParticleMeshRenderData.RendererLayout->GetVFVariables_RenderThread();
 
+#if STATS
+	OutSortInfo.OwnerName = EmitterStatID.GetName();
+#endif
 	OutSortInfo.ParticleCount = ParticleMeshRenderData.SourceParticleData->GetNumInstances();
 	OutSortInfo.SortMode = SortMode;
 	OutSortInfo.SetSortFlags(bSortHighPrecision, ParticleMeshRenderData.SourceParticleData->GetGPUDataReadyStage());
diff --git a/Engine/Plugins/FX/Niagara/Source/Niagara/Private/NiagaraRendererSprites.cpp b/Engine/Plugins/FX/Niagara/Source/Niagara/Private/NiagaraRendererSprites.cpp
--- a/Engine/Plugins/FX/Niagara/Source/Niagara/Private/NiagaraRendererSprites.cpp
+++ b/Engine/Plugins/FX/Niagara/Source/Niagara/Private/NiagaraRendererSprites.cpp
@@ -413,6 +413,9 @@ void FNiagaraRendererSprites::InitializeSortInfo(FParticleSpriteRenderData& Part
 {
 	TConstArrayView<FNiagaraRendererVariableInfo> VFVariables = ParticleSpriteRenderData.RendererLayout->GetVFVariables_RenderThread();
 
+#if STATS
+	OutSortInfo.OwnerName = EmitterStatID.GetName();
+#endif
 	OutSortInfo.ParticleCount = ParticleSpriteRenderData.SourceParticleData->GetNumInstances();
 	OutSortInfo.SortMode = SortMode;
 	OutSortInfo.SetSortFlags(bSortHighPrecision, ParticleSpriteRenderData.SourceParticleData->GetGPUDataReadyStage());
