Updated code using script to make sure x_api is set on functions/globalvars instead of on types

diff --git a/Engine/Plugins/FX/Niagara/Source/Niagara/Classes/NiagaraDataInterface.h b/Engine/Plugins/FX/Niagara/Source/Niagara/Classes/NiagaraDataInterface.h
--- a/Engine/Plugins/FX/Niagara/Source/Niagara/Classes/NiagaraDataInterface.h
+++ b/Engine/Plugins/FX/Niagara/Source/Niagara/Classes/NiagaraDataInterface.h
@@ -903,9 +903,9 @@ public:
 	UE_DEPRECATED(5.3, "Please update your code to use IsUsedWithGPUScript")
 	bool IsUsedByGPUEmitter() const { return bUsedWithGPUScript; }
 	UE_DEPRECATED(5.3, "Please update your code to use IsUsedWithCPUScript")
-	NIAGARA_API bool IsUsedWithCPUEmitter() const { return bUsedWithCPUScript; }
+	bool IsUsedWithCPUEmitter() const { return bUsedWithCPUScript; }
 	UE_DEPRECATED(5.3, "Please update your code to use IsUsedWithGPUScript")
-	NIAGARA_API bool IsUsedWithGPUEmitter() const { return bUsedWithGPUScript; }
+	bool IsUsedWithGPUEmitter() const { return bUsedWithGPUScript; }
 
 	/** Used to determine if we need to create CPU resources for the emitter. */
 	bool IsUsedWithCPUScript() const { return bUsedWithCPUScript; }
diff --git a/Engine/Plugins/FX/Niagara/Source/Niagara/Classes/NiagaraDataSet.h b/Engine/Plugins/FX/Niagara/Source/Niagara/Classes/NiagaraDataSet.h
--- a/Engine/Plugins/FX/Niagara/Source/Niagara/Classes/NiagaraDataSet.h
+++ b/Engine/Plugins/FX/Niagara/Source/Niagara/Classes/NiagaraDataSet.h
@@ -195,7 +195,7 @@ public:
 
 
 private:
-	NIAGARA_API FORCEINLINE void CheckUsage(bool bReadOnly)const;
+	FORCEINLINE void CheckUsage(bool bReadOnly)const;
 
 	FORCEINLINE int32 GetSafeComponentBufferSize(int32 RequiredSize) const
 	{
diff --git a/Engine/Plugins/FX/Niagara/Source/Niagara/Public/NiagaraComponent.h b/Engine/Plugins/FX/Niagara/Source/Niagara/Public/NiagaraComponent.h
--- a/Engine/Plugins/FX/Niagara/Source/Niagara/Public/NiagaraComponent.h
+++ b/Engine/Plugins/FX/Niagara/Source/Niagara/Public/NiagaraComponent.h
@@ -7,12 +7,6 @@
 #include "NiagaraDefines.h"
 #include "NiagaraScalabilityState.h"
 #include "NiagaraTickBehaviorEnum.h"
-#if UE_ENABLE_INCLUDE_ORDER_DEPRECATED_IN_5_2
-#include "NiagaraCommon.h"
-#include "NiagaraComponentPool.h"
-#include "NiagaraSystemInstanceController.h"
-#include "PrimitiveSceneProxy.h"
-#endif
 #include "NiagaraUserRedirectionParameterStore.h"
 #include "NiagaraVariant.h"
 #include "PrimitiveViewRelevance.h"
@@ -290,19 +284,19 @@ public:
 	NIAGARA_API void SetAsset(UNiagaraSystem* InAsset, bool bResetExistingOverrideParameters = true);
 
 	UFUNCTION(BlueprintCallable, Category = Niagara, meta = (DisplayName = "Get Niagara System Asset"))
-	NIAGARA_API UNiagaraSystem* GetAsset() const { return Asset; }
+	UNiagaraSystem* GetAsset() const { return Asset; }
 
 	UFUNCTION(BlueprintCallable, Category = Niagara)
 	NIAGARA_API void SetOcclusionQueryMode(ENiagaraOcclusionQueryMode Mode);
 
 	UFUNCTION(BlueprintCallable, Category = Niagara)
-	NIAGARA_API ENiagaraOcclusionQueryMode GetOcclusionQueryMode() const { return OcclusionQueryMode; }
+	ENiagaraOcclusionQueryMode GetOcclusionQueryMode() const { return OcclusionQueryMode; }
 
 	UFUNCTION(BlueprintCallable, Category = Niagara, meta = (DisplayName = "Set Forced Solo Mode"))
 	NIAGARA_API void SetForceSolo(bool bInForceSolo);
 
 	UFUNCTION(BlueprintCallable, Category = Niagara, meta = (DisplayName = "Is In Forced Solo Mode"))
-	NIAGARA_API bool GetForceSolo() const { return bForceSolo; }
+	bool GetForceSolo() const { return bForceSolo; }
 
 private:
 	NIAGARA_API bool RequiresSoloMode() const;
@@ -320,7 +314,7 @@ public:
 	NIAGARA_API void SetCustomTimeDilation(float Dilation = 1.0f);
 
 	UFUNCTION(BlueprintCallable, Category = Niagara)
-	NIAGARA_API float GetCustomTimeDilation() const { return CustomTimeDilation; }
+	float GetCustomTimeDilation() const { return CustomTimeDilation; }
 
 	UFUNCTION(BlueprintCallable, Category = Niagara, meta = (DisplayName = "Get Age Update Mode"))
 	NIAGARA_API ENiagaraAgeUpdateMode GetAgeUpdateMode() const;
@@ -640,10 +634,10 @@ public:
 	NIAGARA_API void SetPreviewLODDistance(bool bEnablePreviewLODDistance, float PreviewLODDistance, float PreviewMaxDistance);
 
 	UFUNCTION(BlueprintCallable, Category = Preview, meta = (Keywords = "preview LOD Distance scalability"))
-	NIAGARA_API FORCEINLINE bool GetPreviewLODDistanceEnabled()const;
+	FORCEINLINE bool GetPreviewLODDistanceEnabled()const;
 
 	UFUNCTION(BlueprintCallable, Category = Preview, meta = (Keywords = "preview LOD Distance scalability"))
-	NIAGARA_API FORCEINLINE float GetPreviewLODDistance()const;
+	FORCEINLINE float GetPreviewLODDistance()const;
 
 	/**
 	Initializes this component for capturing a performance baseline.
diff --git a/Engine/Plugins/FX/Niagara/Source/Niagara/Public/NiagaraSystemInstance.h b/Engine/Plugins/FX/Niagara/Source/Niagara/Public/NiagaraSystemInstance.h
--- a/Engine/Plugins/FX/Niagara/Source/Niagara/Public/NiagaraSystemInstance.h
+++ b/Engine/Plugins/FX/Niagara/Source/Niagara/Public/NiagaraSystemInstance.h
@@ -270,8 +270,8 @@ public:
 	FORCEINLINE USceneComponent* GetAttachComponent() { return AttachComponent.Get(); }
 	FORCEINLINE FNiagaraUserRedirectionParameterStore* GetOverrideParameters() { return OverrideParameters; }
 	FORCEINLINE const FNiagaraUserRedirectionParameterStore* GetOverrideParameters() const { return OverrideParameters; }
-	[[nodiscard]] NIAGARA_API TArrayView<FNiagaraEmitterInstanceRef> GetEmitters() { return Emitters; }
-	[[nodiscard]] NIAGARA_API TConstArrayView<FNiagaraEmitterInstanceRef> GetEmitters() const { return Emitters; }
+	[[nodiscard]] TArrayView<FNiagaraEmitterInstanceRef> GetEmitters() { return Emitters; }
+	[[nodiscard]] TConstArrayView<FNiagaraEmitterInstanceRef> GetEmitters() const { return Emitters; }
 
 	UE_DEPRECATED(5.5, "This function is experimental and subject to change.")
 	static const INiagaraSceneComponentUtils& GetSceneComponentUtils(FNiagaraSystemInstance* Instance);
@@ -433,8 +433,8 @@ public:
 	NIAGARA_API FNiagaraDataSet* GetEventDataSet(FName EmitterName, FName EventName) const;
 	NIAGARA_API void ClearEventDataSets();
 
-	NIAGARA_API FORCEINLINE void SetLODDistance(float InLODDistance, float InMaxLODDistance, bool bOverride);
-	NIAGARA_API FORCEINLINE void ClearLODDistance();
+	FORCEINLINE void SetLODDistance(float InLODDistance, float InMaxLODDistance, bool bOverride);
+	FORCEINLINE void ClearLODDistance();
 
 	NIAGARA_API const FString& GetCrashReporterTag()const;
 
diff --git a/Engine/Source/Runtime/Engine/Classes/Components/BoxComponent.h b/Engine/Source/Runtime/Engine/Classes/Components/BoxComponent.h
--- a/Engine/Source/Runtime/Engine/Classes/Components/BoxComponent.h
+++ b/Engine/Source/Runtime/Engine/Classes/Components/BoxComponent.h
@@ -40,11 +40,11 @@ public:
 
 	// @return the box extent, scaled by the component scale.
 	UFUNCTION(BlueprintCallable, Category="Components|Box")
-	ENGINE_API FVector GetScaledBoxExtent() const;
+	FVector GetScaledBoxExtent() const;
 
 	// @return the box extent, ignoring component scale.
 	UFUNCTION(BlueprintCallable, Category="Components|Box")
-	ENGINE_API FVector GetUnscaledBoxExtent() const;
+	FVector GetUnscaledBoxExtent() const;
 
 	//~ Begin UPrimitiveComponent Interface.
 	ENGINE_API virtual bool IsZeroExtent() const override;
@@ -64,7 +64,7 @@ public:
 	FORCEINLINE void InitBoxExtent(const FVector& InBoxExtent) { BoxExtent = InBoxExtent; }
 
 #if WITH_EDITOR
-	ENGINE_API FEngineShowFlags GetShowFlags() const { return ShowFlags; }
+	FEngineShowFlags GetShowFlags() const { return ShowFlags; }
 	ENGINE_API void SetShowFlags(const FEngineShowFlags& InShowFlags);
 #endif // WITH_EDITOR
 };
