- Fix crash when copying script variable meta data

diff --git a/Engine/Plugins/FX/Niagara/Source/NiagaraEditor/Private/ViewModels/NiagaraParameterPanelViewModel.cpp b/Engine/Plugins/FX/Niagara/Source/NiagaraEditor/Private/ViewModels/NiagaraParameterPanelViewModel.cpp
--- a/Engine/Plugins/FX/Niagara/Source/NiagaraEditor/Private/ViewModels/NiagaraParameterPanelViewModel.cpp
+++ b/Engine/Plugins/FX/Niagara/Source/NiagaraEditor/Private/ViewModels/NiagaraParameterPanelViewModel.cpp
@@ -217,11 +217,11 @@ bool INiagaraImmutableParameterPanelViewModel::GetCanCopyParameterReferenceAndTo
 
 void INiagaraImmutableParameterPanelViewModel::CopyParameterMetaData(const FNiagaraParameterPanelItemBase ItemToCopy) const
 {
-	if (auto It = GetEditableScriptVariablesWithName(ItemToCopy.GetVariable().GetName()).CreateConstIterator(); It)
+	const TArray<UNiagaraScriptVariable*> ScriptVariables = GetEditableScriptVariablesWithName(ItemToCopy.GetVariable().GetName());
+	if (ScriptVariables.Num() > 0 && ScriptVariables[0] != nullptr)
 	{
-		const UNiagaraScriptVariable* ScriptVariable = *It;
 		UNiagaraClipboardContent* ClipboardContent = UNiagaraClipboardContent::Create();
-		ClipboardContent->ScriptVariables.Add({*ScriptVariable});
+		ClipboardContent->ScriptVariables.Add({*ScriptVariables[0]});
 		FNiagaraEditorModule::Get().GetClipboard().SetClipboardContent(ClipboardContent);
 	}
 }
