Niagara - Fix a bug in traversal cache invalidation.
+ Updated the traversal cache to use a serial number to invalidate data collected from traversals when a script is applied, or when undo/redo happens.  The traversals themselves are still invalidated when their change id updates.
+ Properly modify graphs when sychronizing map pins with function call changes.
+ Update the rename tests to trigger events that happen when editing in the UI and additional testing to the rename undo/redo test to catch the graph changes since the current implementation only validates rapid iteration parameter changes.

diff --git a/Engine/Plugins/FX/Niagara/Source/NiagaraEditor/Private/NiagaraEditorModule.cpp b/Engine/Plugins/FX/Niagara/Source/NiagaraEditor/Private/NiagaraEditorModule.cpp
--- a/Engine/Plugins/FX/Niagara/Source/NiagaraEditor/Private/NiagaraEditorModule.cpp
+++ b/Engine/Plugins/FX/Niagara/Source/NiagaraEditor/Private/NiagaraEditorModule.cpp
@@ -1753,6 +1753,9 @@ void FNiagaraEditorModule::ShutdownModule()
 	GraphDataCache.Reset();
 
 	FAssetCompilingManager::Get().UnregisterManager(&FNiagaraSystemCompilingManager::Get());
+
+	TraversalCache.Reset();
+	TraversalBuilder.Reset();
 }
 
 void FNiagaraEditorModule::OnPostEngineInit()
diff --git a/Engine/Plugins/FX/Niagara/Source/NiagaraEditor/Private/TraversalCache/TraversalBuilder.cpp b/Engine/Plugins/FX/Niagara/Source/NiagaraEditor/Private/TraversalCache/TraversalBuilder.cpp
--- a/Engine/Plugins/FX/Niagara/Source/NiagaraEditor/Private/TraversalCache/TraversalBuilder.cpp
+++ b/Engine/Plugins/FX/Niagara/Source/NiagaraEditor/Private/TraversalCache/TraversalBuilder.cpp
@@ -601,6 +601,12 @@ void FTraversalBuilder::Initialize()
 	}
 }
 
+FTraversalBuilder::~FTraversalBuilder()
+{
+	Instance = nullptr;
+	bInitialized = false;
+}
+
 TSharedRef<const FTraversal> FTraversalBuilder::BuildTraversal(
 	const UNiagaraNodeOutput& OutputNode, FTraversalBuilderDebugData* DebugData)
 {
diff --git a/Engine/Plugins/FX/Niagara/Source/NiagaraEditor/Private/TraversalCache/TraversalBuilder.h b/Engine/Plugins/FX/Niagara/Source/NiagaraEditor/Private/TraversalCache/TraversalBuilder.h
--- a/Engine/Plugins/FX/Niagara/Source/NiagaraEditor/Private/TraversalCache/TraversalBuilder.h
+++ b/Engine/Plugins/FX/Niagara/Source/NiagaraEditor/Private/TraversalCache/TraversalBuilder.h
@@ -137,6 +137,8 @@ public:
 public:
 	void Initialize();
 
+	virtual ~FTraversalBuilder();
+
 	static TSharedRef<const FTraversal> BuildTraversal(const UNiagaraNodeOutput& OutputNode, FTraversalBuilderDebugData* DebugData = nullptr);
 
 	static NIAGARAEDITOR_API void ResolveFunctionCallStackNames(const UEdGraph& TopLevelGraph, const TArray<FGuid>& FunctionCallStack, TArray<FString>& OutFunctionCallNames);
diff --git a/Engine/Plugins/FX/Niagara/Source/NiagaraEditor/Private/TraversalCache/TraversalCache.cpp b/Engine/Plugins/FX/Niagara/Source/NiagaraEditor/Private/TraversalCache/TraversalCache.cpp
--- a/Engine/Plugins/FX/Niagara/Source/NiagaraEditor/Private/TraversalCache/TraversalCache.cpp
+++ b/Engine/Plugins/FX/Niagara/Source/NiagaraEditor/Private/TraversalCache/TraversalCache.cpp
@@ -28,10 +28,21 @@ void FTraversalCache::Initialize()
 		checkf(Instance == nullptr, TEXT("FTraversalCache initialized more than once."));
 		Instance = this;
 		FNiagaraEditorModule::Get().OnScriptApplied().AddSP(this, &FTraversalCache::OnScriptApplied);
+		GEditor->RegisterForUndo(this);
 		bInitialized = true;
 	}
 }
 
+FTraversalCache::~FTraversalCache()
+{
+	Instance = nullptr;
+	bInitialized = false;
+	if (GEditor != nullptr)
+	{
+		GEditor->UnregisterForUndo(this);
+	}
+}
+
 TSharedPtr<const FTraversal> FTraversalCache::GetScriptAssetTraversal(const FSoftObjectPath& ScriptSoftObjectPath, const FGuid& ScriptVersion, FTraversalBuilderDebugData* BuilderDebugData)
 {
 	return GetInstance().GetScriptAssetTraversalInternal(ScriptSoftObjectPath, ScriptVersion, BuilderDebugData);
@@ -116,6 +127,22 @@ void FTraversalCache::InvalidateCache()
 	GetInstance().InvalidateCacheInternal();
 }
 
+void FTraversalCache::PostUndo(bool bSuccess)
+{
+	// When undo/redo happens we we don't know exactly what's changed so we update this serial number so that traversals
+	// can invalidate any cached data that relies on external dependencies.  The traversals themselves are only invalidated
+	// when their graph change id is updated.
+	TraversalCachedDataSerialNumber++;
+}
+
+void FTraversalCache::PostRedo(bool bSuccess)
+{
+	// When undo/redo happens we we don't know exactly what's changed so we update this serial number so that traversals
+	// can invalidate any cached data that relies on external dependencies.  The traversals themselves are only invalidated
+	// when their graph change id is updated.
+	TraversalCachedDataSerialNumber++;
+}
+
 FTraversalCache& FTraversalCache::GetInstance()
 {
 	checkf(Instance != nullptr, TEXT("FTraversalCache was not initialized."));
@@ -140,53 +167,10 @@ const UNiagaraScript* FTraversalCache::GetScriptFromPath(const FSoftObjectPath&
 
 void FTraversalCache::OnScriptApplied(UNiagaraScript* InScript, FGuid InVersionGuid)
 {
-	TMap<FScriptTraversalCacheKey, TArray<FScriptTraversalCacheKey>> CacheKeyToReferencerCacheKeys;
-	for (const TPair<FScriptTraversalCacheKey, FScriptTraversalCacheValue>& CachedTraversal : CachedTraversals)
+	FScriptTraversalCacheKey ScriptCacheKey = FScriptTraversalCacheKey::FromScriptAsset(FObjectKey(InScript), InVersionGuid);
+	if (CachedTraversals.Contains(ScriptCacheKey))
 	{
-		for (const FScriptReference& ExternalReference : CachedTraversal.Value.ScriptTraversal->GetExternalReferences())
-		{
-			const UNiagaraScript* Script = GetScriptFromPath(ExternalReference.Path);
-			if (Script == nullptr)
-			{
-				continue;
-			}
-
-			FScriptTraversalCacheKey ExternalReferenceKey = FScriptTraversalCacheKey::FromScriptAsset(FObjectKey(Script), ExternalReference.Version);
-			TArray<FScriptTraversalCacheKey>& ReferencerCacheKeys = CacheKeyToReferencerCacheKeys.FindOrAdd(ExternalReferenceKey);
-			ReferencerCacheKeys.Add(CachedTraversal.Key);
-		}
-	}
-
-	TArray<FScriptTraversalCacheKey> CacheKeysToCheck = { FScriptTraversalCacheKey::FromScriptAsset(FObjectKey(InScript), InVersionGuid) };
-	int32 CacheKeysToCheckIndex = 0;
-
-	TSet<FScriptTraversalCacheKey> CacheKeysToRemove;
-	TSet<FScriptTraversalCacheKey> SeenDependents;
-	while (CacheKeysToCheckIndex < CacheKeysToCheck.Num())
-	{
-		const FScriptTraversalCacheKey& CurrentCacheKey = CacheKeysToCheck[CacheKeysToCheckIndex];
-		CacheKeysToCheckIndex++;
-
-		const TArray<FScriptTraversalCacheKey>* ReferencerCacheKeys = CacheKeyToReferencerCacheKeys.Find(CurrentCacheKey);
-		if (ReferencerCacheKeys == nullptr || ReferencerCacheKeys->Num() == 0)
-		{
-			continue;
-		}
-
-		for (const FScriptTraversalCacheKey& ReferencerCacheKey : *ReferencerCacheKeys)
-		{
-			bool bAlreadyAdded = false;
-			CacheKeysToRemove.Add(ReferencerCacheKey, &bAlreadyAdded);
-			if (bAlreadyAdded == false)
-			{
-				CacheKeysToCheck.Add(ReferencerCacheKey);
-			}
-		}
-	}
-
-	for (const FScriptTraversalCacheKey& CacheKeyToRemove : CacheKeysToRemove)
-	{
-		CachedTraversals.Remove(CacheKeyToRemove);
+		TraversalCachedDataSerialNumber++;
 	}
 }
 
@@ -210,6 +194,14 @@ TSharedPtr<const FTraversal> FTraversalCache::GetScriptAssetTraversalInternal(co
 	FScriptTraversalCacheValue* CachedValue = CachedTraversals.Find(CacheKey);
 	if (CachedValue != nullptr && CachedValue->GraphChangeId == ChangeId)
 	{
+		if (CachedValue->TraversalCachedDataSerialNumber != TraversalCachedDataSerialNumber)
+		{
+			// If the change id still matches, but the traversal cached data serial number has changed, it's possible that a dependent script
+			// has changed, so we need to invalidate any traversal data collected with this traversal before updating to the new
+			// serial number.
+			CachedValue->ScriptTraversal->ResetCachedData();
+			CachedValue->TraversalCachedDataSerialNumber = TraversalCachedDataSerialNumber;
+		}
 		return CachedValue->ScriptTraversal;
 	}
 
@@ -220,11 +212,10 @@ TSharedPtr<const FTraversal> FTraversalCache::GetScriptAssetTraversalInternal(co
 		return TSharedPtr<const FTraversal>();
 	}
 
-
 	TSharedRef<const FTraversal> Traversal = FTraversalBuilder::BuildTraversal(*OutputNode, BuilderDebugData);
 	if (BuilderDebugData == nullptr || BuilderDebugData->HasData() == false)
 	{
-		CachedTraversals.Add(CacheKey, FScriptTraversalCacheValue(ChangeId, Traversal));
+		CachedTraversals.Add(CacheKey, FScriptTraversalCacheValue(ChangeId, TraversalCachedDataSerialNumber, Traversal));
 	}
 	return Traversal;
 }
@@ -239,6 +230,14 @@ TSharedPtr<const FTraversal> FTraversalCache::GetTopLevelScriptTraversalInternal
 	FScriptTraversalCacheValue* CachedValue = CachedTraversals.Find(CacheKey);
 	if (CachedValue != nullptr && CachedValue->GraphChangeId == ChangeId)
 	{
+		if (CachedValue->TraversalCachedDataSerialNumber != TraversalCachedDataSerialNumber)
+		{
+			// If the change id still matches, but the traversal cached data serial number has changed, it's possible that a dependent script
+			// has changed, so we need to invalidate any traversal data collected with this traversal before updating to the new
+			// serial number.
+			CachedValue->ScriptTraversal->ResetCachedData();
+			CachedValue->TraversalCachedDataSerialNumber = TraversalCachedDataSerialNumber;
+		}
 		return CachedValue->ScriptTraversal;
 	}
 
@@ -253,7 +252,7 @@ TSharedPtr<const FTraversal> FTraversalCache::GetTopLevelScriptTraversalInternal
 	{
 		// We only cache the results when either no debug data was supplied, or when the debug data was empty.  When
 		// the debug data HasData returns true, errors were encountered while building.
-		CachedTraversals.Add(CacheKey, FScriptTraversalCacheValue(ChangeId, Traversal));
+		CachedTraversals.Add(CacheKey, FScriptTraversalCacheValue(ChangeId, TraversalCachedDataSerialNumber, Traversal));
 	}
 	return Traversal;
 }
@@ -273,9 +272,18 @@ TSharedPtr<const FTraversal> FTraversalCache::GetTopLevelScriptTraversalInternal
 	}
 
 	FScriptTraversalCacheKey CacheKey = FScriptTraversalCacheKey::FromTopLevelScript(FObjectKey(OwningGraph), OutputNode->GetUsage(), OutputNode->GetUsageId());
+
 	FScriptTraversalCacheValue* CachedValue = CachedTraversals.Find(CacheKey);
 	if (CachedValue != nullptr && CachedValue->GraphChangeId == OwningGraph->GetChangeID())
 	{
+		if (CachedValue->TraversalCachedDataSerialNumber != TraversalCachedDataSerialNumber)
+		{
+			// If the change id still matches, but the traversal cached data serial number has changed, it's possible that a dependent script
+			// has changed, so we need to invalidate any traversal data collected with this traversal before updating to the new
+			// serial number.
+			CachedValue->ScriptTraversal->ResetCachedData();
+			CachedValue->TraversalCachedDataSerialNumber = TraversalCachedDataSerialNumber;
+		}
 		return CachedValue->ScriptTraversal;
 	}
 
@@ -284,7 +292,7 @@ TSharedPtr<const FTraversal> FTraversalCache::GetTopLevelScriptTraversalInternal
 	{
 		// We only cache the results when either no debug data was supplied, or when the debug data was empty.  When
 		// the debug data HasData returns true, errors were encountered while building.
-		CachedTraversals.Add(CacheKey, FScriptTraversalCacheValue(OwningGraph->GetChangeID(), Traversal));
+		CachedTraversals.Add(CacheKey, FScriptTraversalCacheValue(OwningGraph->GetChangeID(), TraversalCachedDataSerialNumber, Traversal));
 	}
 	return Traversal;
 }
diff --git a/Engine/Plugins/FX/Niagara/Source/NiagaraEditor/Private/TraversalCache/TraversalCache.h b/Engine/Plugins/FX/Niagara/Source/NiagaraEditor/Private/TraversalCache/TraversalCache.h
--- a/Engine/Plugins/FX/Niagara/Source/NiagaraEditor/Private/TraversalCache/TraversalCache.h
+++ b/Engine/Plugins/FX/Niagara/Source/NiagaraEditor/Private/TraversalCache/TraversalCache.h
@@ -2,6 +2,7 @@
 
 #pragma once
 
+#include "EditorUndoClient.h"
 #include "Misc/Guid.h"
 #include "NiagaraCommon.h"
 #include "NiagaraTypes.h"
@@ -28,11 +29,13 @@ enum class EStackFunctionReadFilterFlags : uint8
 	EnabledOnly = 2
 };
 
-class FTraversalCache : public TSharedFromThis<FTraversalCache>
+class FTraversalCache : public TSharedFromThis<FTraversalCache>, public FEditorUndoClient
 {
 public:
 	void Initialize();
 
+	virtual ~FTraversalCache();
+
 	static TSharedPtr<const FTraversal> GetScriptAssetTraversal(const FSoftObjectPath& ScriptSoftObjectPath, const FGuid& ScriptVersion, FTraversalBuilderDebugData* BuilderDebugData = nullptr);
 
 	static TSharedPtr<const FTraversal> GetScriptAssetTraversal(const FSoftObjectPath& ScriptSoftObjectPath, const FGuid& ScriptVersion, FTraversalDebugData* TraversalDebugData = nullptr);
@@ -67,6 +70,10 @@ public:
 
 	NIAGARAEDITOR_API static void InvalidateCache();
 
+	//~ FEditorUndoClient interface
+	virtual void PostUndo(bool bSuccess) override;
+	virtual void PostRedo(bool bSuccess) override;
+
 private:
 	void OnScriptApplied(UNiagaraScript* InScript, FGuid InVersionGuid);
 	
@@ -180,13 +187,15 @@ private:
 
 	struct FScriptTraversalCacheValue
 	{
-		FScriptTraversalCacheValue(FGuid InGraphChangeId, TSharedRef<const FTraversal> InScriptTraversal)
+		FScriptTraversalCacheValue(FGuid InGraphChangeId, const uint32& InTraversalCachedDataSerialNumber, TSharedRef<const FTraversal> InScriptTraversal)
 			: GraphChangeId(InGraphChangeId)
+			, TraversalCachedDataSerialNumber(InTraversalCachedDataSerialNumber)
 			, ScriptTraversal(InScriptTraversal)
 		{
 		}
 
 		FGuid GraphChangeId;
+		uint32 TraversalCachedDataSerialNumber;
 		TSharedRef<const FTraversal> ScriptTraversal;
 	};
 
@@ -213,6 +222,11 @@ private:
 	static FTraversalCache* Instance;
 	bool bInitialized = false;
 
+	/** A serial number to track potential dependent graph changes.  We already discard cached traversals when the graph change id updates,
+		but the data collected by the visitors can depend on external references which aren't captured by the graphs change id.  When this
+		number is updated it signals the traversals to discard any data collected by visitors so that it can be recomputed when requested. */
+	uint32 TraversalCachedDataSerialNumber = 0;
+
 	mutable TMap<FScriptTraversalCacheKey, FScriptTraversalCacheValue> CachedTraversals;
 	mutable TMap<FScriptTraversalCacheKey, FParsedRapidIterationParameterData> CachedParsedRapidIterationParameterData;
 	mutable TMap<FSoftObjectPath, TWeakObjectPtr<const UNiagaraScript>> SoftObjectPathToScriptCache;
diff --git a/Engine/Plugins/FX/Niagara/Source/NiagaraEditor/Private/ViewModels/Stack/NiagaraStackGraphUtilities.cpp b/Engine/Plugins/FX/Niagara/Source/NiagaraEditor/Private/ViewModels/Stack/NiagaraStackGraphUtilities.cpp
--- a/Engine/Plugins/FX/Niagara/Source/NiagaraEditor/Private/ViewModels/Stack/NiagaraStackGraphUtilities.cpp
+++ b/Engine/Plugins/FX/Niagara/Source/NiagaraEditor/Private/ViewModels/Stack/NiagaraStackGraphUtilities.cpp
@@ -4148,6 +4148,7 @@ void FNiagaraStackGraphUtilities::SynchronizeReferencingMapPinsWithFunctionCall(
 
 	TSharedPtr<TMap<FGuid, FNiagaraVariable>> AssetVariableIdToVariableMap;
 	const UEdGraphSchema_Niagara* NiagaraSchema = GetDefault<UEdGraphSchema_Niagara>();
+	TSet<UNiagaraGraph*> ModifiedGraphs;
 	
 	// Update the override parameter map set nodes which are used to set the function inputs.
 	TArray<UEdGraphPin*> OverridePins;
@@ -4160,6 +4161,7 @@ void FNiagaraStackGraphUtilities::SynchronizeReferencingMapPinsWithFunctionCall(
 	{
 		// this contains the current up to date metadata of the asset graph this function call points to
 		AssetVariableIdToVariableMap = GetVariableIdToVariableMap(InFunctionCallNode);
+		bool bModifiedGraph = false;
 		for (UEdGraphPin* OverridePin : OverridePins)
 		{
 			FNiagaraVariable CurrentPinOverrideVariable = NiagaraSchema->PinToNiagaraVariable(OverridePin);
@@ -4177,18 +4179,22 @@ void FNiagaraStackGraphUtilities::SynchronizeReferencingMapPinsWithFunctionCall(
 					{
 						FNiagaraVariable UnaliasedCurrentPinOverrideVariable = FNiagaraUtilities::ResolveAliases(CurrentPinOverrideVariable, FNiagaraAliasContext().ChangeModuleNameToModule(AliasedOverridePinHandle.GetName().ToString()));
 						FNiagaraParameterHandle UnaliasedCurrentPinOverrideVariableHandle = FNiagaraParameterHandle(UnaliasedCurrentPinOverrideVariable.GetName());
-						
+
 						// update the override pin if the parameter has changed
 						if (UnaliasedCurrentPinOverrideVariableHandle.GetName() != CurrentAssetParameterHandle.GetName())
 						{
 							FNiagaraVariable UpdatedAliasedOverride = FNiagaraUtilities::ResolveAliases(*CurrentAssetParameter, FNiagaraAliasContext().ChangeModuleToModuleName(InFunctionCallNode.GetFunctionName()));
 							FNiagaraParameterHandle UpdatedOverrideHandle = FNiagaraParameterHandle(UpdatedAliasedOverride.GetName());
+							OverridePin->Modify();
 							OverridePin->PinName = UpdatedOverrideHandle.GetParameterHandleString();
+							InFunctionCallNode.Modify();
 							InFunctionCallNode.UpdateInputNameBinding(BoundGuid, UpdatedOverrideHandle.GetParameterHandleString());
+							bModifiedGraph = true;
 						}
 						if (CurrentPinOverrideVariable.GetType() != CurrentAssetParameter->GetType())
 						{
 							FEdGraphPinType NewPinType = NiagaraSchema->TypeDefinitionToPinType(CurrentAssetParameter->GetType());
+							OverridePin->Modify();
 							OverridePin->PinType = NewPinType;
 							if (OverridePin->LinkedTo.Num() == 1 &&
 								OverridePin->LinkedTo[0]->GetOwningNode() != nullptr &&
@@ -4196,18 +4202,24 @@ void FNiagaraStackGraphUtilities::SynchronizeReferencingMapPinsWithFunctionCall(
 							{
 								// Custom hlsl nodes in the stack are expression dynamic inputs and their output pin type must match their connected input
 								// so update that here.
+								OverridePin->LinkedTo[0]->Modify();
 								OverridePin->LinkedTo[0]->PinType = NewPinType;
 							}
+							bModifiedGraph = true;
 						}
 						break;
 					}
 				}
 			}
 		}
+		if (bModifiedGraph)
+		{
+			ModifiedGraphs.Add(InFunctionCallNode.GetNiagaraGraph());
+		}
 	}
 
 	// Update linked module outputs or module attributes referenced by downstream parameter map gets.
-	auto SynchronizeForStackLinkedModuleParameterPins = [&NiagaraSchema, &InFunctionCallNode, &AssetVariableIdToVariableMap](const TArray<UEdGraphPin*>& StackLinkedModuleParameterPins)
+	auto SynchronizeForStackLinkedModuleParameterPins = [&NiagaraSchema, &InFunctionCallNode, &AssetVariableIdToVariableMap, &ModifiedGraphs](const TArray<UEdGraphPin*>& StackLinkedModuleParameterPins)
 	{
 		if (AssetVariableIdToVariableMap.IsValid() == false)
 		{
@@ -4220,6 +4232,7 @@ void FNiagaraStackGraphUtilities::SynchronizeReferencingMapPinsWithFunctionCall(
 			UNiagaraNodeFunctionCall* TargetFunctionCallNode = GetTargetFunctionCallForStackLinkedModuleParameterPin(*StackLinkedModuleParameterPin);
 			if (TargetFunctionCallNode != nullptr)
 			{
+				bool bModifiedGraph = false;
 				TArray<FGuid> BoundGuids = TargetFunctionCallNode->GetBoundPinGuidsByName(LinkedModuleParameterHandle.GetParameterHandleString());
 				for (const FGuid& BoundGuid : BoundGuids)
 				{
@@ -4232,17 +4245,26 @@ void FNiagaraStackGraphUtilities::SynchronizeReferencingMapPinsWithFunctionCall(
 						{
 							if (LinkedModuleParameterHandle.GetName() != BoundVariableHandle.GetName())
 							{
+								StackLinkedModuleParameterPin->Modify();
 								StackLinkedModuleParameterPin->PinName = BoundVariableResolvedName.GetName();
+								TargetFunctionCallNode->Modify();
 								TargetFunctionCallNode->UpdateInputNameBinding(BoundGuid, BoundVariableHandle.GetParameterHandleString());
+								bModifiedGraph = true;
 							}
 							if (LinkedModuleParameterVariable.GetType() != BoundVariable->GetType())
 							{
+								StackLinkedModuleParameterPin->Modify();
 								StackLinkedModuleParameterPin->PinType = NiagaraSchema->TypeDefinitionToPinType(BoundVariable->GetType());
+								bModifiedGraph = true;
 							}
 							break;
 						}
 					}
 				}
+				if (bModifiedGraph)
+				{
+					ModifiedGraphs.Add(TargetFunctionCallNode->GetNiagaraGraph());
+				}
 			}
 		}
 	};
@@ -4260,6 +4282,11 @@ void FNiagaraStackGraphUtilities::SynchronizeReferencingMapPinsWithFunctionCall(
 	{
 		SynchronizeForStackLinkedModuleParameterPins(StackLinkedAttributePins);
 	}
+
+	for (UNiagaraGraph* ModifiedGraph : ModifiedGraphs)
+	{
+		ModifiedGraph->MarkGraphRequiresSynchronization(__FUNCTION__);
+	}
 }
 
 FGuid FNiagaraStackGraphUtilities::GetScriptVariableIdForLinkedModuleParameter(const FNiagaraVariableBase& LinkedParameter, UNiagaraGraph& TargetGraph)
diff --git a/Engine/Plugins/FX/Niagara/Source/NiagaraEditor/Public/NiagaraEditorModule.h b/Engine/Plugins/FX/Niagara/Source/NiagaraEditor/Public/NiagaraEditorModule.h
--- a/Engine/Plugins/FX/Niagara/Source/NiagaraEditor/Public/NiagaraEditorModule.h
+++ b/Engine/Plugins/FX/Niagara/Source/NiagaraEditor/Public/NiagaraEditorModule.h
@@ -305,7 +305,7 @@ public:
 
 	void PreloadSelectablePluginAssetsByClass(UClass* InClass);
 
-	void ScriptApplied(UNiagaraScript* Script, FGuid VersionGuid = FGuid()) const;
+	NIAGARAEDITOR_API void ScriptApplied(UNiagaraScript* Script, FGuid VersionGuid = FGuid()) const;
 	
 	/** Callback whenever a script is applied/updated in the editor. */
 	FOnScriptApplied& OnScriptApplied();
