- Fix for float bindings outside of parameter data not getting the correct data
- They were indexing into the parameter store as a byte offset vs an element offset

#tests Added a particle lifetime binding to existing emitter with other bindings and observed it pulling the wrong data in
#rnx

#virtualized

diff --git a/Engine/Plugins/FX/Niagara/Source/Niagara/Private/Stateless/NiagaraStatelessEmitterInstance.cpp b/Engine/Plugins/FX/Niagara/Source/Niagara/Private/Stateless/NiagaraStatelessEmitterInstance.cpp
--- a/Engine/Plugins/FX/Niagara/Source/Niagara/Private/Stateless/NiagaraStatelessEmitterInstance.cpp
+++ b/Engine/Plugins/FX/Niagara/Source/Niagara/Private/Stateless/NiagaraStatelessEmitterInstance.cpp
@@ -59,7 +59,7 @@ namespace NiagaraStatelessEmitterInstancePrivate
 	{
 		if (Distribution.ParameterOffset != INDEX_NONE)
 		{
-			OutMin = ParameterStore.GetParameterValueFromOffset<float>(Distribution.ParameterOffset);
+			OutMin = ParameterStore.GetParameterValueFromOffset<float>(Distribution.ParameterOffset * sizeof(int32));
 			OutMax = OutMin;
 		}
 		else
