- Change default for LOD size Y to -1 to fix issue with meshes being culled incorrectly
- A value of -1 indicates that it will never be culled by screen size
- The previous value of 2 would be adjusted to 1.0 and fail the >= screen size test if bounds were large enough

diff --git a/Engine/Plugins/FX/Niagara/Source/Niagara/Public/NiagaraRendererMeshes.h b/Engine/Plugins/FX/Niagara/Source/Niagara/Public/NiagaraRendererMeshes.h
--- a/Engine/Plugins/FX/Niagara/Source/Niagara/Public/NiagaraRendererMeshes.h
+++ b/Engine/Plugins/FX/Niagara/Source/Niagara/Public/NiagaraRendererMeshes.h
@@ -92,7 +92,7 @@ protected:
 		ENiagaraMeshLODMode LODMode = ENiagaraMeshLODMode::LODLevel;
 		int32 LODLevel = 0;
 		float LODDistanceFactor = 1.0f;
-		FVector3f LODScreenSize = FVector3f(0.0f, 2.0f, 1.0f);
+		FVector3f LODScreenSize = FVector3f(0.0f, -1.0f, 1.0f);		// Defaults set to not use screen size culled, 0.0 is min size, -1.0f is max size which is tranlated to > 1.0
 		FVector3f PivotOffset = FVector3f::ZeroVector;
 		ENiagaraMeshPivotOffsetSpace PivotOffsetSpace = ENiagaraMeshPivotOffsetSpace::Mesh;
 		FVector3f Scale = FVector3f(1.0f, 1.0f, 1.0f);
