Another pass on trying to clean up dynamic MIC creation for the renderer properties.
#tests Cooked assets requiring MICs to be recreated and enginetests

diff --git a/Engine/Plugins/FX/Niagara/Source/Niagara/Private/NiagaraRendererProperties.cpp b/Engine/Plugins/FX/Niagara/Source/Niagara/Private/NiagaraRendererProperties.cpp
--- a/Engine/Plugins/FX/Niagara/Source/Niagara/Private/NiagaraRendererProperties.cpp
+++ b/Engine/Plugins/FX/Niagara/Source/Niagara/Private/NiagaraRendererProperties.cpp
@@ -43,6 +43,12 @@ static FAutoConsoleVariableRef CVarNiagaraRendererCookOutStaticEnabledBinding(
 	TEXT("If none zero renderers with static variables used for enabled binding will cook out if they are not enabled."),
 	ECVF_Scalability
 );
+#endif
+
+namespace NiagaraRendererPropertiesPrivate
+{
+
+#if WITH_EDITORONLY_DATA
 
 // Attempts to resolve a static variable across all the scripts
 // If the state is undetermined (i.e. does not exist or is inconsistent across scripts) we will not return a value
@@ -98,8 +104,31 @@ static TOptional<bool> TryResolveStaticVariableBool(const UNiagaraEmitter* Niaga
 
 	return StaticValue;
 }
+
 #endif
 
+void MarkAndRenameMaterialForGarbage(UMaterialInstanceConstant* MIC)
+{
+	if (MIC)
+	{
+		const ERenameFlags RenameFlags = REN_NonTransactional | REN_DoNotDirty | REN_DontCreateRedirectors;
+		UPackage* TransientPackage = GetTransientPackage();
+
+		MIC->MarkAsGarbage();
+		MIC->Rename(nullptr, TransientPackage, RenameFlags);
+
+#if WITH_EDITORONLY_DATA
+		if (UMaterialInterfaceEditorOnlyData* EditorOnlyData = MIC->GetEditorOnlyData())
+		{
+			EditorOnlyData->MarkAsGarbage();
+			EditorOnlyData->Rename(nullptr, TransientPackage, RenameFlags);
+		}
+#endif
+	}
+}
+
+} // NiagaraRendererPropertiesPrivate
+
 void FNiagaraRendererLayout::Initialize(int32 NumVariables)
 {
 	VFVariables_GT.Reset(NumVariables);
@@ -471,7 +500,10 @@ bool UNiagaraRendererProperties::BuildMaterialStaticParameterSet(const FNiagaraR
 					{
 						if (bDefaultValue != bParameterValue)
 						{
-							StaticParameterSet.SetParameterValue(ParameterInfo, FMaterialParameterValue(bParameterValue), EMaterialSetParameterValueFlags::None);
+							FMaterialParameterMetadata ParameterMetadata(bParameterValue);
+							ParameterMetadata.ExpressionGuid = ParameterGuid;
+
+							StaticParameterSet.SetParameterValue(ParameterInfo, ParameterMetadata, EMaterialSetParameterValueFlags::None);
 							bModified = true;
 						}
 					}
@@ -539,49 +571,16 @@ bool UNiagaraRendererProperties::UpdateMaterialStaticParameters(const FNiagaraRe
 
 void UNiagaraRendererProperties::UpdateMaterialParametersMIC(const FNiagaraRendererMaterialParameters& MaterialParameters, TObjectPtr<UMaterialInterface>& InOutMaterial, TObjectPtr<UMaterialInstanceConstant>& InOutMIC)
 {
-	if (InOutMaterial == nullptr || MaterialParameters.StaticBoolParameters.Num() == 0)
-	{
-		InOutMIC = nullptr;
-		return;
-	}
+	TArray<UMaterialInterface*> LocalMaterials = { InOutMaterial };
+	TArray<TObjectPtr<UMaterialInstanceConstant>> LocalMICs = { InOutMIC };
 
-	FStaticParameterSet MaterialParameterSet;
-	if (BuildMaterialStaticParameterSet(MaterialParameters, InOutMaterial, MaterialParameterSet))
-	{
-		FNameBuilder NameBuilder;
-		InOutMaterial->GetFName().ToString(NameBuilder);
-		NameBuilder.Append(TEXT("_MIC"));
-		if (InOutMIC == nullptr)
-		{
-			InOutMIC = NewObject<UMaterialInstanceConstant>(this, FName(NameBuilder));
-			InOutMIC->PreEditChange(nullptr);
-			InOutMIC->SetParentEditorOnly(InOutMaterial);
-		}
-		else if (InOutMIC->Parent != InOutMaterial)
-		{
-			InOutMIC->Rename(NameBuilder.ToString());
-			InOutMIC->PreEditChange(nullptr);
-			InOutMIC->SetParentEditorOnly(InOutMaterial);
-		}
+	UpdateMaterialParametersMIC(MaterialParameters, LocalMaterials, LocalMICs);
 
-		InOutMIC->UpdateStaticPermutation(MaterialParameterSet);
-		InOutMIC->PostEditChange();
-	}
-	else if (InOutMIC)
-	{
-		InOutMIC->MarkAsGarbage();
-		InOutMIC = nullptr;
-	}
+	InOutMIC = LocalMICs.IsEmpty() ? nullptr : LocalMICs[0];
 }
 
 void UNiagaraRendererProperties::UpdateMaterialParametersMIC(const FNiagaraRendererMaterialParameters& MaterialParameters, TArrayView<UMaterialInterface*> Materials, TArray<TObjectPtr<UMaterialInstanceConstant>>& InOutMICs)
 {
-	if (Materials.Num() == 0 || MaterialParameters.StaticBoolParameters.Num() == 0)
-	{
-		InOutMICs.Empty();
-		return;
-	}
-
 	// Create a pool of unique MICs that we can potentially reuse
 	TArray<UMaterialInstanceConstant*> MICPool;
 	MICPool.Reserve(InOutMICs.Num());
@@ -594,6 +593,20 @@ void UNiagaraRendererProperties::UpdateMaterialParametersMIC(const FNiagaraRende
 	}
 	InOutMICs.Reset(0);
 
+	ON_SCOPE_EXIT
+	{
+		// Garabge the MICs we are no longer going to use
+		for (int i = 0; i < MICPool.Num(); ++i)
+		{
+			NiagaraRendererPropertiesPrivate::MarkAndRenameMaterialForGarbage(MICPool[i]);
+		}
+	};
+
+	if (Materials.Num() == 0 || MaterialParameters.StaticBoolParameters.Num() == 0)
+	{
+		return;
+	}
+
 	// Loop over each material to see if we need to generate a MIC for it
 	for ( int i=0; i < Materials.Num(); ++i )
 	{
@@ -612,30 +625,30 @@ void UNiagaraRendererProperties::UpdateMaterialParametersMIC(const FNiagaraRende
 			NameBuilder.Append(TEXT("_MIC"));
 
 			UMaterialInstanceConstant* MIC = nullptr;
-			bool bNeedsRename = false;
 			if (MICPool.Num() > 0)
 			{
 				FName MICName(NameBuilder);
-				const int32 ExistingIndex = MICPool.IndexOfByPredicate([&MICName](UMaterialInstanceConstant* MIC) { return MIC->GetFName() == MICName; });
+				auto MatchesNameAndMaterial = [&MICName, Material](UMaterialInstanceConstant* InMIC) -> bool
+				{
+					return InMIC->GetFName() == MICName && InMIC->Parent == Material;
+				};
+
+				const int32 ExistingIndex = MICPool.IndexOfByPredicate(MatchesNameAndMaterial);
 				if (ExistingIndex != INDEX_NONE)
 				{
 					MIC = MICPool[ExistingIndex];
 					MICPool.RemoveAtSwap(ExistingIndex, EAllowShrinking::No);
 				}
-				else
-				{
-					bNeedsRename = true;
-					MIC = MICPool.Pop();
-				}
 				
 				MIC->PreEditChange(nullptr);
-				if (MIC->Parent != Material)
-				{
-					MIC->SetParentEditorOnly(Material);
-				}
 			}
 			else
 			{
+				if (const UObject* ExistingObject = StaticFindObject(UMaterialInstanceConstant::StaticClass(), this, *NameBuilder))
+				{
+					UE_LOG(LogNiagara, Warning, TEXT("While trying to allocate %s UNiagaraRendererProperties::UpdateMaterialParametersMIC() found a pre-existing object.  This will result in a re-allocation!"), *NameBuilder);
+				}
+
 				MIC = NewObject<UMaterialInstanceConstant>(this, FName(NameBuilder));
 				MIC->PreEditChange(nullptr);
 				MIC->SetParentEditorOnly(Material);
@@ -646,21 +659,8 @@ void UNiagaraRendererProperties::UpdateMaterialParametersMIC(const FNiagaraRende
 
 			InOutMICs.SetNum(i + 1);
 			InOutMICs[i] = MIC;
-			if (bNeedsRename)
-			{
-				if (MIC->Rename(NameBuilder.ToString(), nullptr, REN_Test))
-				{
-					MIC->Rename(NameBuilder.ToString());
-				}
-			}
 		}
 	}
-
-	// Garabge the MICs we are no longer going to use
-	for (int i=0; i < MICPool.Num(); ++i)
-	{
-		MICPool[i]->MarkAsGarbage();
-	}
 }
 
 int32 UNiagaraRendererProperties::GetDynamicParameterChannelMask(const FVersionedNiagaraEmitterData* EmitterData, FName BindingName, int32 DefaultChannelMask) const
@@ -1042,7 +1042,7 @@ bool UNiagaraRendererProperties::NeedsLoadForTargetPlatform(const ITargetPlatfor
 			#if WITH_EDITORONLY_DATA
 				if (GNiagaraRendererCookOutStaticEnabledBinding && RendererEnabledBinding.IsValid())
 				{
-					TOptional<bool> ResolvedStaticValue = TryResolveStaticVariableBool(OwnerEmitter, RendererEnabledBinding.GetParamMapBindableVariable());
+					TOptional<bool> ResolvedStaticValue = NiagaraRendererPropertiesPrivate::TryResolveStaticVariableBool(OwnerEmitter, RendererEnabledBinding.GetParamMapBindableVariable());
 					return ResolvedStaticValue.Get(true);
 				}
 			#endif
@@ -1060,7 +1060,7 @@ bool UNiagaraRendererProperties::NeedsLoadForTargetPlatform(const ITargetPlatfor
 //#if WITH_EDITORONLY_DATA
 //				if (GNiagaraRendererCookOutStaticEnabledBinding && RendererEnabledBinding.IsValid())
 //				{
-//					TOptional<bool> ResolvedStaticValue = TryResolveStaticVariableBool(OwnerEmitter, RendererEnabledBinding.GetParamMapBindableVariable());
+//					TOptional<bool> ResolvedStaticValue = NiagaraRendererPropertiesPrivate::TryResolveStaticVariableBool(OwnerEmitter, RendererEnabledBinding.GetParamMapBindableVariable());
 //					return ResolvedStaticValue.Get(true);
 //				}
 //#endif
