- Add debug draw to renderers
- Implement debug draw for light & decal renderer

diff --git a/Engine/Plugins/FX/Niagara/Source/Niagara/Private/NiagaraRendererDecals.cpp b/Engine/Plugins/FX/Niagara/Source/Niagara/Private/NiagaraRendererDecals.cpp
--- a/Engine/Plugins/FX/Niagara/Source/Niagara/Private/NiagaraRendererDecals.cpp
+++ b/Engine/Plugins/FX/Niagara/Source/Niagara/Private/NiagaraRendererDecals.cpp
@@ -26,18 +26,10 @@ namespace NiagaraRendererDecalsLocal
 
 #define NIAGARARENDERERDECALS_DEBUG_ENABLED !UE_BUILD_SHIPPING
 
-#if NIAGARARENDERERDECALS_DEBUG_ENABLED
-	static bool GDrawDebug = false;
-	static FAutoConsoleVariableRef CVarDrawDebug(
-		TEXT("fx.Niagara.DecalRenderer.DrawDebug"),
-		GDrawDebug,
-		TEXT("When none zero will draw debug information."),
-		ECVF_Default
-	);
-
+#if WITH_NIAGARA_RENDERER_DEBUGDRAW
 	void DrawDebugDecal(UWorld* World, TConstArrayView<FDeferredDecalUpdateParams> AllUpdateParams, double WorldTime)
 	{
-		if (GDrawDebug == false)
+		if (World == nullptr || AllUpdateParams.Num() == 0)
 		{
 			return;
 		}
@@ -47,28 +39,33 @@ namespace NiagaraRendererDecalsLocal
 			ENamedThreads::GameThread,
 			[WeakWorld=MakeWeakObjectPtr(World), AllUpdateParams_GT=TArray<FDeferredDecalUpdateParams>(AllUpdateParams), WorldTime]
 			{
-				if (UWorld* World = WeakWorld.Get())
+				UWorld* World = WeakWorld.Get();
+				ULineBatchComponent* LineBatcher = World ? World->GetLineBatcher(UWorld::ELineBatcherType::World) : nullptr;
+				if (LineBatcher == nullptr)
+				{
+					return;
+				}
+				
+				for (const FDeferredDecalUpdateParams& UpdateParams : AllUpdateParams_GT)
 				{
-					if (ULineBatchComponent* LineBatcher = World->GetLineBatcher(UWorld::ELineBatcherType::World))
+					if ( UpdateParams.OperationType == FDeferredDecalUpdateParams::EOperationType::RemoveFromSceneAndDelete )
 					{
-						for (const FDeferredDecalUpdateParams& UpdateParams : AllUpdateParams_GT)
-						{
-							if ( UpdateParams.OperationType == FDeferredDecalUpdateParams::EOperationType::RemoveFromSceneAndDelete )
-							{
-								continue;
-							}
-
-							const float Fade = (WorldTime - UpdateParams.FadeStartDelay) / UpdateParams.FadeDuration;
-							const FBox BoundsBox(-FVector(1), FVector(1));
-							//const uint8 IntFade = FMath::Clamp(uint8(255.0f * Fade), 0, 255);
-							const FColor DecalColor = UpdateParams.DecalColor.ToFColor(false);
-							LineBatcher->DrawSolidBox(BoundsBox, UpdateParams.Transform, DecalColor, 0, 0.0f);
-
-							LineBatcher->DrawCircle(UpdateParams.Bounds.Origin, FVector::XAxisVector, FVector::YAxisVector, FColor::Red, UpdateParams.Bounds.SphereRadius, 16, 0);
-							LineBatcher->DrawCircle(UpdateParams.Bounds.Origin, FVector::XAxisVector, FVector::ZAxisVector, FColor::Red, UpdateParams.Bounds.SphereRadius, 16, 0);
-							LineBatcher->DrawCircle(UpdateParams.Bounds.Origin, FVector::YAxisVector, FVector::ZAxisVector, FColor::Red, UpdateParams.Bounds.SphereRadius, 16, 0);
-						}
+						continue;
 					}
+
+					//const float Fade = (WorldTime - UpdateParams.FadeStartDelay) / UpdateParams.FadeDuration;
+					//const uint8 IntFade = FMath::Clamp(uint8(255.0f * Fade), 0, 255);
+
+					const FMatrix DecalMatrix = UpdateParams.Transform.ToMatrixWithScale();
+					const FBox BoundsBox(-FVector(1), FVector(1));
+					const FColor DecalColor = UpdateParams.DecalColor.ToFColor(false);
+					LineBatcher->DrawBox(BoundsBox, DecalMatrix, DecalColor, 0, 0.0f);
+
+					const FVector StartLocation = UpdateParams.Transform.TransformPosition(FVector(-1.0f, 0.0f, 0.0f));
+					const FVector MidLocation = UpdateParams.Transform.TransformPosition(FVector(0.0f, 0.0f, 0.0f));
+					const float ArrowSize = FMath::Min(UpdateParams.Transform.GetScale3D().Y, UpdateParams.Transform.GetScale3D().Z) * 0.5f;
+					
+					LineBatcher->DrawDirectionalArrow(StartLocation, MidLocation, ArrowSize * ArrowSize, FLinearColor::Red, 0.0f, 0, 1.0f);
 				}
 			}
 		);
@@ -300,8 +297,11 @@ FNiagaraDynamicDataBase* FNiagaraRendererDecals::GenerateDynamicData(const FNiag
 	// Send updates to RT
 	if (DecalUpdates.Num() > 0)
 	{
-	#if NIAGARARENDERERDECALS_DEBUG_ENABLED
-		DrawDebugDecal(World, DecalUpdates, World->TimeSeconds);
+	#if WITH_NIAGARA_RENDERER_DEBUGDRAW
+		if (RendererProperties->IsDebugDrawEnabled())
+		{
+			DrawDebugDecal(World, DecalUpdates, World->TimeSeconds);
+		}
 	#endif
 		World->Scene->BatchUpdateDecals(MoveTemp(DecalUpdates));
 	}
diff --git a/Engine/Plugins/FX/Niagara/Source/Niagara/Private/NiagaraRendererLights.cpp b/Engine/Plugins/FX/Niagara/Source/Niagara/Private/NiagaraRendererLights.cpp
--- a/Engine/Plugins/FX/Niagara/Source/Niagara/Private/NiagaraRendererLights.cpp
+++ b/Engine/Plugins/FX/Niagara/Source/Niagara/Private/NiagaraRendererLights.cpp
@@ -16,17 +16,58 @@
 #include "PrimitiveViewRelevance.h"
 #include "SceneInterface.h"
 
+#include "Async/Async.h"
+#include "Components/LineBatchComponent.h"
+#include "Engine/World.h"
 #include "SceneView.h"
 
 DECLARE_CYCLE_STAT(TEXT("Generate Particle Lights"), STAT_NiagaraGenLights, STATGROUP_Niagara);
 
-static bool GbEnableNiagaraLightRendering = true;
-static FAutoConsoleVariableRef CVarEnableNiagaraLightRendering(
-	TEXT("fx.EnableNiagaraLightRendering"),
-	GbEnableNiagaraLightRendering,
-	TEXT("If false Niagara Light Renderers are disabled."),
-	ECVF_Default
-);
+namespace NiagaraRendererLightsPrivate
+{
+	static bool GbRendererEnabled = true;
+	static FAutoConsoleVariableRef CVarGbRendererEnabled(
+		TEXT("fx.EnableNiagaraLightRendering"),
+		GbRendererEnabled,
+		TEXT("If false Niagara Light Renderers are disabled."),
+		ECVF_Default
+	);
+
+#if WITH_NIAGARA_RENDERER_DEBUGDRAW 
+	void DebugDraw(FNiagaraSystemInstance* SystemInstance, TArrayView<FNiagaraRendererLights::SimpleLightData> InLightData)
+	{
+		UWorld* World = SystemInstance->GetWorld();
+		if (World == nullptr || InLightData.Num() == 0)
+		{
+			return;
+		}
+
+		AsyncTask(
+			ENamedThreads::GameThread,
+			[WeakWorld=MakeWeakObjectPtr(World), LightDataArray=TArray<FNiagaraRendererLights::SimpleLightData>(InLightData)]
+			{
+				UWorld* World = WeakWorld.Get();
+				ULineBatchComponent* LineBatcher = World ? World->GetLineBatcher(UWorld::ELineBatcherType::World) : nullptr;
+				if (LineBatcher == nullptr)
+				{
+					return;
+				}
+
+				for (const FNiagaraRendererLights::SimpleLightData& LightData : LightDataArray)
+				{
+					const FVector LightLocation = LightData.PerViewEntry.Position;
+					const float LightRadius = LightData.LightEntry.Radius;
+					const FColor LightColor = FLinearColor(LightData.LightEntry.Color.X, LightData.LightEntry.Color.Y, LightData.LightEntry.Color.Z).ToFColor(true);
+
+					LineBatcher->DrawCircle(LightLocation, FVector::XAxisVector, FVector::YAxisVector, LightColor, LightRadius, 16, 0);
+					LineBatcher->DrawCircle(LightLocation, FVector::XAxisVector, FVector::ZAxisVector, LightColor, LightRadius, 16, 0);
+					LineBatcher->DrawCircle(LightLocation, FVector::YAxisVector, FVector::ZAxisVector, LightColor, LightRadius, 16, 0);
+				}
+			}
+		);
+	}
+#endif //WITH_NIAGARA_RENDERER_DEBUGDRAW 
+}
 
 struct FNiagaraDynamicDataLights : public FNiagaraDynamicDataBase
 {
@@ -198,12 +239,19 @@ FNiagaraDynamicDataBase* FNiagaraRendererLights::GenerateDynamicData(const FNiag
 		}
 	}
 
+#if WITH_NIAGARA_RENDERER_DEBUGDRAW
+	if (Properties->IsDebugDrawEnabled())
+	{
+		NiagaraRendererLightsPrivate::DebugDraw(Emitter->GetParentSystemInstance(), DynamicData->LightArray);
+	}
+#endif
+
 	return DynamicData;
 }
 
 void FNiagaraRendererLights::GatherSimpleLights(FSimpleLightArray& OutParticleLights)const
 {
-	if (GbEnableNiagaraLightRendering == false)
+	if (NiagaraRendererLightsPrivate::GbRendererEnabled == false)
 	{
 		return;
 	}
diff --git a/Engine/Plugins/FX/Niagara/Source/Niagara/Private/NiagaraRendererProperties.cpp b/Engine/Plugins/FX/Niagara/Source/Niagara/Private/NiagaraRendererProperties.cpp
--- a/Engine/Plugins/FX/Niagara/Source/Niagara/Private/NiagaraRendererProperties.cpp
+++ b/Engine/Plugins/FX/Niagara/Source/Niagara/Private/NiagaraRendererProperties.cpp
@@ -853,6 +853,34 @@ FText UNiagaraRendererProperties::GetWidgetDisplayName() const
 	return GetClass()->GetDisplayNameText();
 }
 
+#if WITH_NIAGARA_RENDERER_DEBUGDRAW
+bool UNiagaraRendererProperties::SupportsDebugDraw() const
+{
+	return false;
+}
+
+TOptional<FText> UNiagaraRendererProperties::GetDebugDrawTooltip() const
+{
+	return {};
+}
+
+bool UNiagaraRendererProperties::IsDebugDrawEnabled() const
+{
+	if (SupportsDebugDraw() && bDebugDrawEnabled)
+	{
+		if (UNiagaraSystem* NiagaraSystem = GetTypedOuter<UNiagaraSystem>())
+		{
+			return NiagaraSystem->ShouldDisableDebugSwitches() == false;
+		}
+	}
+	return false;
+}
+void UNiagaraRendererProperties::SetDebugDrawEnabled(bool bInEnabled)
+{
+	bDebugDrawEnabled = bInEnabled;
+}
+#endif
+
 void UNiagaraRendererProperties::RenameVariable(const FNiagaraVariableBase& OldVariable, const FNiagaraVariableBase& NewVariable, const FVersionedNiagaraEmitter& InEmitter)
 {
 	// Handle the renaming of generic renderer bindings...
diff --git a/Engine/Plugins/FX/Niagara/Source/Niagara/Public/NiagaraDecalRendererProperties.h b/Engine/Plugins/FX/Niagara/Source/Niagara/Public/NiagaraDecalRendererProperties.h
--- a/Engine/Plugins/FX/Niagara/Source/Niagara/Public/NiagaraDecalRendererProperties.h
+++ b/Engine/Plugins/FX/Niagara/Source/Niagara/Public/NiagaraDecalRendererProperties.h
@@ -54,6 +54,9 @@ public:
 	virtual ENiagaraRendererSourceDataMode GetCurrentSourceMode() const override { return SourceMode; }
 	virtual bool PopulateRequiredBindings(FNiagaraParameterStore& InParameterStore) override;
 	virtual bool NeedsMIDsForMaterials() const override { return MaterialParameters.HasAnyBindings(); }
+#if WITH_NIAGARA_RENDERER_DEBUGDRAW
+	virtual bool SupportsDebugDraw() const { return true; }
+#endif
 	//UNiagaraRendererProperties Interface END
 
 	void UpdateMICs();
diff --git a/Engine/Plugins/FX/Niagara/Source/Niagara/Public/NiagaraLightRendererProperties.h b/Engine/Plugins/FX/Niagara/Source/Niagara/Public/NiagaraLightRendererProperties.h
--- a/Engine/Plugins/FX/Niagara/Source/Niagara/Public/NiagaraLightRendererProperties.h
+++ b/Engine/Plugins/FX/Niagara/Source/Niagara/Public/NiagaraLightRendererProperties.h
@@ -45,6 +45,9 @@ public:
 	virtual ENiagaraRendererSourceDataMode GetCurrentSourceMode() const override { return SourceMode; }
 	virtual bool PopulateRequiredBindings(FNiagaraParameterStore& InParameterStore) override;
 	virtual void CacheFromCompiledData(const FNiagaraDataSetCompiledData* CompiledData) override;
+#if WITH_NIAGARA_RENDERER_DEBUGDRAW
+	virtual bool SupportsDebugDraw() const { return true; }
+#endif
 	//UNiagaraRendererProperties Interface END
 
 	/** Whether or not to draw a single element for the Emitter or to draw the particles.*/
diff --git a/Engine/Plugins/FX/Niagara/Source/Niagara/Public/NiagaraRendererProperties.h b/Engine/Plugins/FX/Niagara/Source/Niagara/Public/NiagaraRendererProperties.h
--- a/Engine/Plugins/FX/Niagara/Source/Niagara/Public/NiagaraRendererProperties.h
+++ b/Engine/Plugins/FX/Niagara/Source/Niagara/Public/NiagaraRendererProperties.h
@@ -30,6 +30,8 @@ struct FStreamingRenderAssetPrimitiveInfo;
 enum EMaterialUsage : int;
 enum EMaterialDomain : int;
 
+#define WITH_NIAGARA_RENDERER_DEBUGDRAW	WITH_EDITORONLY_DATA
+
 #if WITH_EDITOR
 // Helper class for GUI error handling
 DECLARE_DELEGATE(FNiagaraRendererFeedbackFix);
@@ -345,7 +347,6 @@ public:
 	virtual void CollectPSOPrecacheData(const FNiagaraEmitterInstance* InEmitter, FPSOPrecacheParamsList& OutParams) const {};
 
 #if WITH_EDITORONLY_DATA
-
 	NIAGARA_API virtual bool IsSupportedVariableForBinding(const FNiagaraVariableBase& InSourceForBinding, const FName& InTargetBindingName) const;
 
 	/** Internal handling of any emitter variable renames. Note that this doesn't modify the renderer, the caller will need to do that if it is desired.*/
@@ -386,9 +387,13 @@ public:
 	NIAGARA_API virtual FText GetWidgetDisplayName() const;
 #endif // WITH_EDITORONLY_DATA
 
-#if WITH_EDITOR
-	NIAGARA_API FOnPropertiesChanged& OnPropertiesChanged();
+#if WITH_NIAGARA_RENDERER_DEBUGDRAW
+	NIAGARA_API virtual bool SupportsDebugDraw() const;
+	NIAGARA_API virtual TOptional<FText> GetDebugDrawTooltip() const;
+	NIAGARA_API virtual bool IsDebugDrawEnabled() const;
+	NIAGARA_API virtual void SetDebugDrawEnabled(bool bInEnabled);
 #endif
+
 	virtual ENiagaraRendererSourceDataMode GetCurrentSourceMode() const {	return ENiagaraRendererSourceDataMode::Particles;}
 
 	NIAGARA_API virtual bool GetIsActive() const;
@@ -480,6 +485,10 @@ protected:
 	NIAGARA_API int32 GetDynamicParameterCombinedChannelMask(FName Parameter0Name, FName Parameter1Name, FName Parameter2Name, FName Parameter3Name) const;
 #endif
 
+#if WITH_EDITORONLY_DATA
+	UPROPERTY()
+	bool bDebugDrawEnabled = false;
+#endif
 #if WITH_EDITOR
 	FOnPropertiesChanged OnPropertiesChangedDelegate;
 #endif
diff --git a/Engine/Plugins/FX/Niagara/Source/NiagaraEditor/Private/ViewModels/Stack/NiagaraStackModuleItem.cpp b/Engine/Plugins/FX/Niagara/Source/NiagaraEditor/Private/ViewModels/Stack/NiagaraStackModuleItem.cpp
--- a/Engine/Plugins/FX/Niagara/Source/NiagaraEditor/Private/ViewModels/Stack/NiagaraStackModuleItem.cpp
+++ b/Engine/Plugins/FX/Niagara/Source/NiagaraEditor/Private/ViewModels/Stack/NiagaraStackModuleItem.cpp
@@ -1343,6 +1343,23 @@ void UNiagaraStackModuleItem::SetIsEnabledInternal(bool bInIsEnabled)
 	OnRequestFullRefreshDeferred().Broadcast();
 }
 
+bool UNiagaraStackModuleItem::SupportsDebugDraw() const
+{
+	return GetModuleNode().ContainsDebugSwitch();
+}
+
+TOptional<FText> UNiagaraStackModuleItem::GetDebugDrawTooltip() const
+{
+	if (FVersionedNiagaraScriptData* ScriptData = GetModuleNode().GetScriptData())
+	{
+		if (ScriptData->DebugDrawMessage.IsEmpty() == false)
+		{
+			return ScriptData->DebugDrawMessage;
+		}
+	}
+	return {};
+}
+
 bool UNiagaraStackModuleItem::IsDebugDrawEnabled() const
 {
 	return FunctionCallNode->DebugState != ENiagaraFunctionDebugState::NoDebug;
diff --git a/Engine/Plugins/FX/Niagara/Source/NiagaraEditor/Private/ViewModels/Stack/NiagaraStackRendererItem.cpp b/Engine/Plugins/FX/Niagara/Source/NiagaraEditor/Private/ViewModels/Stack/NiagaraStackRendererItem.cpp
--- a/Engine/Plugins/FX/Niagara/Source/NiagaraEditor/Private/ViewModels/Stack/NiagaraStackRendererItem.cpp
+++ b/Engine/Plugins/FX/Niagara/Source/NiagaraEditor/Private/ViewModels/Stack/NiagaraStackRendererItem.cpp
@@ -684,5 +684,38 @@ void UNiagaraStackRendererItem::RendererChanged()
 	}
 }
 
+bool UNiagaraStackRendererItem::SupportsDebugDraw() const
+{
+	UNiagaraRendererProperties* Renderer = RendererProperties.Get();
+	return Renderer ? Renderer->SupportsDebugDraw() : false;
+}
+
+TOptional<FText> UNiagaraStackRendererItem::GetDebugDrawTooltip() const
+{
+	UNiagaraRendererProperties* Renderer = RendererProperties.Get();
+	return Renderer ? Renderer->GetDebugDrawTooltip() : TOptional<FText>();
+}
+
+bool UNiagaraStackRendererItem::IsDebugDrawEnabled() const
+{
+	UNiagaraRendererProperties* Renderer = RendererProperties.Get();
+	return Renderer ? Renderer->IsDebugDrawEnabled() : false;
+}
+
+void UNiagaraStackRendererItem::SetDebugDrawEnabled(bool bInEnabled)
+{
+	if (UNiagaraRendererProperties* Renderer = RendererProperties.Get())
+	{
+		FScopedTransaction ScopedTransaction(LOCTEXT("SetRendererDebugDrawEnabledState", "Set renderer debug draw enabled state."));
+		Renderer->Modify();
+		Renderer->SetDebugDrawEnabled(bInEnabled);
+
+		TArray<UObject*> ChangedObjects;
+		ChangedObjects.Add(Renderer);
+		OnDataObjectModified().Broadcast(ChangedObjects, ENiagaraDataObjectChange::Changed);
+		RefreshChildren();
+	}
+}
+
 #undef LOCTEXT_NAMESPACE
 
diff --git a/Engine/Plugins/FX/Niagara/Source/NiagaraEditor/Public/ViewModels/Stack/NiagaraStackItem.h b/Engine/Plugins/FX/Niagara/Source/NiagaraEditor/Public/ViewModels/Stack/NiagaraStackItem.h
--- a/Engine/Plugins/FX/Niagara/Source/NiagaraEditor/Public/ViewModels/Stack/NiagaraStackItem.h
+++ b/Engine/Plugins/FX/Niagara/Source/NiagaraEditor/Public/ViewModels/Stack/NiagaraStackItem.h
@@ -76,6 +76,11 @@ public:
 
 	virtual bool GetIsInherited() const { return false; }
 
+	virtual bool SupportsDebugDraw() const { return false; }
+	virtual TOptional<FText> GetDebugDrawTooltip() const { return {}; }
+	virtual bool IsDebugDrawEnabled() const { return false; }
+	virtual void SetDebugDrawEnabled(bool bInEnabled) {}
+
 protected:
 	NIAGARAEDITOR_API virtual void RefreshChildrenInternal(const TArray<UNiagaraStackEntry*>& CurrentChildren, TArray<UNiagaraStackEntry*>& NewChildren, TArray<FStackIssue>& NewIssues) override;
 
diff --git a/Engine/Plugins/FX/Niagara/Source/NiagaraEditor/Public/ViewModels/Stack/NiagaraStackModuleItem.h b/Engine/Plugins/FX/Niagara/Source/NiagaraEditor/Public/ViewModels/Stack/NiagaraStackModuleItem.h
--- a/Engine/Plugins/FX/Niagara/Source/NiagaraEditor/Public/ViewModels/Stack/NiagaraStackModuleItem.h
+++ b/Engine/Plugins/FX/Niagara/Source/NiagaraEditor/Public/ViewModels/Stack/NiagaraStackModuleItem.h
@@ -114,8 +114,10 @@ public:
 		SetIsEnabledInternal(bEnabled);
 	}
 
-	NIAGARAEDITOR_API bool IsDebugDrawEnabled() const;
-	NIAGARAEDITOR_API void SetDebugDrawEnabled(bool bInEnabled);
+	NIAGARAEDITOR_API virtual bool SupportsDebugDraw() const override;
+	NIAGARAEDITOR_API virtual TOptional<FText> GetDebugDrawTooltip() const override;
+	NIAGARAEDITOR_API virtual bool IsDebugDrawEnabled() const override;
+	NIAGARAEDITOR_API virtual void SetDebugDrawEnabled(bool bInEnabled) override;
 
 	NIAGARAEDITOR_API bool OpenSourceAsset() const;
 
diff --git a/Engine/Plugins/FX/Niagara/Source/NiagaraEditor/Public/ViewModels/Stack/NiagaraStackRendererItem.h b/Engine/Plugins/FX/Niagara/Source/NiagaraEditor/Public/ViewModels/Stack/NiagaraStackRendererItem.h
--- a/Engine/Plugins/FX/Niagara/Source/NiagaraEditor/Public/ViewModels/Stack/NiagaraStackRendererItem.h
+++ b/Engine/Plugins/FX/Niagara/Source/NiagaraEditor/Public/ViewModels/Stack/NiagaraStackRendererItem.h
@@ -97,6 +97,11 @@ private:
 	NIAGARAEDITOR_API void RefreshIssues(TArray<FStackIssue>& NewIssues);
 	NIAGARAEDITOR_API void ProcessRendererIssues(const TArray<FNiagaraRendererFeedback>& InIssues, EStackIssueSeverity Severity, TArray<FStackIssue>& OutIssues);
 
+	virtual bool SupportsDebugDraw() const override;
+	virtual TOptional<FText> GetDebugDrawTooltip() const override;
+	virtual bool IsDebugDrawEnabled() const override;
+	virtual void SetDebugDrawEnabled(bool bInEnabled) override;
+
 private:
 	TSharedPtr<INiagaraStackRenderersOwner> RenderersOwner;
 
diff --git a/Engine/Plugins/FX/Niagara/Source/NiagaraEditorWidgets/Private/SNiagaraOverviewStack.cpp b/Engine/Plugins/FX/Niagara/Source/NiagaraEditorWidgets/Private/SNiagaraOverviewStack.cpp
--- a/Engine/Plugins/FX/Niagara/Source/NiagaraEditorWidgets/Private/SNiagaraOverviewStack.cpp
+++ b/Engine/Plugins/FX/Niagara/Source/NiagaraEditorWidgets/Private/SNiagaraOverviewStack.cpp
@@ -1013,58 +1013,46 @@ TSharedRef<ITableRow> SNiagaraOverviewStack::OnGenerateRowForEntry(UNiagaraStack
 			];
 		}
 		
+
 		UNiagaraStackModuleItem* StackModuleItem = Cast<UNiagaraStackModuleItem>(StackItem);
-		
-		if (StackModuleItem)
+		// Scratch icon
+		if (StackModuleItem && StackModuleItem->IsScratchModule())
 		{			
-			// Scratch icon
-			if(StackModuleItem->IsScratchModule())
-			{
-				OptionsBox->AddSlot()
-					.VAlign(VAlign_Center)
-					.AutoWidth()
-					.Padding(3, 0, 3, 1)
-					[					
-						SNew(SImage)
-						.ToolTipText(LOCTEXT("ScratchPadOverviewTooltip", "This module is a scratch pad script."))
-						.Image(FNiagaraEditorStyle::Get().GetBrush("Tab.ScratchPad"))
-						.ColorAndOpacity(FSlateColor::UseSubduedForeground())
-						.DesiredSizeOverride(FVector2D(14.0f, 14.0f))
-					];
-			}
+			OptionsBox->AddSlot()
+				.VAlign(VAlign_Center)
+				.AutoWidth()
+				.Padding(3, 0, 3, 1)
+				[					
+					SNew(SImage)
+					.ToolTipText(LOCTEXT("ScratchPadOverviewTooltip", "This module is a scratch pad script."))
+					.Image(FNiagaraEditorStyle::Get().GetBrush("Tab.ScratchPad"))
+					.ColorAndOpacity(FSlateColor::UseSubduedForeground())
+					.DesiredSizeOverride(FVector2D(14.0f, 14.0f))
+				];
+		}
 
-			// Debug draw 
-			if(StackModuleItem->GetModuleNode().ContainsDebugSwitch())
-			{
-				FText TooltipText = FNiagaraEditorSharedTexts::DebugDrawUIActionBaseText;
-				
-				if(FVersionedNiagaraScriptData* ScriptData = StackModuleItem->GetModuleNode().GetScriptData())
-				{
-					if(ScriptData->DebugDrawMessage.IsEmpty() == false)
-					{
-						TooltipText = ScriptData->DebugDrawMessage;
-					}
-				}
-				
-				OptionsBox->AddSlot()
+		// Debug draw 
+		if (StackItem->SupportsDebugDraw())
+		{
+			TOptional<FText> TooltipText = StackItem->GetDebugDrawTooltip();
+			OptionsBox->AddSlot()
+				.VAlign(VAlign_Center)
+				.AutoWidth()
+				.Padding(3, 0)
+				[
+					SNew(SButton)
+					.HAlign(HAlign_Center)
 					.VAlign(VAlign_Center)
-					.AutoWidth()
-					.Padding(3, 0)
+					.ForegroundColor(FLinearColor::Transparent)
+					.ToolTipText(TooltipText.Get(FNiagaraEditorSharedTexts::DebugDrawUIActionBaseText))
+					.OnClicked(this, &SNiagaraOverviewStack::ToggleModuleDebugDraw, StackItem)
+					.ContentPadding(FMargin(1.0f))
+					.ButtonStyle(FNiagaraEditorWidgetsStyle::Get(), "NiagaraEditor.Stack.SimpleButton")
 					[
-						SNew(SButton)
-						.HAlign(HAlign_Center)
-						.VAlign(VAlign_Center)
-						.ForegroundColor(FLinearColor::Transparent)
-						.ToolTipText(TooltipText)
-						.OnClicked(this, &SNiagaraOverviewStack::ToggleModuleDebugDraw, StackItem)
-						.ContentPadding(FMargin(1.0f))
-						.ButtonStyle(FNiagaraEditorWidgetsStyle::Get(), "NiagaraEditor.Stack.SimpleButton")
-						[
-							SNew(SImage)
-							.Image(this, &SNiagaraOverviewStack::GetDebugIconBrush, StackItem)
-						]
-					];
-			}
+						SNew(SImage)
+						.Image(this, &SNiagaraOverviewStack::GetDebugIconBrush, StackItem)
+					]
+				];
 		}
 
 		// Inheritance Icon
@@ -1399,8 +1387,7 @@ EVisibility SNiagaraOverviewStack::GetShouldDebugDrawStatusVisibility(UNiagaraSt
 
 bool SNiagaraOverviewStack::IsModuleDebugDrawEnabled(UNiagaraStackItem* Item) const
 {
-	UNiagaraStackModuleItem* ModuleItem = Cast<UNiagaraStackModuleItem>(Item);
-	return ModuleItem && ModuleItem->GetIsEnabled() && ModuleItem->IsDebugDrawEnabled();
+	return Item->GetIsEnabled() && Item->IsDebugDrawEnabled();
 }
 
 const FSlateBrush* SNiagaraOverviewStack::GetDebugIconBrush(UNiagaraStackItem* Item) const
@@ -1412,10 +1399,9 @@ const FSlateBrush* SNiagaraOverviewStack::GetDebugIconBrush(UNiagaraStackItem* I
 
 FReply SNiagaraOverviewStack::ToggleModuleDebugDraw(UNiagaraStackItem* Item)
 {
-	UNiagaraStackModuleItem* ModuleItem = Cast<UNiagaraStackModuleItem>(Item);
-	if (ModuleItem != nullptr)
+	if (Item->SupportsDebugDraw())
 	{
-		ModuleItem->SetDebugDrawEnabled(!ModuleItem->IsDebugDrawEnabled());
+		Item->SetDebugDrawEnabled(!Item->IsDebugDrawEnabled());
 	}
 
 	return FReply::Handled();
