- Add missing null check

#rnx

diff --git a/Engine/Plugins/FX/Niagara/Source/NiagaraEditor/Private/TraversalCache/TraversalCache.cpp b/Engine/Plugins/FX/Niagara/Source/NiagaraEditor/Private/TraversalCache/TraversalCache.cpp
--- a/Engine/Plugins/FX/Niagara/Source/NiagaraEditor/Private/TraversalCache/TraversalCache.cpp
+++ b/Engine/Plugins/FX/Niagara/Source/NiagaraEditor/Private/TraversalCache/TraversalCache.cpp
@@ -28,7 +28,10 @@ void FTraversalCache::Initialize()
 		checkf(Instance == nullptr, TEXT("FTraversalCache initialized more than once."));
 		Instance = this;
 		FNiagaraEditorModule::Get().OnScriptApplied().AddSP(this, &FTraversalCache::OnScriptApplied);
-		GEditor->RegisterForUndo(this);
+		if (GEditor != nullptr)
+		{
+			GEditor->RegisterForUndo(this);
+		}
 		bInitialized = true;
 	}
 }
