Fixing bool writes coming from the GPU -> CPU.

diff --git a/Engine/Plugins/FX/Niagara/Shaders/Private/DataChannel/NiagaraDataInterfaceDataChannelCommon.ush b/Engine/Plugins/FX/Niagara/Shaders/Private/DataChannel/NiagaraDataInterfaceDataChannelCommon.ush
--- a/Engine/Plugins/FX/Niagara/Shaders/Private/DataChannel/NiagaraDataInterfaceDataChannelCommon.ush
+++ b/Engine/Plugins/FX/Niagara/Shaders/Private/DataChannel/NiagaraDataInterfaceDataChannelCommon.ush
@@ -82,6 +82,20 @@ void NiagaraDataChannelRead_Int32(Buffer<int> InBuffer, uint Stride, in out uint
 	NiagaraDataChannelRead_Int32(InBuffer, Stride, RegisterIdx, ParticleIdx, OutValue.w);
 }
 
+void NiagaraDataChannelRead_Bool(Buffer<int> InBuffer, uint Stride, in out uint RegisterIdx, uint ParticleIdx, out bool OutValue)
+{
+	if(RegisterIdx != -1)
+	{
+		int IntVal = InBuffer[RegisterIdx * Stride + ParticleIdx];
+		OutValue = IntVal == 0 ? false : true;
+		++RegisterIdx;
+	}
+	else
+	{
+		OutValue = false;
+	}
+}
+
 /*
 //TODO: Half Support | 
 void NiagaraDataChannelRead_Half(Buffer<half> InBuffer, uint Stride, in out uint RegisterIdx, uint ParticleIdx, out half OutValue)
@@ -181,6 +195,15 @@ void NiagaraDataChannelWrite_Int32(RWBuffer<int> InBuffer, uint Stride, in out u
 	NiagaraDataChannelWrite_Int32(InBuffer, Stride, RegisterIdx, ParticleIdx, InValue.w);
 }
 
+void NiagaraDataChannelWrite_Bool(RWBuffer<int> InBuffer, uint Stride, in out uint RegisterIdx, uint ParticleIdx, bool InValue)
+{
+	if(RegisterIdx != -1)
+	{
+		InBuffer[RegisterIdx * Stride + ParticleIdx] = InValue ? 0xFFFFFFFF : 0;
+		++RegisterIdx;
+	}
+}
+
 /*
 //TODO: Half Support | 
 void NiagaraDataChannelWrite_Half(RWBuffer<half> InBuffer, uint Stride, in out uint RegisterIdx, uint ParticleIdx, half InValue)
diff --git a/Engine/Plugins/FX/Niagara/Shaders/Private/DataChannel/NiagaraDataInterfaceDataChannelTemplateReadCommon.ush b/Engine/Plugins/FX/Niagara/Shaders/Private/DataChannel/NiagaraDataInterfaceDataChannelTemplateReadCommon.ush
--- a/Engine/Plugins/FX/Niagara/Shaders/Private/DataChannel/NiagaraDataInterfaceDataChannelTemplateReadCommon.ush
+++ b/Engine/Plugins/FX/Niagara/Shaders/Private/DataChannel/NiagaraDataInterfaceDataChannelTemplateReadCommon.ush
@@ -122,6 +122,8 @@ struct FNDCAccessContext_{ParameterName}
 	void Read_Int32(out int3 Value){ NiagaraDataChannelRead_Int32(GetBuffer_Int32(), GetIntStride(), Register_Int32, ReadIndex, Value); }
 	void Read_Int32(out int4 Value){ NiagaraDataChannelRead_Int32(GetBuffer_Int32(), GetIntStride(), Register_Int32, ReadIndex, Value); }
 
+	void Read_Bool(out bool Value){ NiagaraDataChannelRead_Bool(GetBuffer_Int32(), GetIntStride(), Register_Int32, ReadIndex, Value); }
+
 	//TODO: Half Support | void Read_Half(out half Value){ NiagaraDataChannelRead_Half(GetBuffer_Half(), GetHalfStride(), Register_Half, ReadIndex, Value); }
 	//TODO: Half Support | void Read_Half(out half2 Value){ NiagaraDataChannelRead_Half(GetBuffer_Half(), GetHalfStride(), Register_Half, ReadIndex, Value); }
 	//TODO: Half Support | void Read_Half(out half3 Value){ NiagaraDataChannelRead_Half(GetBuffer_Half(), GetHalfStride(), Register_Half, ReadIndex, Value); }
diff --git a/Engine/Plugins/FX/Niagara/Shaders/Private/DataChannel/NiagaraDataInterfaceDataChannelTemplateWriteCommon.ush b/Engine/Plugins/FX/Niagara/Shaders/Private/DataChannel/NiagaraDataInterfaceDataChannelTemplateWriteCommon.ush
--- a/Engine/Plugins/FX/Niagara/Shaders/Private/DataChannel/NiagaraDataInterfaceDataChannelTemplateWriteCommon.ush
+++ b/Engine/Plugins/FX/Niagara/Shaders/Private/DataChannel/NiagaraDataInterfaceDataChannelTemplateWriteCommon.ush
@@ -156,6 +156,8 @@ struct FNDCAccessContext_{ParameterName}
 	void Write_Int32(int3 Value){ NiagaraDataChannelWrite_Int32(GetCurrBuffer_Int32(), GetInt32Stride(), CurrRegister_Int32, WriteIndex, Value);	}
 	void Write_Int32(int4 Value){ NiagaraDataChannelWrite_Int32(GetCurrBuffer_Int32(), GetInt32Stride(), CurrRegister_Int32, WriteIndex, Value);	}
 
+	void Write_Bool(bool Value){ NiagaraDataChannelWrite_Bool(GetCurrBuffer_Int32(), GetInt32Stride(), CurrRegister_Int32, WriteIndex, Value);	}
+
 	//TODO: Half Support | void Write_Half(half Value){ NiagaraDataChannelWrite_Half(GetCurrBuffer_Half(), GetHalfStride(), CurrRegister_Half, WriteIndex, Value); }
 	//TODO: Half Support | void Write_Half(half2 Value){ NiagaraDataChannelWrite_Half(GetCurrBuffer_Half(), GetHalfStride(), CurrRegister_Half, WriteIndex, Value); }
 	//TODO: Half Support | void Write_Half(half3 Value){ NiagaraDataChannelWrite_Half(GetCurrBuffer_Half(), GetHalfStride(), CurrRegister_Half, WriteIndex, Value); }
diff --git a/Engine/Plugins/FX/Niagara/Source/Niagara/Private/DataInterface/NiagaraDataInterfaceDataChannelCommon.cpp b/Engine/Plugins/FX/Niagara/Source/Niagara/Private/DataInterface/NiagaraDataInterfaceDataChannelCommon.cpp
--- a/Engine/Plugins/FX/Niagara/Source/Niagara/Private/DataInterface/NiagaraDataInterfaceDataChannelCommon.cpp
+++ b/Engine/Plugins/FX/Niagara/Source/Niagara/Private/DataInterface/NiagaraDataInterfaceDataChannelCommon.cpp
@@ -482,7 +482,7 @@ void Write_{FunctionParameterName}_{ParameterName}(FNDCAccessContext_{ParameterN
 			}
 			else if (Struct == FNiagaraTypeDefinition::GetBoolStruct())
 			{
-				FunctionParameterComponentBufferType = TEXT("Int32");
+				FunctionParameterComponentBufferType = TEXT("Bool");
 				FunctionParameterComponentType = HlslGenContext.GetStructHlslTypeName(FNiagaraTypeDefinition::GetIntDef());
 				OutCode += FString::Format(bRead ? *ReadDataTemplate : *WriteDataTemplate, HlslTemplateArgs);
 
