[Niagara] Initialize PrimitiveId in GetMaterialVertexParameters() function.

This fixes an index-out-of-bounds error when a custom node causes the material translator to pull in GetPerInstanceCustomData().
This function uses PrimitiveId which is initialized with INVALID_PRIMITIVE_ID and not written again before being read in LoadInstanceDataElementUBO() which then results in an out of bounds error.

#rnx

diff --git a/Engine/Plugins/FX/Niagara/Shaders/Private/NiagaraMeshVertexFactory.ush b/Engine/Plugins/FX/Niagara/Shaders/Private/NiagaraMeshVertexFactory.ush
--- a/Engine/Plugins/FX/Niagara/Shaders/Private/NiagaraMeshVertexFactory.ush
+++ b/Engine/Plugins/FX/Niagara/Shaders/Private/NiagaraMeshVertexFactory.ush
@@ -378,6 +378,10 @@ FMaterialVertexParameters GetMaterialVertexParameters(
 	FMaterialVertexParameters Result = MakeInitializedMaterialVertexParameters();
 	Result.SceneData = Intermediates.SceneData;
 	
+#if VF_USE_PRIMITIVE_SCENE_DATA
+	Result.PrimitiveId = Intermediates.SceneData.PrimitiveId;
+#endif
+
 	Result.WorldPosition = WorldPosition;
 	if (bIsPreviousFrame)
 	{
