Clean up any MIC that is found which would result in a reallocation during UpdateMIC.  Also expands on the logging for this scenario to maybe try to shed some light on where they could be coming from (or rather how we don't know about them/have references to them)

diff --git a/Engine/Plugins/FX/Niagara/Source/Niagara/Private/NiagaraRendererProperties.cpp b/Engine/Plugins/FX/Niagara/Source/Niagara/Private/NiagaraRendererProperties.cpp
--- a/Engine/Plugins/FX/Niagara/Source/Niagara/Private/NiagaraRendererProperties.cpp
+++ b/Engine/Plugins/FX/Niagara/Source/Niagara/Private/NiagaraRendererProperties.cpp
@@ -644,9 +644,19 @@ void UNiagaraRendererProperties::UpdateMaterialParametersMIC(const FNiagaraRende
 			}
 			else
 			{
-				if (const UObject* ExistingObject = StaticFindObject(UMaterialInstanceConstant::StaticClass(), this, *NameBuilder))
+				if (UObject* ExistingObject = StaticFindObject(UMaterialInstanceConstant::StaticClass(), this, *NameBuilder))
 				{
-					UE_LOG(LogNiagara, Warning, TEXT("While trying to allocate %s UNiagaraRendererProperties::UpdateMaterialParametersMIC() found a pre-existing object.  This will result in a re-allocation!"), *NameBuilder);
+					const FString FullName = ExistingObject->GetFullName();
+					const EInternalObjectFlags InternalFlags = ExistingObject->GetInternalFlags();
+					const EObjectFlags Flags = ExistingObject->GetFlags();
+
+					UE_LOG(LogNiagara, Warning, TEXT("While trying to allocate %s UNiagaraRendererProperties::UpdateMaterialParametersMIC() found a pre-existing object.  This will result in a re-allocation!  FullName: %s | InternalFlags: %x | Flags: %x"),
+						*NameBuilder, *FullName, InternalFlags, Flags);
+
+					if (UMaterialInstanceConstant* ExistingMIC = Cast<UMaterialInstanceConstant>(ExistingObject))
+					{
+						NiagaraRendererPropertiesPrivate::MarkAndRenameMaterialForGarbage(ExistingMIC);
+					}
 				}
 
 				MIC = NewObject<UMaterialInstanceConstant>(this, FName(NameBuilder));
