Niagara Movie Preview Assets
- Added new plugin to contain the movie files
- Scripts, Emitters and Systems can link to a movie file via soft object path that can be used for preview purposes
- Implemented movie tooltip for scripts

diff --git a/Engine/Plugins/FX/Niagara/Config/BaseNiagara.ini b/Engine/Plugins/FX/Niagara/Config/BaseNiagara.ini
--- a/Engine/Plugins/FX/Niagara/Config/BaseNiagara.ini
+++ b/Engine/Plugins/FX/Niagara/Config/BaseNiagara.ini
@@ -123,3 +123,5 @@ bShowGridInViewport=True
 bShowInstructionsCount=False
 bShowParticleCountsInViewport=False
 bShowEmitterExecutionOrder=False
+bDisplayPreviewMoviesInTooltips=True
+TooltipPreviewMovieSize=(X=128.000000,Y=128.000000)
diff --git a/Engine/Plugins/FX/Niagara/Source/Niagara/Classes/NiagaraEmitter.h b/Engine/Plugins/FX/Niagara/Source/Niagara/Classes/NiagaraEmitter.h
--- a/Engine/Plugins/FX/Niagara/Source/Niagara/Classes/NiagaraEmitter.h
+++ b/Engine/Plugins/FX/Niagara/Source/Niagara/Classes/NiagaraEmitter.h
@@ -743,6 +743,9 @@ public:
 	UPROPERTY(EditAnywhere, AdvancedDisplay, Category = "Asset Options", AssetRegistrySearchable, DisplayName="Asset Description")
 	FText TemplateAssetDescription;
 
+	UPROPERTY(EditAnywhere, Category="Preview", AssetRegistrySearchable, meta = (SkipSystemResetOnChange = "true", AllowedClasses="/Script/MediaAssets.FileMediaSource"))
+	FSoftObjectPath PreviewMoviePath;
+	
 	/** Category to collate this emitter into for "add new emitter" dialogs.*/
 	UPROPERTY(AssetRegistrySearchable)
 	FText Category;
diff --git a/Engine/Plugins/FX/Niagara/Source/Niagara/Classes/NiagaraScript.h b/Engine/Plugins/FX/Niagara/Source/Niagara/Classes/NiagaraScript.h
--- a/Engine/Plugins/FX/Niagara/Source/Niagara/Classes/NiagaraScript.h
+++ b/Engine/Plugins/FX/Niagara/Source/Niagara/Classes/NiagaraScript.h
@@ -894,6 +894,9 @@ private:
 	FGuid UsageId;
 
 #if WITH_EDITORONLY_DATA
+	UPROPERTY(EditAnywhere, Category="Preview", AssetRegistrySearchable, meta = (AllowedClasses="/Script/MediaAssets.FileMediaSource"))
+	FSoftObjectPath PreviewMoviePath;
+	
 	/** The exposed version is the version that is used by default when a user adds this script somewhere. It is basically the published version and allows a script maintainer to create and test newer versions that are not used by normal users. */
 	UPROPERTY()
 	FGuid ExposedVersion;
diff --git a/Engine/Plugins/FX/Niagara/Source/Niagara/Classes/NiagaraSystem.h b/Engine/Plugins/FX/Niagara/Source/Niagara/Classes/NiagaraSystem.h
--- a/Engine/Plugins/FX/Niagara/Source/Niagara/Classes/NiagaraSystem.h
+++ b/Engine/Plugins/FX/Niagara/Source/Niagara/Classes/NiagaraSystem.h
@@ -503,6 +503,9 @@ public:
 	UPROPERTY(meta = (SkipSystemResetOnChange = "true"))
 	TArray<FNiagaraAssetTagDefinitionReference> AssetTags;
 
+	UPROPERTY(EditAnywhere, Category="Preview", AssetRegistrySearchable, meta = (SkipSystemResetOnChange = "true", AllowedClasses="/Script/MediaAssets.FileMediaSource"))
+	FSoftObjectPath PreviewMoviePath;
+	
 	UPROPERTY()
 	TArray<TObjectPtr<UNiagaraScript>> ScratchPadScripts;
 
diff --git a/Engine/Plugins/FX/Niagara/Source/NiagaraEditor/NiagaraEditor.Build.cs b/Engine/Plugins/FX/Niagara/Source/NiagaraEditor/NiagaraEditor.Build.cs
--- a/Engine/Plugins/FX/Niagara/Source/NiagaraEditor/NiagaraEditor.Build.cs
+++ b/Engine/Plugins/FX/Niagara/Source/NiagaraEditor/NiagaraEditor.Build.cs
@@ -89,7 +89,8 @@ public class NiagaraEditor : ModuleRules
                 "AppFramework",
 				"ToolWidgets",
 				"Sequencer",
-				"DataHierarchyEditor"
+				"DataHierarchyEditor", 
+				"MediaAssets"
             }
         );
 
diff --git a/Engine/Plugins/FX/Niagara/Source/NiagaraEditor/Private/EdGraphSchema_Niagara.cpp b/Engine/Plugins/FX/Niagara/Source/NiagaraEditor/Private/EdGraphSchema_Niagara.cpp
--- a/Engine/Plugins/FX/Niagara/Source/NiagaraEditor/Private/EdGraphSchema_Niagara.cpp
+++ b/Engine/Plugins/FX/Niagara/Source/NiagaraEditor/Private/EdGraphSchema_Niagara.cpp
@@ -534,6 +534,7 @@ TArray<TSharedPtr<FNiagaraAction_NewNode>> UEdGraphSchema_Niagara::GetGraphActio
 
 		TSharedPtr<FNiagaraAction_NewNode> Action = AddNewNodeMenuAction(NewActions, FunctionCallNode, DisplayName, Categories, FavoritesActionData, TooltipDesc, Keywords, SourceData);
 		Action->bIsInLibrary = bIsInLibrary;
+		Action->PreviewMoviePath = FNiagaraEditorUtilities::Preview::GetPreviewMovieObjectPath(ScriptAsset);
 
 		return Action;
 	};
diff --git a/Engine/Plugins/FX/Niagara/Source/NiagaraEditor/Private/NiagaraEditorModule.cpp b/Engine/Plugins/FX/Niagara/Source/NiagaraEditor/Private/NiagaraEditorModule.cpp
--- a/Engine/Plugins/FX/Niagara/Source/NiagaraEditor/Private/NiagaraEditorModule.cpp
+++ b/Engine/Plugins/FX/Niagara/Source/NiagaraEditor/Private/NiagaraEditorModule.cpp
@@ -149,6 +149,7 @@
 
 #include "AssetCompilingManager.h"
 #include "Editor.h"
+#include "FileMediaSource.h"
 #include "ILevelSequenceModule.h"
 #include "HAL/PlatformFileManager.h"
 #include "HAL/FileManager.h"
@@ -2458,6 +2459,34 @@ FNiagaraEditorModule::FOnScriptApplied& FNiagaraEditorModule::OnScriptApplied()
 	return OnScriptAppliedDelegate;
 }
 
+const FNiagaraEditorModule::FMoviePreviewObjects& FNiagaraEditorModule::FindOrAddMoviePreviewObjects(const FAssetData& AssetData)
+{
+	if(MoviePreviewObjects.Contains(AssetData))
+	{
+		return MoviePreviewObjects[AssetData];
+	}
+	
+	FString PreviewMoviePathString;
+	bool bHasMovieSequence = AssetData.GetTagValue("PreviewMoviePath", PreviewMoviePathString);
+	FSoftObjectPath PreviewMoviePath(PreviewMoviePathString);
+	UObject* MoviePreviewObject = PreviewMoviePath.TryLoad();
+	UFileMediaSource* MediaSource = Cast<UFileMediaSource>(MoviePreviewObject);
+	
+	FMoviePreviewObjects& PreviewObjects = MoviePreviewObjects.Add(AssetData);
+	PreviewObjects.MediaPlayer.Reset(NewObject<UMediaPlayer>(GetTransientPackage(), UMediaPlayer::StaticClass()));
+	PreviewObjects.MediaTexture.Reset(NewObject<UMediaTexture>(GetTransientPackage()));
+
+	PreviewObjects.MediaPlayer->SetLooping(true);
+	PreviewObjects.MediaPlayer->AffectedByPIEHandling = false;
+	PreviewObjects.MediaPlayer->OpenSource(MediaSource);
+	PreviewObjects.MediaTexture->SetMediaPlayer(PreviewObjects.MediaPlayer.Get());
+	PreviewObjects.MediaTexture->UpdateResource();
+	PreviewObjects.MediaPlayer->Play();
+	PreviewObjects.ImageBrush = new FSlateImageBrush(PreviewObjects.MediaTexture.Get(), {80.f, 80.f});
+
+	return PreviewObjects;
+}
+
 UNiagaraParameterCollection* FNiagaraEditorModule::FindCollectionForVariable(const FString& VariableName)
 {
 	auto FindCachedCollectionByPrefix = [this](const FString& Prefix)
diff --git a/Engine/Plugins/FX/Niagara/Source/NiagaraEditor/Private/NiagaraEditorUtilities.cpp b/Engine/Plugins/FX/Niagara/Source/NiagaraEditor/Private/NiagaraEditorUtilities.cpp
--- a/Engine/Plugins/FX/Niagara/Source/NiagaraEditor/Private/NiagaraEditorUtilities.cpp
+++ b/Engine/Plugins/FX/Niagara/Source/NiagaraEditor/Private/NiagaraEditorUtilities.cpp
@@ -4707,6 +4707,107 @@ EAxisList::Type FNiagaraEditorUtilities::VectorComponentToAxis(int32 NumComponen
 	return EAxisList::None;
 }
 
+TOptional<FSoftObjectPath> FNiagaraEditorUtilities::Preview::GetPreviewMovieObjectPath(const FAssetData& AssetData)
+{
+	FString PreviewMoviePathString;
+	bool bHasMovieSequence = AssetData.GetTagValue("PreviewMoviePath", PreviewMoviePathString);
+	FSoftObjectPath PreviewMoviePath(PreviewMoviePathString);
+
+	if(bHasMovieSequence && PreviewMoviePath.IsAsset())
+	{
+		return PreviewMoviePath;
+	}
+
+	return TOptional<FSoftObjectPath>();
+}
+
+TSharedRef<SToolTip> FNiagaraEditorUtilities::Tooltips::CreateTooltip(FText Description, TOptional<FSoftObjectPath> PreviewMoviePath, const FCreateTooltipArguments& CreateTooltipArguments)
+{
+	if(CreateTooltipArguments.bTryCreatePreviewMovie == true && CreateTooltipArguments.PreviewMovieOwner.IsValid() == false)
+	{
+		UE_LOGFMT(LogNiagaraEditor, Error, "Preview Movie Tooltip was requested but no owner was specified.");
+	}
+	
+	SHorizontalBox::FSlot* MoviePreviewSlot;
+	SHorizontalBox::FSlot* DescriptionSlot;
+	
+	TSharedRef<SHorizontalBox> TooltipContent = SNew(SHorizontalBox)
+		+ SHorizontalBox::Slot().Expose(MoviePreviewSlot)
+		.HAlign(HAlign_Center)
+		.VAlign(VAlign_Center)
+		.AutoWidth()
+		+ SHorizontalBox::Slot().Expose(DescriptionSlot)
+		.Padding(4.f, 2.f);
+
+	TSharedPtr<SBox> PreviewMovieContainer = nullptr;
+	if(CreateTooltipArguments.bTryCreatePreviewMovie && CreateTooltipArguments.PreviewMovieOwner.IsValid() && PreviewMoviePath.IsSet() && PreviewMoviePath->IsAsset())
+	{
+		PreviewMovieContainer = SNew(SBox)
+			.Visibility_Lambda([]()
+			{
+				return GetDefault<UNiagaraEditorSettings>()->bDisplayPreviewMoviesInTooltips ? EVisibility::Visible : EVisibility::Collapsed;
+			})
+			.WidthOverride_Lambda([]()
+			{
+				return GetDefault<UNiagaraEditorSettings>()->TooltipPreviewMovieSize.X;
+			})
+			.HeightOverride_Lambda([]()
+			{
+				return GetDefault<UNiagaraEditorSettings>()->TooltipPreviewMovieSize.Y;
+			})
+			[
+				SNew(SImage)
+				// We only have one preview movie brush & one media player per previewmovieviewmodel.
+				.Image(CreateTooltipArguments.PreviewMovieOwner->GetMoviePreviewImageBrush())
+			];
+	}
+
+	if(PreviewMovieContainer.IsValid())
+	{
+		MoviePreviewSlot->AttachWidget(PreviewMovieContainer.ToSharedRef());
+	}
+
+	TSharedRef<STextBlock> DescriptionText = SNew(STextBlock)
+		.Text(Description)
+		.WrapTextAt(SToolTip::GetToolTipWrapWidth() - 500.f)
+		.Font(FCoreStyle::Get().GetFontStyle("ToolTip.Font"))
+		.ColorAndOpacity(FLinearColor::Black);
+	
+	DescriptionSlot->AttachWidget(DescriptionText);
+	
+	return SNew(SToolTip)
+		.BorderImage(FCoreStyle::Get().GetBrush("ToolTip.BrightBackground"))
+		[
+			TooltipContent
+		];
+}
+
+TSharedRef<SToolTip> FNiagaraEditorUtilities::Tooltips::CreateTooltipForNiagaraAction(TSharedRef<FNiagaraMenuAction_Generic> Action, const FCreateTooltipArguments& CreateTooltipArguments)
+{	
+	return CreateTooltip(Action->ToolTip, Action->PreviewMoviePath, CreateTooltipArguments);
+}
+
+TSharedRef<SToolTip> FNiagaraEditorUtilities::Tooltips::CreateTooltipForScript(UNiagaraScript& Script, FGuid VersionGuid, const FCreateTooltipArguments& CreateTooltipArguments)
+{
+	FVersionedNiagaraScriptData* ScriptData = nullptr;
+	if(VersionGuid.IsValid())
+	{
+		ScriptData = Script.GetScriptData(VersionGuid);
+	}
+	else
+	{
+		ScriptData = Script.GetLatestScriptData(); 
+	}
+
+	TOptional<FSoftObjectPath> PreviewMoviePath;
+	if(CreateTooltipArguments.bTryCreatePreviewMovie)
+	{
+		PreviewMoviePath = FNiagaraEditorUtilities::Preview::GetPreviewMovieObjectPath(FAssetData(&Script));
+	}
+
+	return CreateTooltip(ScriptData->Description, PreviewMoviePath, CreateTooltipArguments);
+}
+
 FNiagaraVariable FNiagaraEditorUtilities::UserParameters::DuplicateUserParameter(FNiagaraVariable ParameterToDuplicate, UNiagaraSystem& System)
 {
 	FScopedTransaction AddTransaction(LOCTEXT("DuplicateUserParameter", "Duplicate User Parameter"));
diff --git a/Engine/Plugins/FX/Niagara/Source/NiagaraEditor/Private/ViewModels/NiagaraScriptGraphViewModel.cpp b/Engine/Plugins/FX/Niagara/Source/NiagaraEditor/Private/ViewModels/NiagaraScriptGraphViewModel.cpp
--- a/Engine/Plugins/FX/Niagara/Source/NiagaraEditor/Private/ViewModels/NiagaraScriptGraphViewModel.cpp
+++ b/Engine/Plugins/FX/Niagara/Source/NiagaraEditor/Private/ViewModels/NiagaraScriptGraphViewModel.cpp
@@ -31,6 +31,7 @@ FNiagaraScriptGraphViewModel::FNiagaraScriptGraphViewModel(TAttribute<FText> InD
 	: DisplayName(InDisplayName)
 	, Commands(MakeShareable(new FUICommandList()))
 	, NodeSelection(MakeShareable(new FNiagaraObjectSelection()))
+	, PreviewMovieViewModel(MakeShared<FNiagaraPreviewMovieViewModel>())
 	, bIsForDataProcessingOnly(bInIsForDataProcessingOnly)
 {
 	if (bIsForDataProcessingOnly == false)
@@ -40,6 +41,8 @@ FNiagaraScriptGraphViewModel::FNiagaraScriptGraphViewModel(TAttribute<FText> InD
 	}
 
 	ErrorColor = FAppStyle::GetColor("ErrorReporting.BackgroundColor");
+
+	PreviewMovieViewModel->Initialize();
 }
 
 FNiagaraScriptGraphViewModel::~FNiagaraScriptGraphViewModel()
@@ -126,6 +129,11 @@ void FNiagaraScriptGraphViewModel::SetErrorTextToolTip(FString ErrorMsgToolTip)
 	ErrorMsg = ErrorMsgToolTip;
 }
 
+TSharedPtr<FNiagaraPreviewMovieViewModel> FNiagaraScriptGraphViewModel::GetPreviewMovieViewModel()
+{
+	return PreviewMovieViewModel;
+}
+
 void FNiagaraScriptGraphViewModel::PostUndo(bool bSuccess)
 {
 	// The graph may have been deleted as a result of the undo operation so make sure it's valid
diff --git a/Engine/Plugins/FX/Niagara/Source/NiagaraEditor/Private/ViewModels/NiagaraScriptGraphViewModel.h b/Engine/Plugins/FX/Niagara/Source/NiagaraEditor/Private/ViewModels/NiagaraScriptGraphViewModel.h
--- a/Engine/Plugins/FX/Niagara/Source/NiagaraEditor/Private/ViewModels/NiagaraScriptGraphViewModel.h
+++ b/Engine/Plugins/FX/Niagara/Source/NiagaraEditor/Private/ViewModels/NiagaraScriptGraphViewModel.h
@@ -67,6 +67,8 @@ public:
 
 	void SetErrorTextToolTip(FString ErrorMsgToolTip);
 
+	TSharedPtr<FNiagaraPreviewMovieViewModel> GetPreviewMovieViewModel();
+
 	//~ FEditorUndoClient interface.
 	virtual void PostUndo(bool bSuccess) override;
 	virtual void PostRedo(bool bSuccess) override { PostUndo(bSuccess); }
@@ -99,7 +101,7 @@ private:
 
 	/** The set of nodes objects currently selected in the graph. */
 	TSharedRef<FNiagaraObjectSelection> NodeSelection;
-
+	
 	/** A multicast delegate which is called whenever nodes are pasted into the graph. */
 	FOnNodesPasted OnNodesPastedDelegate;
 
@@ -112,6 +114,8 @@ private:
 	/** Used to set the error color */
 	FSlateColor ErrorColor;
 
+	TSharedPtr<FNiagaraPreviewMovieViewModel> PreviewMovieViewModel;
+
 	/** Whether or not this view model is going to be used for data processing only and will not be shown in the UI. */
 	bool bIsForDataProcessingOnly;
 };
diff --git a/Engine/Plugins/FX/Niagara/Source/NiagaraEditor/Private/ViewModels/NiagaraSystemViewModel.cpp b/Engine/Plugins/FX/Niagara/Source/NiagaraEditor/Private/ViewModels/NiagaraSystemViewModel.cpp
--- a/Engine/Plugins/FX/Niagara/Source/NiagaraEditor/Private/ViewModels/NiagaraSystemViewModel.cpp
+++ b/Engine/Plugins/FX/Niagara/Source/NiagaraEditor/Private/ViewModels/NiagaraSystemViewModel.cpp
@@ -163,6 +163,9 @@ void FNiagaraSystemViewModel::Initialize(UNiagaraSystem& InSystem, FNiagaraSyste
 	SystemGraphSelectionViewModel = MakeShared<FNiagaraSystemGraphSelectionViewModel>();
 	SystemGraphSelectionViewModel->Initialize(this->AsShared());
 
+	PreviewMovieViewModel = MakeShared<FNiagaraPreviewMovieViewModel>();
+	PreviewMovieViewModel->Initialize();
+	
 	bSupportCompileForEdit = !bIsForDataProcessingOnly && (EditMode == ENiagaraSystemViewModelEditMode::SystemAsset || EditMode == ENiagaraSystemViewModelEditMode::EmitterAsset);
 	bIsCompileForEdit = false;
 	if (bSupportCompileForEdit)
diff --git a/Engine/Plugins/FX/Niagara/Source/NiagaraEditor/Private/ViewModels/Stack/NiagaraStackEntry.cpp b/Engine/Plugins/FX/Niagara/Source/NiagaraEditor/Private/ViewModels/Stack/NiagaraStackEntry.cpp
--- a/Engine/Plugins/FX/Niagara/Source/NiagaraEditor/Private/ViewModels/Stack/NiagaraStackEntry.cpp
+++ b/Engine/Plugins/FX/Niagara/Source/NiagaraEditor/Private/ViewModels/Stack/NiagaraStackEntry.cpp
@@ -265,6 +265,16 @@ FText UNiagaraStackEntry::GetTooltipText() const
 	return FText();
 }
 
+bool UNiagaraStackEntry::UsesCustomTooltipWidget() const
+{
+	return false;
+}
+
+TSharedPtr<IToolTip> UNiagaraStackEntry::GetCustomTooltipWidget() const
+{
+	return nullptr;
+}
+
 bool UNiagaraStackEntry::GetCanExpand() const
 {
 	return true;
diff --git a/Engine/Plugins/FX/Niagara/Source/NiagaraEditor/Private/ViewModels/Stack/NiagaraStackModuleItem.cpp b/Engine/Plugins/FX/Niagara/Source/NiagaraEditor/Private/ViewModels/Stack/NiagaraStackModuleItem.cpp
--- a/Engine/Plugins/FX/Niagara/Source/NiagaraEditor/Private/ViewModels/Stack/NiagaraStackModuleItem.cpp
+++ b/Engine/Plugins/FX/Niagara/Source/NiagaraEditor/Private/ViewModels/Stack/NiagaraStackModuleItem.cpp
@@ -38,6 +38,7 @@
 #include "ViewModels/Stack/NiagaraStackModuleItemOutputCollection.h"
 #include "NiagaraEditorModule.h"
 #include "Toolkits/NiagaraSystemToolkit.h"
+#include "Widgets/SToolTip.h"
 
 // TODO: Remove these
 #include "ViewModels/Stack/NiagaraStackErrorItem.h"
@@ -119,6 +120,36 @@ FText UNiagaraStackModuleItem::GetTooltipText() const
 	}
 }
 
+bool UNiagaraStackModuleItem::UsesCustomTooltipWidget() const
+{
+	return GetDefault<UNiagaraEditorSettings>()->bDisplayPreviewMoviesInTooltips && FunctionCallNode != nullptr && FunctionCallNode->FunctionScript != nullptr;
+}
+
+TSharedPtr<IToolTip> UNiagaraStackModuleItem::GetCustomTooltipWidget() const
+{
+	using namespace FNiagaraEditorUtilities::Tooltips;
+	if(FunctionCallNode != nullptr && FunctionCallNode->FunctionScript != nullptr)
+	{
+		FAssetData AssetData(FunctionCallNode->FunctionScript);
+		if(GetDefault<UNiagaraEditorSettings>()->bDisplayPreviewMoviesInTooltips && GetSystemViewModel()->GetPreviewMovieViewModel()->IsPlayingPreviewMovieForAsset(AssetData) == false)
+		{
+			GetSystemViewModel()->GetPreviewMovieViewModel()->PlayMovieForAsset(AssetData);
+		}
+		
+		if(TooltipCache.IsValid() == false)
+		{
+			FCreateTooltipArguments CreateTooltipArguments;
+			CreateTooltipArguments.bTryCreatePreviewMovie = GetDefault<UNiagaraEditorSettings>()->bDisplayPreviewMoviesInTooltips;
+			CreateTooltipArguments.PreviewMovieOwner = GetSystemViewModel()->GetPreviewMovieViewModel();
+			TooltipCache = CreateTooltipForScript(*FunctionCallNode->FunctionScript, FunctionCallNode->SelectedScriptVersion, CreateTooltipArguments);
+		}		
+		
+		return TooltipCache.ToSharedRef();
+	}
+	
+	return nullptr; 
+}
+
 INiagaraStackItemGroupAddUtilities* UNiagaraStackModuleItem::GetGroupAddUtilities()
 {
 	return GroupAddUtilities;
@@ -140,6 +171,7 @@ void UNiagaraStackModuleItem::RefreshChildrenInternal(const TArray<UNiagaraStack
 	bCanMoveAndDeleteCache.Reset();
 	bIsScratchModuleCache.Reset();
 	DisplayNameCache.Reset();
+	TooltipCache.Reset();
 
 	if (FunctionCallNode != nullptr &&
 		(FunctionCallNode->HasValidScriptAndGraph() || FunctionCallNode->Signature.IsValid()) &&
diff --git a/Engine/Plugins/FX/Niagara/Source/NiagaraEditor/Private/ViewModels/Stack/NiagaraStackRenderItemGroup.cpp b/Engine/Plugins/FX/Niagara/Source/NiagaraEditor/Private/ViewModels/Stack/NiagaraStackRenderItemGroup.cpp
--- a/Engine/Plugins/FX/Niagara/Source/NiagaraEditor/Private/ViewModels/Stack/NiagaraStackRenderItemGroup.cpp
+++ b/Engine/Plugins/FX/Niagara/Source/NiagaraEditor/Private/ViewModels/Stack/NiagaraStackRenderItemGroup.cpp
@@ -68,7 +68,12 @@ public:
 	{
 		return FNiagaraActionSourceData();
 	}
-	
+
+	virtual TOptional<FSoftObjectPath> GetPreviewMoviePath() const override
+	{
+		return {};
+	}
+
 	const FNiagaraRendererCreationInfo& GetRendererCreationInfo() { return RendererCreationInfo; }
 
 private:
diff --git a/Engine/Plugins/FX/Niagara/Source/NiagaraEditor/Private/ViewModels/Stack/NiagaraStackScriptItemGroup.cpp b/Engine/Plugins/FX/Niagara/Source/NiagaraEditor/Private/ViewModels/Stack/NiagaraStackScriptItemGroup.cpp
--- a/Engine/Plugins/FX/Niagara/Source/NiagaraEditor/Private/ViewModels/Stack/NiagaraStackScriptItemGroup.cpp
+++ b/Engine/Plugins/FX/Niagara/Source/NiagaraEditor/Private/ViewModels/Stack/NiagaraStackScriptItemGroup.cpp
@@ -48,35 +48,68 @@ const FName INiagaraStackItemGroupAddAction::FavoriteActionsProfile = FName("Sta
 class FScriptGroupAddAction : public INiagaraStackItemGroupAddAction
 {
 public:
-	static TSharedRef<FScriptGroupAddAction> CreateAssetModuleAction(FAssetData AssetData)
+	enum class ESpecificActionType
+	{
+		NewScratchModule,
+		NewSetSpecificParameters,
+		ModuleWizard,
+		Undefined
+	};
+	
+	struct FCreationArguments
 	{
+		FText DisplayName;
+		TArray<FString> Categories;
+		FText Description;
+		FText Keywords;
+		TOptional<FNiagaraFavoritesActionData> FavoritesActionData;
+		TOptional<FSoftObjectPath> PreviewMoviePath;
+		bool bIsInLibrary = false;
+		FNiagaraVariable ModuleParameterVariable;
+		FAssetData ModuleAssetData;
+		UNiagaraScript* ModuleScript;
+		TSharedPtr<FModuleWizardModel> ModuleWizardModel;
+		TOptional<ESpecificActionType> SpecificActionType;
+		FNiagaraActionSourceData SourceData = FNiagaraActionSourceData(EScriptSource::Niagara, FText::FromString("Niagara"), true);
+		bool bDisplaySource = true;
+		
+		bool bRenameParameterOnAdd = false;
+		
+	};
+	static TSharedRef<FScriptGroupAddAction> CreateAssetModuleAction(FAssetData AssetData)
+	{		
 		FText Category;
 		AssetData.GetTagValue(GET_MEMBER_NAME_CHECKED(FVersionedNiagaraScriptData, Category), Category);
 		if (Category.IsEmptyOrWhitespace())
 		{
 			Category = LOCTEXT("ModuleNotCategorized", "Uncategorized Modules");
 		}
-
-		bool bIsInLibrary = FNiagaraEditorUtilities::IsScriptAssetInLibrary(AssetData);
-
-		FText DisplayName = FNiagaraEditorUtilities::FormatScriptName(AssetData.AssetName, bIsInLibrary);
-
+		
 		FText AssetDescription;
 		AssetData.GetTagValue(GET_MEMBER_NAME_CHECKED(FVersionedNiagaraScriptData, Description), AssetDescription);
-		FText Description = FNiagaraEditorUtilities::FormatScriptDescription(AssetDescription, AssetData.GetSoftObjectPath(), bIsInLibrary);
 
 		FText Keywords;
 		AssetData.GetTagValue(GET_MEMBER_NAME_CHECKED(FVersionedNiagaraScriptData, Keywords), Keywords);
 
-		bool bSuggested = AssetData.GetTagValueRef<bool>(GET_MEMBER_NAME_CHECKED(FVersionedNiagaraScriptData, bSuggested));
+		bool bIsInLibrary = FNiagaraEditorUtilities::IsScriptAssetInLibrary(AssetData);
+		TTuple<EScriptSource, FText> Source = FNiagaraEditorUtilities::GetScriptSource(AssetData);
 
+		bool bSuggested = AssetData.GetTagValueRef<bool>(GET_MEMBER_NAME_CHECKED(FVersionedNiagaraScriptData, bSuggested));
 		FNiagaraFavoritesActionData FavoritesActionData;
 		FavoritesActionData.ActionIdentifier.Names.Add(FName(AssetData.GetObjectPathString()));
 		FavoritesActionData.bFavoriteByDefault = bSuggested;
-		
-		TTuple<EScriptSource, FText> Source = FNiagaraEditorUtilities::GetScriptSource(AssetData);
-		FNiagaraActionSourceData SourceData(Source.Key, Source.Value, true);
-		return MakeShareable(new FScriptGroupAddAction(DisplayName, {Category.ToString()}, Description, Keywords, FavoritesActionData, bIsInLibrary, FNiagaraVariable(), false, AssetData, nullptr, false, false, SourceData));
+
+		FCreationArguments CreationArguments;
+		CreationArguments.ModuleAssetData = AssetData;
+		CreationArguments.DisplayName = FNiagaraEditorUtilities::FormatScriptName(AssetData.AssetName, bIsInLibrary);
+		CreationArguments.Categories = { Category.ToString() };
+		CreationArguments.Description = FNiagaraEditorUtilities::FormatScriptDescription(AssetDescription, AssetData.GetSoftObjectPath(), bIsInLibrary);
+		CreationArguments.Keywords = Keywords;
+		CreationArguments.SourceData = FNiagaraActionSourceData(Source.Key, Source.Value, true);
+		CreationArguments.FavoritesActionData = FavoritesActionData;
+		CreationArguments.bIsInLibrary = bIsInLibrary;
+		CreationArguments.PreviewMoviePath = FNiagaraEditorUtilities::Preview::GetPreviewMovieObjectPath(AssetData);
+		return MakeShareable(new FScriptGroupAddAction(CreationArguments));
 	}
 
 	static TSharedRef<FScriptGroupAddAction> CreateModuleActionFromScratchPadScript(UNiagaraScript* ModuleScript)
@@ -87,7 +120,7 @@ public:
 		{
 			Category = LOCTEXT("ModuleNotCategorized", "Uncategorized Modules");
 		}
-
+		
 		// we assume scratch pads should always be displayed, so we act like it's a library action so it's not accidentally filtered out
 		bool bIsInLibrary = true;
 		FText DisplayName = FNiagaraEditorUtilities::FormatScriptName(ModuleScript->GetFName(), bIsInLibrary);
@@ -97,10 +130,22 @@ public:
 		FNiagaraFavoritesActionData FavoritesActionData;
 		FavoritesActionData.ActionIdentifier.Names.Add(FName(GetPathNameSafe(ModuleScript)));
 		FavoritesActionData.bFavoriteByDefault = ScriptData->bSuggested;
+
+		TOptional<FSoftObjectPath> PreviewMoviePath = FNiagaraEditorUtilities::Preview::GetPreviewMovieObjectPath(FAssetData(ModuleScript));
 		
-		FAssetData Data = FAssetData(ModuleScript);
 		FNiagaraActionSourceData SourceData(EScriptSource::Niagara, FText::FromString("Scratch Pad"), true);
-		return MakeShareable(new FScriptGroupAddAction(DisplayName, {Category.ToString()}, Description, Keywords, FavoritesActionData, bIsInLibrary, FNiagaraVariable(), false, FAssetData(), ModuleScript, false, false, SourceData));
+
+		FCreationArguments CreationArguments;
+		CreationArguments.ModuleScript = ModuleScript;
+		CreationArguments.DisplayName = DisplayName;
+		CreationArguments.Categories = { Category.ToString() };
+		CreationArguments.FavoritesActionData = FavoritesActionData;
+		CreationArguments.PreviewMoviePath = PreviewMoviePath;
+		CreationArguments.SourceData = SourceData;
+		CreationArguments.bIsInLibrary = bIsInLibrary;
+		CreationArguments.Description = Description;
+		CreationArguments.Keywords = Keywords;
+		return MakeShareable(new FScriptGroupAddAction(CreationArguments));
 	}
 
 	static TSharedRef<FScriptGroupAddAction> CreateExistingParameterModuleAction(FNiagaraVariable ParameterVariable)
@@ -111,7 +156,12 @@ public:
 		FText AttributeDescription = FNiagaraConstants::GetAttributeDescription(ParameterVariable);
 		FText Description = FText::Format(LOCTEXT("ExistingParameterModuleDescriptoinFormat", "Description: Set the parameter {0}. {1}"), FText::FromName(ParameterVariable.GetName()), AttributeDescription);
 
-		return MakeShareable(new FScriptGroupAddAction(DisplayName, {Category.ToString()}, Description, FText(), {}, true, ParameterVariable, false, FAssetData(), nullptr, false, false));
+		FCreationArguments CreationArguments;
+		CreationArguments.DisplayName = DisplayName;
+		CreationArguments.Description = Description;
+		CreationArguments.Categories = { Category.ToString() };
+		CreationArguments.bIsInLibrary = true;
+		return MakeShareable(new FScriptGroupAddAction(CreationArguments));
 	}
 
 	static TSharedRef<FScriptGroupAddAction> CreateNewParameterModuleAction(FName NewParameterNamespace, FNiagaraTypeDefinition NewParameterType)
@@ -123,7 +173,14 @@ public:
 		FNiagaraParameterHandle NewParameterHandle(NewParameterNamespace, *(TEXT("New") + NewParameterType.GetName()));
 		FNiagaraVariable NewParameter(NewParameterType, NewParameterHandle.GetParameterHandleString());
 
-		return MakeShareable(new FScriptGroupAddAction( DisplayName, {Category.ToString()}, Description, FText(), {}, true, NewParameter, true, FAssetData(), nullptr, false, false));
+		FCreationArguments CreationArguments;
+		CreationArguments.DisplayName = DisplayName;
+		CreationArguments.Description = Description;
+		CreationArguments.Categories = { Category.ToString() };
+		CreationArguments.bIsInLibrary = true;
+		CreationArguments.bRenameParameterOnAdd = true;
+		CreationArguments.ModuleParameterVariable = NewParameter;
+		return MakeShareable(new FScriptGroupAddAction(CreationArguments));
 	}
 
 	static TSharedRef<FScriptGroupAddAction> CreateNewScratchModuleAction()
@@ -134,8 +191,14 @@ public:
 		FNiagaraFavoritesActionData FavoritesActionData;
 		FavoritesActionData.ActionIdentifier.Names.Add(FName(DisplayName.ToString()));
 		FavoritesActionData.bFavoriteByDefault = true;
-		
-		return MakeShareable(new FScriptGroupAddAction( DisplayName, {}, Description, FText(), FavoritesActionData, true, FNiagaraVariable(), false, FAssetData(), nullptr, true, false));
+
+		FCreationArguments CreationArguments;
+		CreationArguments.DisplayName = DisplayName;
+		CreationArguments.Description = Description;
+		CreationArguments.bIsInLibrary = true;
+		CreationArguments.SpecificActionType = ESpecificActionType::NewScratchModule;
+		CreationArguments.FavoritesActionData = FavoritesActionData;
+		return MakeShareable(new FScriptGroupAddAction(CreationArguments));
 	}
 
 	static TSharedRef<FScriptGroupAddAction> CreateNewSetSpecificModuleAction()
@@ -146,8 +209,14 @@ public:
 		FNiagaraFavoritesActionData FavoritesActionData;
 		FavoritesActionData.ActionIdentifier.Names.Add(FName(DisplayName.ToString()));
 		FavoritesActionData.bFavoriteByDefault = true;
-		
-		return MakeShareable(new FScriptGroupAddAction( DisplayName, {}, Description, FText(), FavoritesActionData, true, FNiagaraVariable(), false, FAssetData(), nullptr, false, true));
+
+		FCreationArguments CreationArguments;
+		CreationArguments.DisplayName = DisplayName;
+		CreationArguments.Description = Description;
+		CreationArguments.bIsInLibrary = true;
+		CreationArguments.SpecificActionType = ESpecificActionType::NewSetSpecificParameters;
+		CreationArguments.FavoritesActionData = FavoritesActionData;
+		return MakeShareable(new FScriptGroupAddAction(CreationArguments));
 	}
 
 	static TSharedRef<FScriptGroupAddAction> CreateWizardAction(const FModuleWizardGenerator::FAction& Action)
@@ -155,9 +224,17 @@ public:
 		FNiagaraFavoritesActionData FavoritesActionData;
 		FavoritesActionData.ActionIdentifier.Names.Add(Action.WizardModel->GetIdentifier());
 		FavoritesActionData.bFavoriteByDefault = Action.bSuggestedAction;
+
+		FCreationArguments CreationArguments;
+		CreationArguments.DisplayName = Action.DisplayName;
+		CreationArguments.Description = Action.Description;
+		CreationArguments.Keywords = Action.Keywords;
+		CreationArguments.bIsInLibrary = true;
+		CreationArguments.ModuleWizardModel = Action.WizardModel;
+		CreationArguments.SpecificActionType = ESpecificActionType::ModuleWizard;
+		CreationArguments.FavoritesActionData = FavoritesActionData;
 		
-		TSharedRef<FScriptGroupAddAction> AddAction = MakeShareable(new FScriptGroupAddAction( Action.DisplayName, {}, Action.Description, Action.Keywords, FavoritesActionData, true, FNiagaraVariable(), false, FAssetData(), nullptr, false, true));
-		AddAction->ModuleWizardModel = Action.WizardModel;
+		TSharedRef<FScriptGroupAddAction> AddAction = MakeShareable(new FScriptGroupAddAction( CreationArguments));
 		return AddAction;
 	}
 
@@ -186,6 +263,11 @@ public:
 		return FavoritesData;
 	}
 
+	virtual TOptional<FSoftObjectPath> GetPreviewMoviePath() const override
+	{
+		return PreviewMoviePath;
+	}
+
 	virtual bool IsInLibrary() const override
 	{
 		return bIsInLibrary;
@@ -218,17 +300,17 @@ public:
 
 	bool IsNewScratchModuleAction() const
 	{
-		return bIsNewScratchModuleAction;
+		return SpecificActionType.IsSet() && SpecificActionType.GetValue() == ESpecificActionType::NewScratchModule;
 	}
 
 	bool IsNewSetSpecificModuleAction() const
 	{
-		return bIsNewSetSpecificModuleAction;
+		return SpecificActionType.IsSet() && SpecificActionType.GetValue() == ESpecificActionType::NewSetSpecificParameters;
 	}
 
 	bool HasWizardWidget() const
 	{
-		return ModuleWizardModel.IsValid();
+		return SpecificActionType.IsSet() && SpecificActionType.GetValue() == ESpecificActionType::ModuleWizard && ModuleWizardModel.IsValid();
 	}
 
 	TSharedRef<SNiagaraModuleWizard> CreateWizard() const
@@ -238,41 +320,39 @@ public:
 	}
 
 private:
-	FScriptGroupAddAction(FText InDisplayName, TArray<FString> InCategories, FText InDescription, FText InKeywords, TOptional<FNiagaraFavoritesActionData> InFavoritesData, bool bInIsInLibrary,
-		FNiagaraVariable InModuleParameterVariable, bool bInRenameParameterOnAdd,
-		FAssetData InModuleAssetData, UNiagaraScript* InModuleScript,
-		bool bInIsNewScratchModuleAction, bool bInIsNewSetSpecificModuleAction, FNiagaraActionSourceData InSourceData = FNiagaraActionSourceData(EScriptSource::Niagara, FText::FromString(TEXT("Niagara")), true))
-
-		: DisplayName(InDisplayName)
-		, Categories(InCategories)
-		, Description(InDescription)
-		, Keywords(InKeywords)
-		, FavoritesData(InFavoritesData)
-		, bIsInLibrary(bInIsInLibrary)
-		, ModuleParameterVariable(InModuleParameterVariable)
-		, bRenameParameterOnAdd(bInRenameParameterOnAdd)
-		, ModuleAssetData(InModuleAssetData)
-		, ModuleScript(InModuleScript)
-		, bIsNewScratchModuleAction(bInIsNewScratchModuleAction)
-		, bIsNewSetSpecificModuleAction(bInIsNewSetSpecificModuleAction)
-		, SourceData(InSourceData)
-	{
-	}
+	FScriptGroupAddAction(const FCreationArguments& CreationArguments)
+		: DisplayName(CreationArguments.DisplayName)
+		, Categories(CreationArguments.Categories)
+		, Description(CreationArguments.Description)
+		, Keywords(CreationArguments.Keywords)
+		, FavoritesData(CreationArguments.FavoritesActionData)
+		, PreviewMoviePath(CreationArguments.PreviewMoviePath)
+		, bIsInLibrary(CreationArguments.bIsInLibrary)
+		, ModuleParameterVariable(CreationArguments.ModuleParameterVariable)
+		, ModuleAssetData(CreationArguments.ModuleAssetData)
+		, ModuleScript(CreationArguments.ModuleScript)
+		, ModuleWizardModel(CreationArguments.ModuleWizardModel)
+		, SpecificActionType(CreationArguments.SpecificActionType)
+		, SourceData(CreationArguments.SourceData)
+		, bRenameParameterOnAdd(CreationArguments.bRenameParameterOnAdd)
+		{
+			
+		}
 
 	FText DisplayName;
 	TArray<FString> Categories;
 	FText Description;
 	FText Keywords;
 	TOptional<FNiagaraFavoritesActionData> FavoritesData;
+	TOptional<FSoftObjectPath> PreviewMoviePath;
 	bool bIsInLibrary;
 	FNiagaraVariable ModuleParameterVariable;
-	bool bRenameParameterOnAdd;
 	FAssetData ModuleAssetData;
 	UNiagaraScript* ModuleScript;
-	bool bIsNewScratchModuleAction;
-	bool bIsNewSetSpecificModuleAction;
-	FNiagaraActionSourceData SourceData;
 	TSharedPtr<FModuleWizardModel> ModuleWizardModel;
+	TOptional<ESpecificActionType> SpecificActionType;
+	FNiagaraActionSourceData SourceData;
+	bool bRenameParameterOnAdd;
 };
 
 class FScriptItemGroupAddUtilities : public TNiagaraStackItemGroupAddUtilities<UNiagaraNodeFunctionCall*>
diff --git a/Engine/Plugins/FX/Niagara/Source/NiagaraEditor/Private/ViewModels/Stack/NiagaraStackStatelessEmitterSimulateGroup.cpp b/Engine/Plugins/FX/Niagara/Source/NiagaraEditor/Private/ViewModels/Stack/NiagaraStackStatelessEmitterSimulateGroup.cpp
--- a/Engine/Plugins/FX/Niagara/Source/NiagaraEditor/Private/ViewModels/Stack/NiagaraStackStatelessEmitterSimulateGroup.cpp
+++ b/Engine/Plugins/FX/Niagara/Source/NiagaraEditor/Private/ViewModels/Stack/NiagaraStackStatelessEmitterSimulateGroup.cpp
@@ -98,6 +98,7 @@ public:
 	virtual TOptional<FNiagaraFavoritesActionData> GetFavoritesData() const override { return TOptional<FNiagaraFavoritesActionData>(); }
 	virtual bool IsInLibrary() const override { return true; }
 	virtual FNiagaraActionSourceData GetSourceData() const override { return FNiagaraActionSourceData(); }
+	virtual TOptional<FSoftObjectPath> GetPreviewMoviePath() const override { return {}; }
 
 private:
 	TWeakObjectPtr<UNiagaraStatelessModule> StatelessModuleWeak;
diff --git a/Engine/Plugins/FX/Niagara/Source/NiagaraEditor/Private/ViewModels/Stack/NiagaraStackStatelessEmitterSpawnGroup.cpp b/Engine/Plugins/FX/Niagara/Source/NiagaraEditor/Private/ViewModels/Stack/NiagaraStackStatelessEmitterSpawnGroup.cpp
--- a/Engine/Plugins/FX/Niagara/Source/NiagaraEditor/Private/ViewModels/Stack/NiagaraStackStatelessEmitterSpawnGroup.cpp
+++ b/Engine/Plugins/FX/Niagara/Source/NiagaraEditor/Private/ViewModels/Stack/NiagaraStackStatelessEmitterSpawnGroup.cpp
@@ -109,6 +109,7 @@ public:
 	virtual TOptional<FNiagaraFavoritesActionData> GetFavoritesData() const override { return TOptional<FNiagaraFavoritesActionData>(); }
 	virtual bool IsInLibrary() const override { return true; }
 	virtual FNiagaraActionSourceData GetSourceData() const override { return FNiagaraActionSourceData(); }
+	virtual TOptional<FSoftObjectPath> GetPreviewMoviePath() const override { return {}; }
 
 private:
 	ENiagaraStatelessSpawnInfoType SpawnInfoType;
diff --git a/Engine/Plugins/FX/Niagara/Source/NiagaraEditor/Private/ViewModels/Stack/NiagaraStackSystemSettingsGroup.cpp b/Engine/Plugins/FX/Niagara/Source/NiagaraEditor/Private/ViewModels/Stack/NiagaraStackSystemSettingsGroup.cpp
--- a/Engine/Plugins/FX/Niagara/Source/NiagaraEditor/Private/ViewModels/Stack/NiagaraStackSystemSettingsGroup.cpp
+++ b/Engine/Plugins/FX/Niagara/Source/NiagaraEditor/Private/ViewModels/Stack/NiagaraStackSystemSettingsGroup.cpp
@@ -66,6 +66,11 @@ public:
 		return FNiagaraActionSourceData();
 	}
 
+	virtual TOptional<FSoftObjectPath> GetPreviewMoviePath() const override
+	{
+		return {};
+	}
+
 private:
 	FNiagaraVariable NewParameterVariable;
 };
diff --git a/Engine/Plugins/FX/Niagara/Source/NiagaraEditor/Private/Widgets/SNiagaraGraphActionMenu.cpp b/Engine/Plugins/FX/Niagara/Source/NiagaraEditor/Private/Widgets/SNiagaraGraphActionMenu.cpp
--- a/Engine/Plugins/FX/Niagara/Source/NiagaraEditor/Private/Widgets/SNiagaraGraphActionMenu.cpp
+++ b/Engine/Plugins/FX/Niagara/Source/NiagaraEditor/Private/Widgets/SNiagaraGraphActionMenu.cpp
@@ -197,6 +197,7 @@ TSharedRef<SWidget> SNiagaraGraphActionMenu::OnGenerateWidgetForItem(const TShar
 	FCreateNiagaraWidgetForActionData ActionData(Item);
 	ActionData.HighlightText = TAttribute<FText>::Create(TAttribute<FText>::FGetter::CreateRaw(this, &SNiagaraGraphActionMenu::GetFilterText));
 	ActionData.FavoriteActionsProfileName = SNiagaraGraphActionMenu::FavoriteActionsProfile;
+	ActionData.PreviewMovieViewModel = ViewModel->GetPreviewMovieViewModel();
 	return SNew(SNiagaraActionWidget, ActionData);
 }
 
diff --git a/Engine/Plugins/FX/Niagara/Source/NiagaraEditor/Private/Widgets/SNiagaraGraphActionWidget.cpp b/Engine/Plugins/FX/Niagara/Source/NiagaraEditor/Private/Widgets/SNiagaraGraphActionWidget.cpp
--- a/Engine/Plugins/FX/Niagara/Source/NiagaraEditor/Private/Widgets/SNiagaraGraphActionWidget.cpp
+++ b/Engine/Plugins/FX/Niagara/Source/NiagaraEditor/Private/Widgets/SNiagaraGraphActionWidget.cpp
@@ -13,6 +13,8 @@
 #include "Config/NiagaraFavoriteActionsConfig.h"
 #include "Widgets/Layout/SSpacer.h"
 #include "Widgets/Layout/SSeparator.h"
+#include "Widgets/SToolTip.h"
+#include "Widgets/IToolTip.h"
 
 #define LOCTEXT_NAMESPACE "NiagaraGraphActionWidget"
 
@@ -84,6 +86,7 @@ void SNiagaraActionWidget::Construct(const FArguments& InArgs, const FCreateNiag
 {
 	ActionPtr = InCreateData.Action;
 	FavoriteProfileName = InCreateData.FavoriteActionsProfileName;
+	PreviewMovieViewModel = InCreateData.PreviewMovieViewModel;
 
 	TSharedPtr<SWidget> NameWidget;
 	if (ActionPtr->GetParameterVariable().IsSet())
@@ -117,7 +120,7 @@ void SNiagaraActionWidget::Construct(const FArguments& InArgs, const FCreateNiag
 	this->ChildSlot
 	[
 		SNew(SHorizontalBox)
-		.ToolTipText(InCreateData.Action->ToolTip)
+		.ToolTip(this, &SNiagaraActionWidget::GetCustomTooltip)
 		+ SHorizontalBox::Slot()
 		.VAlign(VAlign_Center)
 		.AutoWidth()
@@ -262,4 +265,24 @@ FText SNiagaraActionWidget::GetFavoriteButtonTooltipText() const
 	return LOCTEXT("FavoriteButtonTooltip", "Add this action to favorites");
 }
 
+TSharedPtr<IToolTip> SNiagaraActionWidget::GetCustomTooltip() const
+{
+	using namespace FNiagaraEditorUtilities::Tooltips;
+	
+	if(PreviewMovieViewModel.IsValid() && ActionPtr->PreviewMoviePath.IsSet())
+	{
+		PreviewMovieViewModel.Pin()->PlayMovie(ActionPtr->PreviewMoviePath.GetValue());
+	}
+		
+	if(TooltipCache.IsValid() == false)
+	{
+		FCreateTooltipArguments CreateTooltipArguments;
+		CreateTooltipArguments.bTryCreatePreviewMovie = GetDefault<UNiagaraEditorSettings>()->bDisplayPreviewMoviesInTooltips;
+		CreateTooltipArguments.PreviewMovieOwner = PreviewMovieViewModel.IsValid() ? PreviewMovieViewModel.Pin() : nullptr;
+		TooltipCache = CreateTooltipForNiagaraAction(ActionPtr.ToSharedRef(), CreateTooltipArguments);
+	}		
+		
+	return TooltipCache.ToSharedRef();
+}
+
 #undef LOCTEXT_NAMESPACE
diff --git a/Engine/Plugins/FX/Niagara/Source/NiagaraEditor/Public/NiagaraActions.h b/Engine/Plugins/FX/Niagara/Source/NiagaraEditor/Public/NiagaraActions.h
--- a/Engine/Plugins/FX/Niagara/Source/NiagaraEditor/Public/NiagaraActions.h
+++ b/Engine/Plugins/FX/Niagara/Source/NiagaraEditor/Public/NiagaraActions.h
@@ -150,6 +150,8 @@ struct FNiagaraMenuAction_Base
 	float SearchWeightMultiplier = 1.f;
 	
 	TOptional<FNiagaraFavoritesActionData> FavoritesActionData;
+
+	TOptional<FSoftObjectPath> PreviewMoviePath;
 	
 	mutable bool bIsHovered = false;
 };
diff --git a/Engine/Plugins/FX/Niagara/Source/NiagaraEditor/Public/NiagaraEditorModule.h b/Engine/Plugins/FX/Niagara/Source/NiagaraEditor/Public/NiagaraEditorModule.h
--- a/Engine/Plugins/FX/Niagara/Source/NiagaraEditor/Public/NiagaraEditorModule.h
+++ b/Engine/Plugins/FX/Niagara/Source/NiagaraEditor/Public/NiagaraEditorModule.h
@@ -9,12 +9,15 @@
 #include "Toolkits/AssetEditorToolkit.h"
 #include "NiagaraTypes.h"
 #include "INiagaraCompiler.h"
+#include "MediaPlayer.h"
+#include "MediaTexture.h"
 #include "NiagaraCompilationTypes.h"
 #include "NiagaraPerfBaseline.h"
 #include "NiagaraDebuggerCommon.h"
 #include "NiagaraRendererProperties.h"
 #include "Customizations/NiagaraDataInterfaceSimCacheVisualizer.h"
 #include "Customizations/NiagaraStackObjectPropertyCustomization.h"
+#include "Brushes/SlateImageBrush.h"
 #include "NiagaraEditorModule.generated.h"
 
 class UNiagaraStackObject;
@@ -316,7 +319,17 @@ public:
 	UE::Niagara::TraversalCache::FTraversalBuilder& GetParameterTraversalBuilder() { return *TraversalBuilder.Get(); };
 	UE::Niagara::TraversalCache::FTraversalCache& GetParameterTraversalCache() { return *TraversalCache.Get(); }
 
+	struct FMoviePreviewObjects
+	{
+		TStrongObjectPtr<UMediaPlayer> MediaPlayer;
+		TStrongObjectPtr<UMediaTexture> MediaTexture;
+		FSlateImageBrush* ImageBrush = nullptr;
+	};
+	
+	const FMoviePreviewObjects& FindOrAddMoviePreviewObjects(const FAssetData& AssetData);
 private:
+	TMap<FAssetData, FMoviePreviewObjects> MoviePreviewObjects;
+	
 	class FDeferredDestructionContainerBase
 	{
 	public:
diff --git a/Engine/Plugins/FX/Niagara/Source/NiagaraEditor/Public/NiagaraEditorSettings.h b/Engine/Plugins/FX/Niagara/Source/NiagaraEditor/Public/NiagaraEditorSettings.h
--- a/Engine/Plugins/FX/Niagara/Source/NiagaraEditor/Public/NiagaraEditorSettings.h
+++ b/Engine/Plugins/FX/Niagara/Source/NiagaraEditor/Public/NiagaraEditorSettings.h
@@ -365,6 +365,14 @@ public:
 	UPROPERTY(config, EditAnywhere, Category = "Niagara Renderer", meta = (AllowedClasses = "/Script/Engine.MaterialInterface"))
 	FSoftObjectPath DefaultSpriteRendererMaterial;
 
+	/** Whether we allow using previews in the editor. */
+	UPROPERTY(Config, EditAnywhere, Category = "Preview")
+	bool bDisplayPreviewMoviesInTooltips = true;
+
+	/** The size of the movie previews used in tooltips. */
+	UPROPERTY(Config, EditAnywhere, Category = "Preview", meta=(EditCondition="bDisplayPreviewMoviesInTooltips==true"))
+	FVector2f TooltipPreviewMovieSize = FVector2f(128.f, 128.f);
+
 	/** Niagara script which is required in the system update script to control system state. */
 	UPROPERTY(config, EditAnywhere, Category = Niagara)
 	FSoftObjectPath RequiredSystemUpdateScript;
diff --git a/Engine/Plugins/FX/Niagara/Source/NiagaraEditor/Public/NiagaraEditorUtilities.h b/Engine/Plugins/FX/Niagara/Source/NiagaraEditor/Public/NiagaraEditorUtilities.h
--- a/Engine/Plugins/FX/Niagara/Source/NiagaraEditor/Public/NiagaraEditorUtilities.h
+++ b/Engine/Plugins/FX/Niagara/Source/NiagaraEditor/Public/NiagaraEditorUtilities.h
@@ -437,8 +437,28 @@ namespace FNiagaraEditorUtilities
 		NIAGARAEDITOR_API const UNiagaraScriptVariable* FindScriptVariableForUserParameter(const FGuid& UserParameterGuid, const UNiagaraSystem& System);
 		void ReplaceUserParameterReferences(TSharedRef<FNiagaraEmitterViewModel> EmitterViewModel, FNiagaraVariable OldUserParameter, FNiagaraVariable NewUserParameter);
 	}
+
+	namespace Preview
+	{
+		NIAGARAEDITOR_API TOptional<FSoftObjectPath> GetPreviewMovieObjectPath(const FAssetData& AssetData);
+	}
+	
 	namespace Tooltips
 	{
+		struct FCreateTooltipArguments
+		{
+			/** If this is true, PreviewMovieOwner has to be valid as well. If false, will not attempt to create a movie. */
+			bool bTryCreatePreviewMovie = false;
+			TSharedPtr<FNiagaraPreviewMovieViewModel> PreviewMovieOwner;
+		};
+
+		/** An all purpose function to create a tooltip with a text and an optional movie. */
+		NIAGARAEDITOR_API TSharedRef<SToolTip> CreateTooltip(FText Description, TOptional<FSoftObjectPath> PreviewMoviePath = TOptional<FSoftObjectPath>(), const FCreateTooltipArguments& CreateTooltipArguments = FCreateTooltipArguments());
+		/** Creates a tooltip for a niagara action. Supports preview movies. */
+		NIAGARAEDITOR_API TSharedRef<SToolTip> CreateTooltipForNiagaraAction(TSharedRef<FNiagaraMenuAction_Generic> Action, const FCreateTooltipArguments& CreateTooltipArguments = FCreateTooltipArguments());
+		/** Creates a tooltip for a script. Supports preview movies.*/
+		NIAGARAEDITOR_API TSharedRef<SToolTip> CreateTooltipForScript(UNiagaraScript& Script, FGuid VersionGuid = FGuid(), const FCreateTooltipArguments& CreateTooltipArguments = FCreateTooltipArguments());
+		
 		NIAGARAEDITOR_API TSharedRef<SToolTip> CreateStackNoteTooltip(UNiagaraStackNote& StackNote);
 		NIAGARAEDITOR_API FText GetMinimalEmitterCreationTooltip();
 	}
diff --git a/Engine/Plugins/FX/Niagara/Source/NiagaraEditor/Public/SNiagaraGraphActionWidget.h b/Engine/Plugins/FX/Niagara/Source/NiagaraEditor/Public/SNiagaraGraphActionWidget.h
--- a/Engine/Plugins/FX/Niagara/Source/NiagaraEditor/Public/SNiagaraGraphActionWidget.h
+++ b/Engine/Plugins/FX/Niagara/Source/NiagaraEditor/Public/SNiagaraGraphActionWidget.h
@@ -8,6 +8,7 @@
 #include "SlateFwd.h"
 #include "SGraphActionMenu.h"
 #include "EdGraph/EdGraphSchema.h"
+#include "ViewModels/NiagaraPreviewMovieViewModel.h"
 #include "Widgets/DeclarativeSyntaxSupport.h"
 #include "Widgets/SCompoundWidget.h"
 
@@ -31,6 +32,8 @@ struct FCreateNiagaraWidgetForActionData
 	bool bIsReadOnly = true;
 	/** If specified, adds a 'favorite' toggle to the widget */
 	TOptional<FName> FavoriteActionsProfileName;
+	/** Needs to be supplied for preview movie playback to work. */
+	TSharedPtr<FNiagaraPreviewMovieViewModel> PreviewMovieViewModel;
 };
 
 /** Custom widget for GraphActionMenu */
@@ -73,11 +76,14 @@ private:
 	EVisibility OnShouldShowFavoriteButton() const;
 	const FSlateBrush* GetFavoriteBrush() const;
 	FText GetFavoriteButtonTooltipText() const;
+	TSharedPtr<IToolTip> GetCustomTooltip() const;
 private:
 	/** The item that we want to display with this widget */
 	TSharedPtr<struct FNiagaraMenuAction_Generic> ActionPtr;
+	TWeakPtr<FNiagaraPreviewMovieViewModel> PreviewMovieViewModel;
 	/** Delegate executed when mouse button goes down */
 	FOnActionClicked MouseButtonDownDelegate;
 	FNiagaraActionSourceData SourceData;
 	TOptional<FName> FavoriteProfileName;
+	mutable TSharedPtr<IToolTip> TooltipCache;
 };
diff --git a/Engine/Plugins/FX/Niagara/Source/NiagaraEditor/Public/ViewModels/NiagaraPreviewMovieViewModel.cpp b/Engine/Plugins/FX/Niagara/Source/NiagaraEditor/Public/ViewModels/NiagaraPreviewMovieViewModel.cpp
new file mode 100644
--- /dev/null
+++ b/Engine/Plugins/FX/Niagara/Source/NiagaraEditor/Public/ViewModels/NiagaraPreviewMovieViewModel.cpp
@@ -0,0 +1,104 @@
+﻿// Copyright Epic Games, Inc. All Rights Reserved.
+
+#include "NiagaraPreviewMovieViewModel.h"
+#include "ViewModels/NiagaraSystemViewModel.h"
+#include "UObject/UObjectGlobals.h"
+#include "FileMediaSource.h"
+#include "NiagaraEditorSettings.h"
+#include "NiagaraEditorUtilities.h"
+#include "NiagaraTrace.h"
+
+FNiagaraPreviewMovieViewModel::FNiagaraPreviewMovieViewModel()
+{
+	MediaPlayer.Reset(NewObject<UMediaPlayer>(GetTransientPackage(), UMediaPlayer::StaticClass()));
+	MediaTexture.Reset(NewObject<UMediaTexture>(GetTransientPackage()));
+	MediaPlayer->SetLooping(true);
+	MediaPlayer->AffectedByPIEHandling = false;
+	MediaTexture->SetMediaPlayer(MediaPlayer.Get());
+	
+	const UNiagaraEditorSettings* EditorSettings = GetDefault<UNiagaraEditorSettings>();
+	ImageBrush = MakeUnique<FSlateImageBrush>(MediaTexture.Get(), EditorSettings->TooltipPreviewMovieSize);
+}
+
+FNiagaraPreviewMovieViewModel::~FNiagaraPreviewMovieViewModel()
+{
+	Shutdown();
+}
+
+void FNiagaraPreviewMovieViewModel::Initialize()
+{
+	if(const UNiagaraEditorSettings* EditorSettings = GetDefault<UNiagaraEditorSettings>())
+	{
+		EditorSettings->OnSettingsChanged().AddSP(this, &FNiagaraPreviewMovieViewModel::OnEditorSettingsChanged);
+	}
+}
+
+void FNiagaraPreviewMovieViewModel::Shutdown()
+{
+	if(const UNiagaraEditorSettings* EditorSettings = GetDefault<UNiagaraEditorSettings>())
+	{
+		EditorSettings->OnSettingsChanged().RemoveAll(this);
+	}
+}
+
+bool FNiagaraPreviewMovieViewModel::IsPlayingPreviewMovieForAsset(const FAssetData& AssetData) const
+{
+	return IsPlayingPreviewMovie(FNiagaraEditorUtilities::Preview::GetPreviewMovieObjectPath(AssetData).Get(FSoftObjectPath()));
+}
+
+bool FNiagaraPreviewMovieViewModel::IsPlayingPreviewMovie(const FSoftObjectPath& InPreviewMoviePath) const
+{
+	return InPreviewMoviePath == CurrentPreviewMoviePath;
+}
+
+bool FNiagaraPreviewMovieViewModel::PlayMovieForAsset(const FAssetData& AssetData)
+{
+	TRACE_CPUPROFILER_EVENT_SCOPE_ON_CHANNEL(FNiagaraMoviePreviewViewModel::PlayMovieForAsset, NiagaraChannel)
+	
+	// If we are already playing this, return
+	if(IsPlayingPreviewMovieForAsset(AssetData))
+	{
+		return false;
+	}
+
+	TOptional<FSoftObjectPath> PreviewMoviePath = FNiagaraEditorUtilities::Preview::GetPreviewMovieObjectPath(AssetData);
+
+	if(PreviewMoviePath.IsSet() == false)
+	{
+		return false;
+	}
+
+	return PlayMovie(PreviewMoviePath.GetValue());
+}
+
+bool FNiagaraPreviewMovieViewModel::PlayMovie(const FSoftObjectPath& InPreviewMoviePath)
+{
+	TRACE_CPUPROFILER_EVENT_SCOPE_ON_CHANNEL(FNiagaraMoviePreviewViewModel::PlayMovie, NiagaraChannel)
+
+	if(IsPlayingPreviewMovie(InPreviewMoviePath))
+	{
+		return false;
+	}
+	
+	// Even if this fails, we update the current preview movie path so that further tries with the same path won't drain perf unnecessarily
+	UFileMediaSource* MediaSource = Cast<UFileMediaSource>(InPreviewMoviePath.TryLoad());
+
+	CurrentPreviewMoviePath = InPreviewMoviePath;
+	
+	if(MediaSource == nullptr)
+	{
+		return false;
+	}
+	
+	MediaPlayer->OpenSource(MediaSource);
+	MediaTexture->UpdateResource();
+	return MediaPlayer->Play();
+}
+
+void FNiagaraPreviewMovieViewModel::OnEditorSettingsChanged(const FString& PropertyName, const UNiagaraEditorSettings* NiagaraEditorSettings)
+{
+	if(PropertyName == GET_MEMBER_NAME_CHECKED(UNiagaraEditorSettings, TooltipPreviewMovieSize))
+	{
+		ImageBrush->SetImageSize(NiagaraEditorSettings->TooltipPreviewMovieSize);
+	}
+}
diff --git a/Engine/Plugins/FX/Niagara/Source/NiagaraEditor/Public/ViewModels/NiagaraPreviewMovieViewModel.h b/Engine/Plugins/FX/Niagara/Source/NiagaraEditor/Public/ViewModels/NiagaraPreviewMovieViewModel.h
new file mode 100644
--- /dev/null
+++ b/Engine/Plugins/FX/Niagara/Source/NiagaraEditor/Public/ViewModels/NiagaraPreviewMovieViewModel.h
@@ -0,0 +1,38 @@
+﻿// Copyright Epic Games, Inc. All Rights Reserved.
+
+#pragma once
+
+#include "Templates/SharedPointer.h"
+#include "MediaPlayer.h"
+#include "MediaTexture.h"
+#include "NiagaraEditorSettings.h"
+#include "Brushes/SlateImageBrush.h"
+#include "Templates/UniquePtr.h"
+#include "AssetRegistry/AssetData.h"
+
+class FNiagaraSystemViewModel;
+
+class FNiagaraPreviewMovieViewModel : public TSharedFromThis<FNiagaraPreviewMovieViewModel>
+{
+public:
+	FNiagaraPreviewMovieViewModel();
+	~FNiagaraPreviewMovieViewModel();
+
+	void Initialize();
+	void Shutdown();
+	
+	bool IsPlayingPreviewMovieForAsset(const FAssetData& AssetData) const;
+	bool IsPlayingPreviewMovie(const FSoftObjectPath& InPreviewMoviePath) const;
+	bool PlayMovieForAsset(const FAssetData& AssetData);
+	bool PlayMovie(const FSoftObjectPath& InPreviewMoviePath);
+	const FSlateImageBrush* GetMoviePreviewImageBrush() const { return ImageBrush.Get(); }
+
+protected:
+	void OnEditorSettingsChanged(const FString& PropertyName, const UNiagaraEditorSettings* NiagaraEditorSettings);
+
+private:
+	TStrongObjectPtr<UMediaPlayer> MediaPlayer;
+	TStrongObjectPtr<UMediaTexture> MediaTexture;
+	TUniquePtr<FSlateImageBrush> ImageBrush;
+	FSoftObjectPath CurrentPreviewMoviePath;
+};
diff --git a/Engine/Plugins/FX/Niagara/Source/NiagaraEditor/Public/ViewModels/NiagaraSystemViewModel.h b/Engine/Plugins/FX/Niagara/Source/NiagaraEditor/Public/ViewModels/NiagaraSystemViewModel.h
--- a/Engine/Plugins/FX/Niagara/Source/NiagaraEditor/Public/ViewModels/NiagaraSystemViewModel.h
+++ b/Engine/Plugins/FX/Niagara/Source/NiagaraEditor/Public/ViewModels/NiagaraSystemViewModel.h
@@ -9,6 +9,7 @@
 #include "NiagaraCommon.h"
 #include "NiagaraEditorCommon.h"
 #include "NiagaraEditorModule.h"
+#include "NiagaraPreviewMovieViewModel.h"
 #include "NiagaraOverviewNode.h"
 #include "NiagaraSystemScalabilityViewModel.h"
 #include "NiagaraUserParameterPanelViewModel.h"
@@ -488,6 +489,8 @@ public:
 	const TSharedPtr<FNiagaraSystemGraphSelectionViewModel>& GetSystemGraphSelectionViewModel() const { return SystemGraphSelectionViewModel; };
 
 	TSharedPtr<FNiagaraUserParameterPanelViewModel> GetUserParameterPanelViewModel() const { return UserParameterPanelViewModel; }
+
+	TSharedPtr<FNiagaraPreviewMovieViewModel> GetPreviewMovieViewModel() const { return PreviewMovieViewModel; }
 	
 	/** Utility method to gather all graphs. */
 	TArray<UNiagaraGraph*> GetAllGraphs();
@@ -656,6 +659,9 @@ private:
 	/** The view model for the System script. */
 	TSharedPtr<FNiagaraSystemScriptViewModel> SystemScriptViewModel;
 
+	/** The view model that owns and manages the singular movie player available to an asset. */
+	TSharedPtr<FNiagaraPreviewMovieViewModel> PreviewMovieViewModel;
+	
 	/** Handles linking change events from the sequencer movie scene. */
 	UE::MovieScene::TNonIntrusiveEventHandler<UE::MovieScene::ISignedObjectEventHandler> MovieSceneEventHandler;
 
diff --git a/Engine/Plugins/FX/Niagara/Source/NiagaraEditor/Public/ViewModels/Stack/INiagaraStackItemGroupAddUtilities.h b/Engine/Plugins/FX/Niagara/Source/NiagaraEditor/Public/ViewModels/Stack/INiagaraStackItemGroupAddUtilities.h
--- a/Engine/Plugins/FX/Niagara/Source/NiagaraEditor/Public/ViewModels/Stack/INiagaraStackItemGroupAddUtilities.h
+++ b/Engine/Plugins/FX/Niagara/Source/NiagaraEditor/Public/ViewModels/Stack/INiagaraStackItemGroupAddUtilities.h
@@ -24,6 +24,9 @@ public:
 	/** Required for favorites feature. */
 	virtual TOptional<FNiagaraFavoritesActionData> GetFavoritesData() const = 0;
 
+	/** If specified and valid, will display this movie on its tooltip. */
+	virtual TOptional<FSoftObjectPath> GetPreviewMoviePath() const = 0;
+	
 	/** Indicates if this action represents a library action */
 	virtual bool IsInLibrary() const = 0;
 
diff --git a/Engine/Plugins/FX/Niagara/Source/NiagaraEditor/Public/ViewModels/Stack/NiagaraStackEmitterPropertiesGroup.h b/Engine/Plugins/FX/Niagara/Source/NiagaraEditor/Public/ViewModels/Stack/NiagaraStackEmitterPropertiesGroup.h
--- a/Engine/Plugins/FX/Niagara/Source/NiagaraEditor/Public/ViewModels/Stack/NiagaraStackEmitterPropertiesGroup.h
+++ b/Engine/Plugins/FX/Niagara/Source/NiagaraEditor/Public/ViewModels/Stack/NiagaraStackEmitterPropertiesGroup.h
@@ -37,7 +37,8 @@ private:
 		virtual TOptional<FNiagaraFavoritesActionData> GetFavoritesData() const override { return TOptional<FNiagaraFavoritesActionData>(); }
 		virtual bool IsInLibrary() const override { return true; }
 		virtual FNiagaraActionSourceData GetSourceData() const override { return FNiagaraActionSourceData(); }
-
+		virtual TOptional<FSoftObjectPath> GetPreviewMoviePath() const override { return {}; }
+		
 		EStageAddMode AddMode;
 		UClass* SimulationStageClass;
 		TArray<FString> Categories;
diff --git a/Engine/Plugins/FX/Niagara/Source/NiagaraEditor/Public/ViewModels/Stack/NiagaraStackEntry.h b/Engine/Plugins/FX/Niagara/Source/NiagaraEditor/Public/ViewModels/Stack/NiagaraStackEntry.h
--- a/Engine/Plugins/FX/Niagara/Source/NiagaraEditor/Public/ViewModels/Stack/NiagaraStackEntry.h
+++ b/Engine/Plugins/FX/Niagara/Source/NiagaraEditor/Public/ViewModels/Stack/NiagaraStackEntry.h
@@ -17,6 +17,7 @@ class UNiagaraStackNote;
 class UNiagaraClipboardContent;
 class UNiagaraStackEntry;
 struct FNiagaraStackNoteData;
+class IToolTip;
 
 UENUM()
 enum class EStackIssueSeverity : uint8
@@ -271,8 +272,10 @@ public:
 	NIAGARAEDITOR_API UNiagaraStackEditorData& GetStackEditorData() const;
 
 	NIAGARAEDITOR_API FString GetStackEditorDataKey() const;
-
+	
 	NIAGARAEDITOR_API virtual FText GetTooltipText() const;
+	NIAGARAEDITOR_API virtual bool UsesCustomTooltipWidget() const;
+	NIAGARAEDITOR_API virtual TSharedPtr<IToolTip> GetCustomTooltipWidget() const;
 
 	NIAGARAEDITOR_API virtual bool GetCanExpand() const;
 
diff --git a/Engine/Plugins/FX/Niagara/Source/NiagaraEditor/Public/ViewModels/Stack/NiagaraStackModuleItem.h b/Engine/Plugins/FX/Niagara/Source/NiagaraEditor/Public/ViewModels/Stack/NiagaraStackModuleItem.h
--- a/Engine/Plugins/FX/Niagara/Source/NiagaraEditor/Public/ViewModels/Stack/NiagaraStackModuleItem.h
+++ b/Engine/Plugins/FX/Niagara/Source/NiagaraEditor/Public/ViewModels/Stack/NiagaraStackModuleItem.h
@@ -34,6 +34,8 @@ public:
 	NIAGARAEDITOR_API virtual FText GetDisplayName() const override;
 	NIAGARAEDITOR_API virtual UObject* GetDisplayedObject() const override;
 	NIAGARAEDITOR_API virtual FText GetTooltipText() const override;
+	virtual bool UsesCustomTooltipWidget() const override;
+	virtual TSharedPtr<IToolTip> GetCustomTooltipWidget() const override;
 
 	NIAGARAEDITOR_API INiagaraStackItemGroupAddUtilities* GetGroupAddUtilities();
 
@@ -151,6 +153,7 @@ private:
 	mutable TOptional<FText> DisplayNameCache;
 	bool bIsEnabled;
 	bool bCanRefresh;
+	mutable TSharedPtr<IToolTip> TooltipCache;
 
 	UPROPERTY()
 	TObjectPtr<UNiagaraStackModuleItemLinkedInputCollection> LinkedInputCollection;
diff --git a/Engine/Plugins/FX/Niagara/Source/NiagaraEditorWidgets/Private/SNiagaraOverviewStack.cpp b/Engine/Plugins/FX/Niagara/Source/NiagaraEditorWidgets/Private/SNiagaraOverviewStack.cpp
--- a/Engine/Plugins/FX/Niagara/Source/NiagaraEditorWidgets/Private/SNiagaraOverviewStack.cpp
+++ b/Engine/Plugins/FX/Niagara/Source/NiagaraEditorWidgets/Private/SNiagaraOverviewStack.cpp
@@ -104,6 +104,8 @@ class SNiagaraSystemOverviewEntryListRow : public STableRow<UNiagaraStackEntry*>
 
 		StackEntry->OnStructureChanged().AddSP(this, &SNiagaraSystemOverviewEntryListRow::InvalidateCachedData);
 
+		TSharedPtr<SBorder> TooltipBorder;
+		
 		TSharedPtr<SOverlay> RowOverlay;
 		STableRow<UNiagaraStackEntry*>::Construct(STableRow<UNiagaraStackEntry*>::FArguments()
 			.Style(InArgs._Style)
@@ -118,9 +120,8 @@ class SNiagaraSystemOverviewEntryListRow : public STableRow<UNiagaraStackEntry*>
 				SAssignNew(ScalabilityOverlay, SOverlay)
 				+ SOverlay::Slot()
 				[
-					SNew(SBorder)
+					SAssignNew(TooltipBorder, SBorder)
 					.BorderImage(FNiagaraEditorWidgetsStyle::Get().GetBrush("Niagara.TableViewRowBorder"))
-					.ToolTipText_UObject(StackEntry, &UNiagaraStackEntry::GetTooltipText)
 					.Padding(FMargin(0, 0, 0, 1))
 					[
 						SNew(SHorizontalBox)
@@ -229,6 +230,16 @@ class SNiagaraSystemOverviewEntryListRow : public STableRow<UNiagaraStackEntry*>
 				})
 			];
 		}
+
+		if(StackEntry->UsesCustomTooltipWidget())
+		{
+			TooltipBorder->SetToolTip(TAttribute<TSharedPtr<IToolTip>>::Create(
+				TAttribute<TSharedPtr<IToolTip>>::FGetter::CreateUObject(StackEntry, &UNiagaraStackEntry::GetCustomTooltipWidget)));
+		}
+		else
+		{
+			TooltipBorder->SetToolTipText(TAttribute<FText>::Create(TAttribute<FText>::FGetter::CreateUObject(StackEntry, &UNiagaraStackEntry::GetTooltipText)));
+		}
 	}
 
 	static const FTableRowStyle& GetDefaultStyle()
@@ -552,10 +563,20 @@ private:
 			[
 				SAssignNew(NameTextBlock, STextBlock)
 				.Text(this, &SNiagaraSystemOverviewItemName::GetItemDisplayName)
-				.ToolTipText(this, &SNiagaraSystemOverviewItemName::GetItemToolTip)
 				.IsEnabled_UObject(StackItem.Get(), &UNiagaraStackEntry::GetIsEnabledAndOwnerIsEnabled)
 				.ColorAndOpacity(FSlateColor::UseForeground())
 			];
+
+			if(StackItem->UsesCustomTooltipWidget())
+			{
+				NameTextBlock->SetToolTip(TAttribute<TSharedPtr<IToolTip>>::Create(
+					TAttribute<TSharedPtr<IToolTip>>::FGetter::CreateUObject(StackItem.Get(), &UNiagaraStackEntry::GetCustomTooltipWidget)));
+			}
+			else
+			{
+				NameTextBlock->SetToolTipText(TAttribute<FText>::Create(TAttribute<FText>::FGetter::CreateUObject(StackItem.Get(), &UNiagaraStackEntry::GetTooltipText)));
+			}
+			
 			UpdateTextStyle();
 		}
 	}
@@ -584,7 +605,7 @@ private:
 		return FText::GetEmpty();
 	}
 
-	FText GetItemToolTip() const
+	FText GetItemToolTipText() const
 	{
 		if (StackItem.IsValid() && StackItem->IsFinalized() == false)
 		{
diff --git a/Engine/Plugins/FX/Niagara/Source/NiagaraEditorWidgets/Private/Stack/SNiagaraStackFunctionInputValue.cpp b/Engine/Plugins/FX/Niagara/Source/NiagaraEditorWidgets/Private/Stack/SNiagaraStackFunctionInputValue.cpp
--- a/Engine/Plugins/FX/Niagara/Source/NiagaraEditorWidgets/Private/Stack/SNiagaraStackFunctionInputValue.cpp
+++ b/Engine/Plugins/FX/Niagara/Source/NiagaraEditorWidgets/Private/Stack/SNiagaraStackFunctionInputValue.cpp
@@ -1306,7 +1306,8 @@ TArray<TSharedPtr<FNiagaraMenuAction_Generic>> SNiagaraStackFunctionInputValue::
                 FNiagaraMenuAction_Generic::FOnExecuteAction::CreateSP(this, &SNiagaraStackFunctionInputValue::DynamicInputScriptSelected, DynamicInputScript),
                 DisplayName, {CategoryName.ToString()}, FavoritesActionData, Tooltip, ScriptData->Keywords));
 			DynamicInputAction->FavoritesActionData = FavoritesActionData;
-			DynamicInputAction->SourceData = FNiagaraActionSourceData(Source.Key, Source.Value, true);			
+			DynamicInputAction->SourceData = FNiagaraActionSourceData(Source.Key, Source.Value, true);
+			DynamicInputAction->PreviewMoviePath = FNiagaraEditorUtilities::Preview::GetPreviewMovieObjectPath(FAssetData(DynamicInputScript));
 			
 			DynamicInputAction->bIsExperimental = ScriptData->bExperimental;
 			DynamicInputAction->bIsInLibrary = bIsInLibrary;
@@ -1583,6 +1584,7 @@ TSharedRef<SWidget> SNiagaraStackFunctionInputValue::OnGenerateWidgetForItem(
 	FCreateNiagaraWidgetForActionData ActionData(Item);
 	ActionData.HighlightText = TAttribute<FText>::Create(TAttribute<FText>::FGetter::CreateRaw(this, &SNiagaraStackFunctionInputValue::GetFilterText));
 	ActionData.FavoriteActionsProfileName = FavoriteActionsProfile;
+	ActionData.PreviewMovieViewModel = FunctionInput->GetSystemViewModel()->GetPreviewMovieViewModel();
 	return SNew(SNiagaraActionWidget, ActionData)
 		.bShowTypeIfParameter(false);
 }
diff --git a/Engine/Plugins/FX/Niagara/Source/NiagaraEditorWidgets/Private/Stack/SNiagaraStackItemGroupAddMenu.cpp b/Engine/Plugins/FX/Niagara/Source/NiagaraEditorWidgets/Private/Stack/SNiagaraStackItemGroupAddMenu.cpp
--- a/Engine/Plugins/FX/Niagara/Source/NiagaraEditorWidgets/Private/Stack/SNiagaraStackItemGroupAddMenu.cpp
+++ b/Engine/Plugins/FX/Niagara/Source/NiagaraEditorWidgets/Private/Stack/SNiagaraStackItemGroupAddMenu.cpp
@@ -123,6 +123,7 @@ TArray<TSharedPtr<FNiagaraMenuAction_Generic>> SNiagaraStackItemGroupAddMenu::Co
             AddAction->GetDisplayName(), AddAction->GetCategories(), AddAction->GetFavoritesData(), AddAction->GetDescription(), AddAction->GetKeywords()));
 		NewNodeAction->SourceData = AddAction->GetSourceData();
 		NewNodeAction->bIsInLibrary = AddAction->IsInLibrary();
+		NewNodeAction->PreviewMoviePath = AddAction->GetPreviewMoviePath();
 		
 		OutActions.Add(NewNodeAction);
 	}
@@ -264,6 +265,7 @@ TSharedRef<SWidget> SNiagaraStackItemGroupAddMenu::OnGenerateWidgetForItem(
 	FCreateNiagaraWidgetForActionData ActionData(Item);
 	ActionData.FavoriteActionsProfileName = INiagaraStackItemGroupAddAction::FavoriteActionsProfile;
 	ActionData.HighlightText = TAttribute<FText>::Create(TAttribute<FText>::FGetter::CreateRaw(this, &SNiagaraStackItemGroupAddMenu::GetFilterText));
+	ActionData.PreviewMovieViewModel = SourceEntry.IsValid() ? SourceEntry->GetSystemViewModel()->GetPreviewMovieViewModel() : nullptr;
 	return SNew(SNiagaraActionWidget, ActionData);
 }
 
diff --git a/Engine/Plugins/FX/NiagaraPreviewContent/NiagaraPreviewContent.uplugin b/Engine/Plugins/FX/NiagaraPreviewContent/NiagaraPreviewContent.uplugin
new file mode 100644
--- /dev/null
+++ b/Engine/Plugins/FX/NiagaraPreviewContent/NiagaraPreviewContent.uplugin
@@ -0,0 +1,24 @@
+﻿{
+	"FileVersion": 3,
+	"Version": 1,
+	"VersionName": "1.0",
+	"FriendlyName": "NiagaraPreviewContent",
+	"Description": "Contains movie files used in Niagara",
+	"Category": "FX",
+	"CreatedBy": "Epic Games, Inc.",
+	"CreatedByURL": "https://epicgames.com",
+	"DocsURL": "",
+	"MarketplaceURL": "",
+	"EnabledByDefault" : true,
+	"CanContainContent": true,
+	"IsBetaVersion": false,
+	"IsExperimentalVersion": false,
+	"Installed": false,
+	"Modules": [
+		{
+			"Name": "NiagaraPreviewContent",
+			"Type": "Editor",
+			"LoadingPhase": "Default"
+		}
+	]
+}
\ No newline at end of file
diff --git a/Engine/Plugins/FX/NiagaraPreviewContent/Source/NiagaraPreviewContent/NiagaraPreviewContent.Build.cs b/Engine/Plugins/FX/NiagaraPreviewContent/Source/NiagaraPreviewContent/NiagaraPreviewContent.Build.cs
new file mode 100644
--- /dev/null
+++ b/Engine/Plugins/FX/NiagaraPreviewContent/Source/NiagaraPreviewContent/NiagaraPreviewContent.Build.cs
@@ -0,0 +1,53 @@
+﻿// Copyright Epic Games, Inc. All Rights Reserved.
+
+using UnrealBuildTool;
+
+public class NiagaraPreviewContent : ModuleRules
+{
+	public NiagaraPreviewContent(ReadOnlyTargetRules Target) : base(Target)
+	{
+		PCHUsage = ModuleRules.PCHUsageMode.UseExplicitOrSharedPCHs;
+		
+		PublicIncludePaths.AddRange(
+			new string[] {
+				// ... add public include paths required here ...
+			}
+			);
+				
+		
+		PrivateIncludePaths.AddRange(
+			new string[] {
+				// ... add other private include paths required here ...
+			}
+			);
+			
+		
+		PublicDependencyModuleNames.AddRange(
+			new string[]
+			{
+				"Core",
+				// ... add other public dependencies that you statically link with here ...
+			}
+			);
+			
+		
+		PrivateDependencyModuleNames.AddRange(
+			new string[]
+			{
+				"CoreUObject",
+				"Engine",
+				"Slate",
+				"SlateCore",
+				// ... add private dependencies that you statically link with here ...	
+			}
+			);
+		
+		
+		DynamicallyLoadedModuleNames.AddRange(
+			new string[]
+			{
+				// ... add any modules that your module loads dynamically here ...
+			}
+			);
+	}
+}
diff --git a/Engine/Plugins/FX/NiagaraPreviewContent/Source/NiagaraPreviewContent/Private/NiagaraPreviewContent.cpp b/Engine/Plugins/FX/NiagaraPreviewContent/Source/NiagaraPreviewContent/Private/NiagaraPreviewContent.cpp
new file mode 100644
--- /dev/null
+++ b/Engine/Plugins/FX/NiagaraPreviewContent/Source/NiagaraPreviewContent/Private/NiagaraPreviewContent.cpp
@@ -0,0 +1,20 @@
+﻿// Copyright Epic Games, Inc. All Rights Reserved.
+
+#include "NiagaraPreviewContent.h"
+
+#define LOCTEXT_NAMESPACE "FNiagaraPreviewContentModule"
+
+void FNiagaraPreviewContentModule::StartupModule()
+{
+	// This code will execute after your module is loaded into memory; the exact timing is specified in the .uplugin file per-module
+}
+
+void FNiagaraPreviewContentModule::ShutdownModule()
+{
+	// This function may be called during shutdown to clean up your module.  For modules that support dynamic reloading,
+	// we call this function before unloading the module.
+}
+
+#undef LOCTEXT_NAMESPACE
+	
+IMPLEMENT_MODULE(FNiagaraPreviewContentModule, NiagaraPreviewContent)
\ No newline at end of file
diff --git a/Engine/Plugins/FX/NiagaraPreviewContent/Source/NiagaraPreviewContent/Public/NiagaraPreviewContent.h b/Engine/Plugins/FX/NiagaraPreviewContent/Source/NiagaraPreviewContent/Public/NiagaraPreviewContent.h
new file mode 100644
--- /dev/null
+++ b/Engine/Plugins/FX/NiagaraPreviewContent/Source/NiagaraPreviewContent/Public/NiagaraPreviewContent.h
@@ -0,0 +1,14 @@
+﻿// Copyright Epic Games, Inc. All Rights Reserved.
+
+#pragma once
+
+#include "Modules/ModuleManager.h"
+
+class FNiagaraPreviewContentModule : public IModuleInterface
+{
+public:
+
+	/** IModuleInterface implementation */
+	virtual void StartupModule() override;
+	virtual void ShutdownModule() override;
+};
