Fixed 'Set New Or Existing Parameters' module action crashing due to uninitialized member

diff --git a/Engine/Plugins/FX/Niagara/Source/NiagaraEditor/Private/ViewModels/Stack/NiagaraStackScriptItemGroup.cpp b/Engine/Plugins/FX/Niagara/Source/NiagaraEditor/Private/ViewModels/Stack/NiagaraStackScriptItemGroup.cpp
--- a/Engine/Plugins/FX/Niagara/Source/NiagaraEditor/Private/ViewModels/Stack/NiagaraStackScriptItemGroup.cpp
+++ b/Engine/Plugins/FX/Niagara/Source/NiagaraEditor/Private/ViewModels/Stack/NiagaraStackScriptItemGroup.cpp
@@ -67,7 +67,7 @@ public:
 		bool bIsInLibrary = false;
 		FNiagaraVariable ModuleParameterVariable;
 		FAssetData ModuleAssetData;
-		UNiagaraScript* ModuleScript;
+		UNiagaraScript* ModuleScript = nullptr;
 		TSharedPtr<FModuleWizardModel> ModuleWizardModel;
 		TOptional<ESpecificActionType> SpecificActionType;
 		FNiagaraActionSourceData SourceData = FNiagaraActionSourceData(EScriptSource::Niagara, FText::FromString("Niagara"), true);
@@ -348,7 +348,7 @@ private:
 	bool bIsInLibrary;
 	FNiagaraVariable ModuleParameterVariable;
 	FAssetData ModuleAssetData;
-	UNiagaraScript* ModuleScript;
+	UNiagaraScript* ModuleScript = nullptr;
 	TSharedPtr<FModuleWizardModel> ModuleWizardModel;
 	TOptional<ESpecificActionType> SpecificActionType;
 	FNiagaraActionSourceData SourceData;
