Fixed crash due to invalid view model

diff --git a/Engine/Plugins/FX/Niagara/Source/NiagaraEditor/Private/Widgets/SNiagaraGraphActionMenu.cpp b/Engine/Plugins/FX/Niagara/Source/NiagaraEditor/Private/Widgets/SNiagaraGraphActionMenu.cpp
--- a/Engine/Plugins/FX/Niagara/Source/NiagaraEditor/Private/Widgets/SNiagaraGraphActionMenu.cpp
+++ b/Engine/Plugins/FX/Niagara/Source/NiagaraEditor/Private/Widgets/SNiagaraGraphActionMenu.cpp
@@ -197,7 +197,7 @@ TSharedRef<SWidget> SNiagaraGraphActionMenu::OnGenerateWidgetForItem(const TShar
 	FCreateNiagaraWidgetForActionData ActionData(Item);
 	ActionData.HighlightText = TAttribute<FText>::Create(TAttribute<FText>::FGetter::CreateRaw(this, &SNiagaraGraphActionMenu::GetFilterText));
 	ActionData.FavoriteActionsProfileName = SNiagaraGraphActionMenu::FavoriteActionsProfile;
-	ActionData.PreviewMovieViewModel = ViewModel->GetPreviewMovieViewModel();
+	ActionData.PreviewMovieViewModel = ViewModel ? ViewModel->GetPreviewMovieViewModel() : nullptr;
 	return SNew(SNiagaraActionWidget, ActionData);
 }
 
