Niagara - Fix traversal cache pin warnings and enable the traversal cache by default.

diff --git a/Engine/Plugins/FX/Niagara/Source/NiagaraEditor/Private/NiagaraEditorModule.cpp b/Engine/Plugins/FX/Niagara/Source/NiagaraEditor/Private/NiagaraEditorModule.cpp
--- a/Engine/Plugins/FX/Niagara/Source/NiagaraEditor/Private/NiagaraEditorModule.cpp
+++ b/Engine/Plugins/FX/Niagara/Source/NiagaraEditor/Private/NiagaraEditorModule.cpp
@@ -232,7 +232,7 @@ static FAutoConsoleVariableRef CVarEnableCustomInlineDynamicInputFormats(
 	TEXT("If > 0 and experimental inline editors for dynamic input trees are enabled, custom formats which are defined on scripts will be applied.\n"),
 	ECVF_Default);
 
-int32 GbEnableTraversalCache = 0;
+int32 GbEnableTraversalCache = 1;
 static FAutoConsoleVariableRef CVarEnableTraversalCache(
 	TEXT("fx.Niagara.EnableTraversalCache"),
 	GbEnableTraversalCache,
diff --git a/Engine/Plugins/FX/Niagara/Source/NiagaraEditor/Private/TraversalCache/TraversalBuilder.cpp b/Engine/Plugins/FX/Niagara/Source/NiagaraEditor/Private/TraversalCache/TraversalBuilder.cpp
--- a/Engine/Plugins/FX/Niagara/Source/NiagaraEditor/Private/TraversalCache/TraversalBuilder.cpp
+++ b/Engine/Plugins/FX/Niagara/Source/NiagaraEditor/Private/TraversalCache/TraversalBuilder.cpp
@@ -81,7 +81,11 @@ namespace BuilderHelpers
 	// also undesired.
 	FNiagaraVariable TraversalPinToNiagaraVariable(const UEdGraphPin& InPin)
 	{
-		bool bNeedsValue = InPin.Direction == EGPD_Input && InPin.bDefaultValueIsIgnored == false && InPin.LinkedTo.Num() == 0;
+		bool bNeedsValue = InPin.Direction == EGPD_Input &&
+			InPin.bDefaultValueIsIgnored == false &&
+			InPin.LinkedTo.Num() == 0 &&
+			(InPin.PinType.PinCategory == UEdGraphSchema_Niagara::PinCategoryType ||
+				InPin.PinType.PinCategory == UEdGraphSchema_Niagara::PinCategoryStaticType);
 		FNiagaraVariable Variable = UEdGraphSchema_Niagara::PinToNiagaraVariable(&InPin, bNeedsValue);
 		return Variable;
 	}
