- Fix for nullptr access crash
- The entry is sometimes finalized and has a null view model ptr

#tests delete a module from an emitter
#rnx

diff --git a/Engine/Plugins/FX/Niagara/Source/NiagaraEditor/Private/ViewModels/Stack/NiagaraStackModuleItem.cpp b/Engine/Plugins/FX/Niagara/Source/NiagaraEditor/Private/ViewModels/Stack/NiagaraStackModuleItem.cpp
--- a/Engine/Plugins/FX/Niagara/Source/NiagaraEditor/Private/ViewModels/Stack/NiagaraStackModuleItem.cpp
+++ b/Engine/Plugins/FX/Niagara/Source/NiagaraEditor/Private/ViewModels/Stack/NiagaraStackModuleItem.cpp
@@ -130,21 +130,24 @@ TSharedPtr<IToolTip> UNiagaraStackModuleItem::GetCustomTooltipWidget() const
 	using namespace FNiagaraEditorUtilities::Tooltips;
 	if(FunctionCallNode != nullptr && FunctionCallNode->FunctionScript != nullptr)
 	{
-		FAssetData AssetData(FunctionCallNode->FunctionScript);
-		if(GetDefault<UNiagaraEditorSettings>()->bDisplayPreviewMoviesInTooltips && GetSystemViewModel()->GetPreviewMovieViewModel()->IsPlayingPreviewMovieForAsset(AssetData) == false)
+		if (TSharedPtr<FNiagaraSystemViewModel> SystemViewModelPtr = UNiagaraStackEntry::GetSystemViewModelPtr())
 		{
-			GetSystemViewModel()->GetPreviewMovieViewModel()->PlayMovieForAsset(AssetData);
+			FAssetData AssetData(FunctionCallNode->FunctionScript);
+			if (GetDefault<UNiagaraEditorSettings>()->bDisplayPreviewMoviesInTooltips && SystemViewModelPtr->GetPreviewMovieViewModel()->IsPlayingPreviewMovieForAsset(AssetData) == false)
+			{
+				SystemViewModelPtr->GetPreviewMovieViewModel()->PlayMovieForAsset(AssetData);
+			}
+
+			if (TooltipCache.IsValid() == false)
+			{
+				FCreateTooltipArguments CreateTooltipArguments;
+				CreateTooltipArguments.bTryCreatePreviewMovie = GetDefault<UNiagaraEditorSettings>()->bDisplayPreviewMoviesInTooltips;
+				CreateTooltipArguments.PreviewMovieOwner = GetSystemViewModel()->GetPreviewMovieViewModel();
+				TooltipCache = CreateTooltipForScript(*FunctionCallNode->FunctionScript, FunctionCallNode->SelectedScriptVersion, CreateTooltipArguments);
+			}
+
+			return TooltipCache.ToSharedRef();
 		}
-		
-		if(TooltipCache.IsValid() == false)
-		{
-			FCreateTooltipArguments CreateTooltipArguments;
-			CreateTooltipArguments.bTryCreatePreviewMovie = GetDefault<UNiagaraEditorSettings>()->bDisplayPreviewMoviesInTooltips;
-			CreateTooltipArguments.PreviewMovieOwner = GetSystemViewModel()->GetPreviewMovieViewModel();
-			TooltipCache = CreateTooltipForScript(*FunctionCallNode->FunctionScript, FunctionCallNode->SelectedScriptVersion, CreateTooltipArguments);
-		}		
-		
-		return TooltipCache.ToSharedRef();
 	}
 	
 	return nullptr; 
