- Fix stateless scale by speed when not applying a curve

#tests make a stateless emitter that uses scale by modules and don't apply a curve
#rnx

#virtualized

diff --git a/Engine/Plugins/FX/Niagara/Source/Niagara/Private/Stateless/Modules/NiagaraStatelessModule_ScaleMeshSizeBySpeed.cpp b/Engine/Plugins/FX/Niagara/Source/Niagara/Private/Stateless/Modules/NiagaraStatelessModule_ScaleMeshSizeBySpeed.cpp
--- a/Engine/Plugins/FX/Niagara/Source/Niagara/Private/Stateless/Modules/NiagaraStatelessModule_ScaleMeshSizeBySpeed.cpp
+++ b/Engine/Plugins/FX/Niagara/Source/Niagara/Private/Stateless/Modules/NiagaraStatelessModule_ScaleMeshSizeBySpeed.cpp
@@ -90,7 +90,7 @@ void UNiagaraStatelessModule_ScaleMeshSizeBySpeed::BuildEmitterData(const FNiaga
 	}
 	else
 	{
-		const float Values[] = { 1.0f, 1.0f };
+		const float Values[] = { 0.0f, 1.0f };
 		BuiltData->ScaleDistribution.X = BuildContext.AddStaticData(Values);
 		BuiltData->ScaleDistribution.Y = 1;
 	}
diff --git a/Engine/Plugins/FX/Niagara/Source/Niagara/Private/Stateless/Modules/NiagaraStatelessModule_ScaleSpriteSizeBySpeed.cpp b/Engine/Plugins/FX/Niagara/Source/Niagara/Private/Stateless/Modules/NiagaraStatelessModule_ScaleSpriteSizeBySpeed.cpp
--- a/Engine/Plugins/FX/Niagara/Source/Niagara/Private/Stateless/Modules/NiagaraStatelessModule_ScaleSpriteSizeBySpeed.cpp
+++ b/Engine/Plugins/FX/Niagara/Source/Niagara/Private/Stateless/Modules/NiagaraStatelessModule_ScaleSpriteSizeBySpeed.cpp
@@ -88,7 +88,7 @@ void UNiagaraStatelessModule_ScaleSpriteSizeBySpeed::BuildEmitterData(const FNia
 	}
 	else
 	{
-		const float Values[] = { 1.0f, 1.0f };
+		const float Values[] = { 0.0f, 1.0f };
 		BuiltData->ScaleDistribution.X = BuildContext.AddStaticData(Values);
 		BuiltData->ScaleDistribution.Y = 1;
 	}
