- Always refresh the columns to avoid getting out of sync with visibile state

#tests Capture a cache, toggle attributes on / off
#rnx

diff --git a/Engine/Plugins/FX/Niagara/Source/NiagaraEditor/Private/Widgets/SNiagaraSimCacheView.cpp b/Engine/Plugins/FX/Niagara/Source/NiagaraEditor/Private/Widgets/SNiagaraSimCacheView.cpp
--- a/Engine/Plugins/FX/Niagara/Source/NiagaraEditor/Private/Widgets/SNiagaraSimCacheView.cpp
+++ b/Engine/Plugins/FX/Niagara/Source/NiagaraEditor/Private/Widgets/SNiagaraSimCacheView.cpp
@@ -299,60 +299,60 @@ void SNiagaraSimCacheView::GenerateColumns()
 	}
 
 	// Do we need to update our columns?
+	bool bRebuildColumns = true;
+
 	TConstArrayView<FNiagaraSimCacheViewModel::FComponentInfo> SelectedComponents = SimCacheViewModel->GetSelectedComponentInfos();
 	if (HeaderRowWidget->GetColumns().Num() == SelectedComponents.Num() + 1)
 	{
-		bool bNeedsRefresh = false;
+		bRebuildColumns = false;
 		for (int32 i = 0; i < SelectedComponents.Num(); ++i)
 		{
 			if (HeaderRowWidget->GetColumns()[i + 1].ColumnId != SelectedComponents[i].Name)
 			{
-				bNeedsRefresh = true;
+				bRebuildColumns = true;
 				break;
 			}
 		}
-
-		if (bNeedsRefresh == false)
-		{
-			return;
-		}
 	}
 
 	//  Give columns a width to prevent them from being shrunk when filtering. 
-	constexpr float ManualWidth = 125.0f;
-	HeaderRowWidget->ClearColumns();
-
-	if (SimCacheViewModel->IsCacheValid())
+	if (bRebuildColumns)
 	{
-		// Generate instance count column
-		HeaderRowWidget->AddColumn(
-			SHeaderRow::Column(NAME_Instance)
-			.DefaultLabel(FText::FromName(NAME_Instance))
-			.HAlignHeader(EHorizontalAlignment::HAlign_Center)
-			.VAlignHeader(EVerticalAlignment::VAlign_Fill)
-			.HAlignCell(EHorizontalAlignment::HAlign_Center)
-			.VAlignCell(EVerticalAlignment::VAlign_Fill)
-			.ManualWidth(ManualWidth)
-			.SortMode(this, &SNiagaraSimCacheView::GetColumnSortMode, NAME_Instance)
-			.OnSort(this, &SNiagaraSimCacheView::OnColumnNameSortModeChanged)
-		);
+		constexpr float ManualWidth = 125.0f;
+		HeaderRowWidget->ClearColumns();
 
-		// Generate a column for each component
-		for (const FNiagaraSimCacheViewModel::FComponentInfo& ComponentInfo : SelectedComponents)
+		if (SimCacheViewModel->IsCacheValid())
 		{
+			// Generate instance count column
 			HeaderRowWidget->AddColumn(
-				SHeaderRow::Column(ComponentInfo.Name)
-				.DefaultLabel(FText::FromName(ComponentInfo.Name))
+				SHeaderRow::Column(NAME_Instance)
+				.DefaultLabel(FText::FromName(NAME_Instance))
 				.HAlignHeader(EHorizontalAlignment::HAlign_Center)
 				.VAlignHeader(EVerticalAlignment::VAlign_Fill)
 				.HAlignCell(EHorizontalAlignment::HAlign_Center)
 				.VAlignCell(EVerticalAlignment::VAlign_Fill)
-				.FillWidth(1.0f)
 				.ManualWidth(ManualWidth)
-				.ShouldGenerateWidget(TAttribute<bool>::Create(TAttribute<bool>::FGetter::CreateSP(this, &SNiagaraSimCacheView::GetShouldGenerateWidget, ComponentInfo.Name)))
-				.SortMode(this, &SNiagaraSimCacheView::GetColumnSortMode, ComponentInfo.Name)
+				.SortMode(this, &SNiagaraSimCacheView::GetColumnSortMode, NAME_Instance)
 				.OnSort(this, &SNiagaraSimCacheView::OnColumnNameSortModeChanged)
 			);
+
+			// Generate a column for each component
+			for (const FNiagaraSimCacheViewModel::FComponentInfo& ComponentInfo : SelectedComponents)
+			{
+				HeaderRowWidget->AddColumn(
+					SHeaderRow::Column(ComponentInfo.Name)
+					.DefaultLabel(FText::FromName(ComponentInfo.Name))
+					.HAlignHeader(EHorizontalAlignment::HAlign_Center)
+					.VAlignHeader(EVerticalAlignment::VAlign_Fill)
+					.HAlignCell(EHorizontalAlignment::HAlign_Center)
+					.VAlignCell(EVerticalAlignment::VAlign_Fill)
+					.FillWidth(1.0f)
+					.ManualWidth(ManualWidth)
+					.ShouldGenerateWidget(TAttribute<bool>::Create(TAttribute<bool>::FGetter::CreateSP(this, &SNiagaraSimCacheView::GetShouldGenerateWidget, ComponentInfo.Name)))
+					.SortMode(this, &SNiagaraSimCacheView::GetColumnSortMode, ComponentInfo.Name)
+					.OnSort(this, &SNiagaraSimCacheView::OnColumnNameSortModeChanged)
+				);
+			}
 		}
 	}
 
