Fixed issue with undo/redo on NDC types.

diff --git a/Engine/Plugins/FX/Niagara/Source/Niagara/Private/NiagaraDataChannel.cpp b/Engine/Plugins/FX/Niagara/Source/Niagara/Private/NiagaraDataChannel.cpp
--- a/Engine/Plugins/FX/Niagara/Source/Niagara/Private/NiagaraDataChannel.cpp
+++ b/Engine/Plugins/FX/Niagara/Source/Niagara/Private/NiagaraDataChannel.cpp
@@ -1644,6 +1644,8 @@ UNiagaraDataChannelReader* UNiagaraDataChannelLibrary::CreateDataChannelReader(c
 
 void UNiagaraDataChannelAsset::PreEditChange(class FProperty* PropertyAboutToChange)
 {
+	Super::PreEditChange(PropertyAboutToChange);
+
 	if(PropertyAboutToChange->GetFName() == GET_MEMBER_NAME_CHECKED(UNiagaraDataChannelAsset, DataChannel) && DataChannel)
 	{
 		CachedPreChangeDataChannel = DataChannel;
@@ -1652,6 +1654,8 @@ void UNiagaraDataChannelAsset::PreEditChange(class FProperty* PropertyAboutToCha
 
 void UNiagaraDataChannelAsset::PostEditChangeProperty(FPropertyChangedEvent& PropertyChangedEvent)
 {
+	Super::PostEditChangeProperty(PropertyChangedEvent);
+
 	if(CachedPreChangeDataChannel && DataChannel)
 	{
 		UEngine::FCopyPropertiesForUnrelatedObjectsParams Params;
