Work around for emitter nodes being potentially incorrectly disabled by explicitly syncing their state with the actual emitters in the new compilation mode.
#tests Looked at asset which wouldn't work until you triggered a refresh of the emitter being enabled

diff --git a/Engine/Plugins/FX/Niagara/Source/NiagaraEditor/Private/NiagaraSystemCompilingManager.cpp b/Engine/Plugins/FX/Niagara/Source/NiagaraEditor/Private/NiagaraSystemCompilingManager.cpp
--- a/Engine/Plugins/FX/Niagara/Source/NiagaraEditor/Private/NiagaraSystemCompilingManager.cpp
+++ b/Engine/Plugins/FX/Niagara/Source/NiagaraEditor/Private/NiagaraSystemCompilingManager.cpp
@@ -6,6 +6,8 @@
 #include "Interfaces/ITargetPlatformManagerModule.h"
 #include "NiagaraCompilationTasks.h"
 #include "NiagaraEditorModule.h"
+#include "NiagaraNodeEmitter.h"
+#include "NiagaraScriptSource.h"
 #include "NiagaraShaderType.h"
 #include "Misc/ScopeRWLock.h"
 #include "ProfilingDebugging/CookStats.h"
@@ -15,7 +17,7 @@
 
 #define LOCTEXT_NAMESPACE "NiagaraCompilationManager"
 
-static int GNiagaraCompilationMaxActiveTaskCount = 24;
+static int GNiagaraCompilationMaxActiveTaskCount = 48;
 static FAutoConsoleVariableRef CVarNiagaraCompilationMaxActiveTaskCount(
 	TEXT("fx.Niagara.Compilation.MaxActiveTaskCount"),
 	GNiagaraCompilationMaxActiveTaskCount,
@@ -113,6 +115,46 @@ void EnsureGlobalsInitialized(const FNiagaraSystemCompilationTask& CompilationTa
 	GShaderCompilingManager->GetAbsoluteShaderDebugInfoDirectory();
 }
 
+void SyncEmitterEnabledState(UNiagaraSystem* System)
+{
+	if (!System)
+	{
+		return;
+	}
+
+	// use the system scripts to find the appropriate emitter nodes to synchronize
+	TArray<UNiagaraNodeEmitter*> EmitterNodes;
+
+	constexpr int32 ExpectedNodesPerEmitter = 2; // spawn + update script
+	EmitterNodes.Reserve(ExpectedNodesPerEmitter * System->GetEmitterHandles().Num());
+
+	for (UNiagaraScript* SystemScript : { System->GetSystemSpawnScript(), System->GetSystemUpdateScript() })
+	{
+		FVersionedNiagaraScriptData* ScriptData = SystemScript ? SystemScript->GetLatestScriptData() : nullptr;
+		if (UNiagaraScriptSource* ScriptSource = Cast<UNiagaraScriptSource>(ScriptData ? ScriptData->GetSource() : nullptr))
+		{
+			if (ScriptSource->NodeGraph)
+			{
+				TArray<UNiagaraNodeEmitter*> GraphEmitterNodes;
+				ScriptSource->NodeGraph->GetNodesOfClass(GraphEmitterNodes);
+
+				for (UNiagaraNodeEmitter* GraphEmitterNode : GraphEmitterNodes)
+				{
+					if (GraphEmitterNode)
+					{
+						EmitterNodes.AddUnique(GraphEmitterNode);
+					}
+				}
+			}
+		}
+	}
+
+	for (UNiagaraNodeEmitter* EmitterNode : EmitterNodes)
+	{
+		EmitterNode->SyncEnabledState();
+	}
+}
+
 };
 
 FNiagaraSystemCompilingManager& FNiagaraSystemCompilingManager::Get()
@@ -395,6 +437,12 @@ FNiagaraCompilationTaskHandle FNiagaraSystemCompilingManager::AddSystem(UNiagara
 		bool bRequiresCompilation = false;
 	};
 
+	// before we evaluate the system and it's scripts we need to do a quick validation check on the NiagaraEmitterNodes.  This is done in the original
+	// compilation mode during PrecompileDuplicate, and while not strictly part of the compilation mode we need to be sure to correct any emitters that
+	// are incorrectly disabled.  There's likely a bug somewhere else that is resulting in some nodes sometimes getting stuck in the disabled mode, but
+	// this is the best way we have to clean that up for now.
+	NiagaraSystemCompilingManagerImpl::SyncEmitterEnabledState(System);
+
 	TArray<FPlatformFeatureLevelPair> FeatureLevels;
 	FindOrAddFeatureLevels(CompileOptions, FeatureLevels);
 
