Fix for compilation loop when spawn/update scripts don't line up
#tests Introduced error which generated the uncoverable error and validated that it now generated an appropriate error without continuously compiling

diff --git a/Engine/Plugins/FX/Niagara/Source/Niagara/Classes/NiagaraSystem.h b/Engine/Plugins/FX/Niagara/Source/Niagara/Classes/NiagaraSystem.h
--- a/Engine/Plugins/FX/Niagara/Source/Niagara/Classes/NiagaraSystem.h
+++ b/Engine/Plugins/FX/Niagara/Source/Niagara/Classes/NiagaraSystem.h
@@ -607,6 +607,7 @@ protected:
 		None,
 		RequestPending,
 		RequestPendingOnDemand,
+		RequestPendingForcedRetry,
 	};
 	ERequestCompileStatus RequestCompileStatus = ERequestCompileStatus::None;
 
diff --git a/Engine/Plugins/FX/Niagara/Source/Niagara/Private/NiagaraSystem.cpp b/Engine/Plugins/FX/Niagara/Source/Niagara/Private/NiagaraSystem.cpp
--- a/Engine/Plugins/FX/Niagara/Source/Niagara/Private/NiagaraSystem.cpp
+++ b/Engine/Plugins/FX/Niagara/Source/Niagara/Private/NiagaraSystem.cpp
@@ -3360,8 +3360,9 @@ bool UNiagaraSystem::QueryCompileComplete(bool bWait)
 	const bool ResultsConsistent = CurrentCompilation->ValidateConsistentResults(Options);
 	if (!ResultsConsistent && !CurrentCompilation->WasForced())
 	{
-		RequestCompile(true, nullptr);
-		return false;
+		ForceGraphToRecompileOnNextCheck();
+		RequestCompileStatus = ERequestCompileStatus::RequestPendingForcedRetry;
+		return true;
 	}
 
 	CurrentCompilation->Apply(Options);
@@ -3580,6 +3581,11 @@ bool UNiagaraSystem::RequestCompile(bool bForce, FNiagaraSystemUpdateContext* Op
 	check(IsInGameThread());
 	TRACE_CPUPROFILER_EVENT_SCOPE(UNiagaraSystem::RequestCompile)
 
+	if (RequestCompileStatus == ERequestCompileStatus::RequestPendingForcedRetry)
+	{
+		bForce = true;
+	}
+
 	RequestCompileStatus = ERequestCompileStatus::None;
 
 	// We remove emitters and scripts on dedicated servers, so skip further work.
