Fix for stackcontext with input default binding
#tests Used an input parameter in an module used in an emitter script and had the default value of the parameter being a StackContext variable

diff --git a/Engine/Plugins/FX/Niagara/Source/NiagaraEditor/Private/NiagaraParameterMapHistory.cpp b/Engine/Plugins/FX/Niagara/Source/NiagaraEditor/Private/NiagaraParameterMapHistory.cpp
--- a/Engine/Plugins/FX/Niagara/Source/NiagaraEditor/Private/NiagaraParameterMapHistory.cpp
+++ b/Engine/Plugins/FX/Niagara/Source/NiagaraEditor/Private/NiagaraParameterMapHistory.cpp
@@ -1834,13 +1834,17 @@ int32 TNiagaraParameterMapHistoryBuilder<GraphBridge>::HandleVariableRead(int32
 				TOptional<ENiagaraDefaultMode> DefaultMode = GraphBridge::GetGraphDefaultMode(Graph, AliasedVar, VariableBinding);
 				if (DefaultMode.IsSet() && *DefaultMode == ENiagaraDefaultMode::Binding && VariableBinding.IsValid())
 				{
-					FNiagaraVariable TempVar = FNiagaraVariable(Var.GetType(), VariableBinding.GetName());
-					if (FNiagaraConstants::GetOldPositionTypeVariables().Contains(TempVar))
+					FNiagaraVariable BoundVar = FNiagaraVariable(Var.GetType(), VariableBinding.GetName());
+					if (FNiagaraConstants::GetOldPositionTypeVariables().Contains(BoundVar))
 					{
 						// Old assets often have vector inputs that default bind to what is now a position type. If we detect that, we change the type to prevent a compiler error.
-						TempVar.SetType(FNiagaraTypeDefinition::GetPositionDef());
+						BoundVar.SetType(FNiagaraTypeDefinition::GetPositionDef());
 					}
-					int32 FoundIdxBinding = AddVariableToHistory(History, TempVar, TempVar, nullptr);
+
+					FNiagaraVariable AliasedBoundVar = BoundVar;
+					BoundVar = ResolveAliases(AliasedBoundVar);
+
+					int32 FoundIdxBinding = AddVariableToHistory(History, BoundVar, AliasedBoundVar, nullptr);
 
 					DefaultValue = History.Variables[FoundIdxBinding].GetName().ToString();
 					bUseRapidIterationParam = false;
