- Allow validation rules to be disabled by config file

diff --git a/Engine/Plugins/FX/Niagara/Source/Niagara/Classes/NiagaraValidationRule.h b/Engine/Plugins/FX/Niagara/Source/Niagara/Classes/NiagaraValidationRule.h
--- a/Engine/Plugins/FX/Niagara/Source/Niagara/Classes/NiagaraValidationRule.h
+++ b/Engine/Plugins/FX/Niagara/Source/Niagara/Classes/NiagaraValidationRule.h
@@ -64,14 +64,19 @@ struct FNiagaraValidationContext
 Base class for system validation logic. 
 These allow Niagara systems to be inspected for content validation either at save time or from a commandlet.
 */
-UCLASS(abstract, EditInlineNew, MinimalAPI)
+UCLASS(abstract, EditInlineNew, MinimalAPI, Config=Editor)
 class UNiagaraValidationRule : public UObject
 {
 	GENERATED_BODY()
-public:
 
+public:
+	NIAGARA_API bool IsEnabled() const { return bIsConfigDisabled == false; }
 #if WITH_EDITOR
 	NIAGARA_API virtual void CheckValidity(const FNiagaraValidationContext& Context, TArray<FNiagaraValidationResult>& OutResults) const;
 #endif
-};
 
+protected:
+	// Allows disabling validation rules from config
+	UPROPERTY(Config)
+	bool bIsConfigDisabled = false;
+};
diff --git a/Engine/Plugins/FX/Niagara/Source/Niagara/Classes/NiagaraValidationRuleSet.h b/Engine/Plugins/FX/Niagara/Source/Niagara/Classes/NiagaraValidationRuleSet.h
--- a/Engine/Plugins/FX/Niagara/Source/Niagara/Classes/NiagaraValidationRuleSet.h
+++ b/Engine/Plugins/FX/Niagara/Source/Niagara/Classes/NiagaraValidationRuleSet.h
@@ -15,7 +15,7 @@ class UNiagaraValidationRuleSet : public UObject
 public:
 	bool HasAnyRules() const
 	{
-		return ValidationRules.ContainsByPredicate([](UNiagaraValidationRule* Rule) { return Rule != nullptr; });
+		return ValidationRules.ContainsByPredicate([](UNiagaraValidationRule* Rule) { return Rule && Rule->IsEnabled(); });
 	}
 
 	UPROPERTY(EditAnywhere, Category = "Validation", Instanced)
diff --git a/Engine/Plugins/FX/Niagara/Source/NiagaraEditor/Private/NiagaraValidationRules.cpp b/Engine/Plugins/FX/Niagara/Source/NiagaraEditor/Private/NiagaraValidationRules.cpp
--- a/Engine/Plugins/FX/Niagara/Source/NiagaraEditor/Private/NiagaraValidationRules.cpp
+++ b/Engine/Plugins/FX/Niagara/Source/NiagaraEditor/Private/NiagaraValidationRules.cpp
@@ -63,7 +63,7 @@ bool NiagaraValidation::HasValidationRules(UNiagaraSystem* NiagaraSystem)
 		{
 			for (UNiagaraValidationRule* Rule : EffectType->ValidationRules)
 			{
-				if (Rule != nullptr)
+				if (Rule && Rule->IsEnabled())
 				{
 					return true;
 				}
@@ -102,7 +102,7 @@ void NiagaraValidation::ValidateAllRulesInSystem(TSharedPtr<FNiagaraSystemViewMo
 		{
 			for (const UNiagaraValidationRule* ValidationRule : ValidationRules)
 			{
-				if (ValidationRule)
+				if (ValidationRule && ValidationRule->IsEnabled())
 				{
 					ValidationRule->CheckValidity(Context, NiagaraValidationResults);
 				}
@@ -145,7 +145,7 @@ void NiagaraValidation::ValidateAllRulesInSystem(TSharedPtr<FNiagaraSystemViewMo
 				Context.Source = Module;
 				for (UNiagaraValidationRule* ValidationRule : Script->ValidationRules)
 				{
-					if (ValidationRule)
+					if (ValidationRule && ValidationRule->IsEnabled())
 					{
 						ValidationRule->CheckValidity(Context, NiagaraValidationResults);
 					}
