- Add resource name to message about culling / sorting being disabled to help track down the assets

#rnx

diff --git a/Engine/Plugins/FX/Niagara/Source/Niagara/Private/NiagaraRendererMeshes.cpp b/Engine/Plugins/FX/Niagara/Source/Niagara/Private/NiagaraRendererMeshes.cpp
--- a/Engine/Plugins/FX/Niagara/Source/Niagara/Private/NiagaraRendererMeshes.cpp
+++ b/Engine/Plugins/FX/Niagara/Source/Niagara/Private/NiagaraRendererMeshes.cpp
@@ -419,7 +419,7 @@ void FNiagaraRendererMeshes::PrepareParticleMeshRenderData(FRHICommandListBase&
 		// Validate what we setup
 		if (SimTarget == ENiagaraSimTarget::GPUComputeSim)
 		{
-			if (!ensureMsgf(!ParticleMeshRenderData.bNeedsCull || ParticleMeshRenderData.bSortCullOnGpu, TEXT("Culling is requested on GPU but we don't support sorting, this will result in incorrect rendering.")))
+			if (!ensureMsgf(!ParticleMeshRenderData.bNeedsCull || ParticleMeshRenderData.bSortCullOnGpu, TEXT("Culling is requested on GPU but sorting is disabled, this will result in incorrect rendering. Asset(%s)."), *SceneProxy->GetResourceName().ToString()))
 			{
 				ParticleMeshRenderData.bNeedsCull = false;
 				ParticleMeshRenderData.bAllowPerParticleMeshLODs = false;
diff --git a/Engine/Plugins/FX/Niagara/Source/Niagara/Private/NiagaraRendererSprites.cpp b/Engine/Plugins/FX/Niagara/Source/Niagara/Private/NiagaraRendererSprites.cpp
--- a/Engine/Plugins/FX/Niagara/Source/Niagara/Private/NiagaraRendererSprites.cpp
+++ b/Engine/Plugins/FX/Niagara/Source/Niagara/Private/NiagaraRendererSprites.cpp
@@ -309,7 +309,7 @@ void FNiagaraRendererSprites::PrepareParticleSpriteRenderData(FRHICommandListBas
 		// Validate what we setup
 		if (SimTarget == ENiagaraSimTarget::GPUComputeSim)
 		{
-			if (!ensureMsgf(!ParticleSpriteRenderData.bNeedsCull || ParticleSpriteRenderData.bSortCullOnGpu, TEXT("Culling is requested on GPU but we don't support sorting, this will result in incorrect rendering.")))
+			if (!ensureMsgf(!ParticleSpriteRenderData.bNeedsCull || ParticleSpriteRenderData.bSortCullOnGpu, TEXT("Culling is requested on GPU but sorting is disabled, this will result in incorrect rendering. Asset(%s)."), *SceneProxy->GetResourceName().ToString()))
 			{
 				ParticleSpriteRenderData.bNeedsCull = false;
 			}
