Fixed curve toolbar style in Niagara to use the default style for curve editors

diff --git a/Engine/Plugins/FX/Niagara/Source/NiagaraEditorWidgets/Private/DetailCustomizations/NiagaraDataInterfaceCurveDetails.cpp b/Engine/Plugins/FX/Niagara/Source/NiagaraEditorWidgets/Private/DetailCustomizations/NiagaraDataInterfaceCurveDetails.cpp
--- a/Engine/Plugins/FX/Niagara/Source/NiagaraEditorWidgets/Private/DetailCustomizations/NiagaraDataInterfaceCurveDetails.cpp
+++ b/Engine/Plugins/FX/Niagara/Source/NiagaraEditorWidgets/Private/DetailCustomizations/NiagaraDataInterfaceCurveDetails.cpp
@@ -1313,7 +1313,7 @@ void FNiagaraDataInterfaceCurveDetailsBase::CustomizeDetails(IDetailLayoutBuilde
 	}
 
 	FToolBarBuilder ToolBarBuilder(CurveEditor->GetCommands(), FMultiBoxCustomization::None, nullptr, true);
-	ToolBarBuilder.SetStyle(&FNiagaraEditorWidgetsStyle::Get(), "NiagaraEditor.CurveDetails.SlimToolBar");
+	ToolBarBuilder.SetStyle(&FAppStyle::Get(), "CurveEditorToolBar");
 
 	ToolBarBuilder.AddComboButton(
 		FUIAction(), 
diff --git a/Engine/Plugins/FX/Niagara/Source/NiagaraEditorWidgets/Private/NiagaraEditorWidgetsStyle.cpp b/Engine/Plugins/FX/Niagara/Source/NiagaraEditorWidgets/Private/NiagaraEditorWidgetsStyle.cpp
--- a/Engine/Plugins/FX/Niagara/Source/NiagaraEditorWidgets/Private/NiagaraEditorWidgetsStyle.cpp
+++ b/Engine/Plugins/FX/Niagara/Source/NiagaraEditorWidgets/Private/NiagaraEditorWidgetsStyle.cpp
@@ -71,7 +71,6 @@ FNiagaraEditorWidgetsStyle::FNiagaraEditorWidgetsStyle() : FSlateStyleSet("Niaga
 	InitTextStyles();
 	InitRowStyles();
 	InitButtonStyles();
-	InitToolbarStyles();
 	InitInlineDynamicInputBrushes();
 	InitSystemOverviewMisc();
 	InitIcons();
@@ -432,18 +431,6 @@ void FNiagaraEditorWidgetsStyle::InitButtonStyles()
 	Set("NiagaraEditor.SystemOverview.InlineParameterButton.Icon", InlineParameterButtonStyleIcon);
 }
 
-void FNiagaraEditorWidgetsStyle::InitToolbarStyles()
-{
-	FToolBarStyle SlimToolbarStyle = FAppStyle::Get().GetWidgetStyle<FToolBarStyle>("SlimToolBar");
-	SlimToolbarStyle
-		.SetBackgroundPadding(FMargin(4.0f, 0.0f))
-		.SetIconSize(Icon16x16);
-	SlimToolbarStyle.ToggleButton
-		.SetPadding(FMargin(8.0f, 2.0f));
-	SlimToolbarStyle.SetBlockPadding(FMargin(2, 0));
-	Set("NiagaraEditor.CurveDetails.SlimToolBar", SlimToolbarStyle);
-}
-
 void FNiagaraEditorWidgetsStyle::InitInlineDynamicInputBrushes()
 {
 	Set("NiagaraEditor.Stack.InlineDynamicInput.GraphNode.Row.Top",				new BOX_BRUSH("Icons/MultiPartBorderRowTop", FMargin(1.0f / 4.0f)));
diff --git a/Engine/Plugins/FX/Niagara/Source/NiagaraEditorWidgets/Private/NiagaraEditorWidgetsStyle.h b/Engine/Plugins/FX/Niagara/Source/NiagaraEditorWidgets/Private/NiagaraEditorWidgetsStyle.h
--- a/Engine/Plugins/FX/Niagara/Source/NiagaraEditorWidgets/Private/NiagaraEditorWidgetsStyle.h
+++ b/Engine/Plugins/FX/Niagara/Source/NiagaraEditorWidgets/Private/NiagaraEditorWidgetsStyle.h
@@ -26,7 +26,6 @@ private:
 	void InitTextStyles();
 	void InitRowStyles();
 	void InitButtonStyles();
-	void InitToolbarStyles();
 	void InitInlineDynamicInputBrushes();
 	void InitSystemOverviewMisc();
 	void InitIcons();
