fix hlsl compilation warning

diff --git a/Engine/Plugins/FX/Niagara/Shaders/Private/NiagaraDataInterfaceSceneCapture2DTemplate.ush b/Engine/Plugins/FX/Niagara/Shaders/Private/NiagaraDataInterfaceSceneCapture2DTemplate.ush
--- a/Engine/Plugins/FX/Niagara/Shaders/Private/NiagaraDataInterfaceSceneCapture2DTemplate.ush
+++ b/Engine/Plugins/FX/Niagara/Shaders/Private/NiagaraDataInterfaceSceneCapture2DTemplate.ush
@@ -59,7 +59,7 @@ void TextureProjectDepth_{ParameterName}(in float3 Position, in float MipLevel,
 	OutInBounds = bValidateTextureBounds ? all(UVW.xy >= 0.0f && UVW.xy < 1.0f) : true;
 	if ( bPointSample )
 	{
-		const int3 Texel = int3(round(UVW * {ParameterName}_TextureSize.xy), int(round(MipLevel)));
+		const int3 Texel = int3(round(UVW.xy * {ParameterName}_TextureSize.xy), int(round(MipLevel)));
 		OutCaptureDepth = {ParameterName}_Texture.Load(Texel).r;
 	}
 	else
