- Add Shift+D to enable selected emitters (reverse of D which disables selected emitters)

#tests Enable / disable selected emitters

diff --git a/Engine/Plugins/FX/Niagara/Source/NiagaraEditor/Private/NiagaraEditorCommands.cpp b/Engine/Plugins/FX/Niagara/Source/NiagaraEditor/Private/NiagaraEditorCommands.cpp
--- a/Engine/Plugins/FX/Niagara/Source/NiagaraEditor/Private/NiagaraEditorCommands.cpp
+++ b/Engine/Plugins/FX/Niagara/Source/NiagaraEditor/Private/NiagaraEditorCommands.cpp
@@ -52,6 +52,7 @@ void FNiagaraEditorCommands::RegisterCommands()
 	UI_COMMAND(CollapseStackToHeaders, "Collapse to Headers", "Expands all headers and collapse all items.", EUserInterfaceActionType::Button, FInputChord(EModifierKey::Control, EKeys::O));
 
 	UI_COMMAND(IsolateSelectedEmitters, "Isolate selected emitters", "Isolate all currently selected emitters.", EUserInterfaceActionType::Button, FInputChord(EKeys::S));
+	UI_COMMAND(EnableSelectedEmitters, "Enable selected emitters", "Enables all currently selected emitters and recompiles the system.", EUserInterfaceActionType::Button, FInputChord(EModifierKey::Shift, EKeys::D));
 	UI_COMMAND(DisableSelectedEmitters, "Disable selected emitters", "Disables all currently selected emitters and recompiles the system.", EUserInterfaceActionType::Button, FInputChord(EKeys::D));
 	UI_COMMAND(HideDisabledModules, "Collapse disabled modules", "Collapses disabled modules in the stack, so they won't take up as much space.", EUserInterfaceActionType::Check, FInputChord(EModifierKey::Control, EKeys::M));
 	
diff --git a/Engine/Plugins/FX/Niagara/Source/NiagaraEditor/Private/Toolkits/NiagaraSystemToolkit.cpp b/Engine/Plugins/FX/Niagara/Source/NiagaraEditor/Private/Toolkits/NiagaraSystemToolkit.cpp
--- a/Engine/Plugins/FX/Niagara/Source/NiagaraEditor/Private/Toolkits/NiagaraSystemToolkit.cpp
+++ b/Engine/Plugins/FX/Niagara/Source/NiagaraEditor/Private/Toolkits/NiagaraSystemToolkit.cpp
@@ -547,6 +547,14 @@ void FNiagaraSystemToolkit::SetupCommands()
 		}),
 		FCanExecuteAction());
 
+	GetToolkitCommands()->MapAction(
+		FNiagaraEditorCommands::Get().EnableSelectedEmitters,
+		FExecuteAction::CreateLambda([this]()
+			{
+				SystemViewModel->EnableSelectedEmitters();
+			}),
+		FCanExecuteAction());
+
 	GetToolkitCommands()->MapAction(
 		FNiagaraEditorCommands::Get().DisableSelectedEmitters,
 		FExecuteAction::CreateLambda([this]()
diff --git a/Engine/Plugins/FX/Niagara/Source/NiagaraEditor/Private/ViewModels/NiagaraSystemViewModel.cpp b/Engine/Plugins/FX/Niagara/Source/NiagaraEditor/Private/ViewModels/NiagaraSystemViewModel.cpp
--- a/Engine/Plugins/FX/Niagara/Source/NiagaraEditor/Private/ViewModels/NiagaraSystemViewModel.cpp
+++ b/Engine/Plugins/FX/Niagara/Source/NiagaraEditor/Private/ViewModels/NiagaraSystemViewModel.cpp
@@ -2143,6 +2143,32 @@ void FNiagaraSystemViewModel::SetSystenIsolationEnabled(bool bIsolationEnabled)
 	}
 }
 
+void FNiagaraSystemViewModel::EnableEmitters(TArray<FGuid> EmitterHandlesIdsToEnable)
+{
+	FScopedTransaction Transaction(LOCTEXT("EnableSelectedEmitters", "Enable selected emitters"));
+	GetSystem().Modify();
+
+	bool bAnyEmittersEnabled = false;
+
+	for (TSharedRef<FNiagaraEmitterHandleViewModel> EmitterHandle : EmitterHandleViewModels)
+	{
+		if (EmitterHandlesIdsToEnable.Contains(EmitterHandle->GetId()))
+		{
+			const bool bEnabled = EmitterHandle->GetEmitterHandle()->SetIsEnabled(true, GetSystem(), false);
+			bAnyEmittersEnabled = bAnyEmittersEnabled || bEnabled;
+		}
+	}
+
+	if (bAnyEmittersEnabled)
+	{
+		GetSystem().RequestCompile(false);
+	}
+	else
+	{
+		Transaction.Cancel();
+	}
+}
+
 void FNiagaraSystemViewModel::DisableEmitters(TArray<FGuid> EmitterHandlesIdsToDisable)
 {
 	FScopedTransaction Transaction(LOCTEXT("DisableSelectedEmitters", "Disabled selected emitters"));
@@ -2515,6 +2541,12 @@ void FNiagaraSystemViewModel::ToggleIsolateEmitterAndSelectedEmitters(FGuid Orig
 	}
 }
 
+void FNiagaraSystemViewModel::EnableSelectedEmitters()
+{
+	const TArray<FGuid> EmitterHandles = GetSelectionViewModel()->GetSelectedEmitterHandleIds();
+	EnableEmitters(EmitterHandles);
+}
+
 void FNiagaraSystemViewModel::DisableSelectedEmitters()
 {
 	const TArray<FGuid> EmitterHandles = GetSelectionViewModel()->GetSelectedEmitterHandleIds();
diff --git a/Engine/Plugins/FX/Niagara/Source/NiagaraEditor/Public/NiagaraEditorCommands.h b/Engine/Plugins/FX/Niagara/Source/NiagaraEditor/Public/NiagaraEditorCommands.h
--- a/Engine/Plugins/FX/Niagara/Source/NiagaraEditor/Public/NiagaraEditorCommands.h
+++ b/Engine/Plugins/FX/Niagara/Source/NiagaraEditor/Public/NiagaraEditorCommands.h
@@ -78,6 +78,7 @@ public:
 	TSharedPtr<FUICommandInfo> ToggleResetDependentSystems;
 
 	TSharedPtr<FUICommandInfo> IsolateSelectedEmitters;
+	TSharedPtr<FUICommandInfo> EnableSelectedEmitters;
 	TSharedPtr<FUICommandInfo> DisableSelectedEmitters;
 	TSharedPtr<FUICommandInfo> HideDisabledModules;
 
diff --git a/Engine/Plugins/FX/Niagara/Source/NiagaraEditor/Public/ViewModels/NiagaraSystemViewModel.h b/Engine/Plugins/FX/Niagara/Source/NiagaraEditor/Public/ViewModels/NiagaraSystemViewModel.h
--- a/Engine/Plugins/FX/Niagara/Source/NiagaraEditor/Public/ViewModels/NiagaraSystemViewModel.h
+++ b/Engine/Plugins/FX/Niagara/Source/NiagaraEditor/Public/ViewModels/NiagaraSystemViewModel.h
@@ -313,6 +313,7 @@ public:
 	 */
 	NIAGARAEDITOR_API void ToggleIsolateEmitterAndSelectedEmitters(FGuid OriginEmitter);
 	NIAGARAEDITOR_API void IsolateSelectedEmitters();
+	void EnableSelectedEmitters();
 	void DisableSelectedEmitters();
 
 	/** Update the compile for edit mode. */
@@ -369,6 +370,9 @@ public:
 
 	void SetSystenIsolationEnabled(bool bIsolationEnabled);
 
+	/** Enable the supplied emitters. */
+	NIAGARAEDITOR_API void EnableEmitters(TArray<FGuid> EmitterHandlesIdsToEnable);
+
 	/** Disable the supplied emitters. */
 	NIAGARAEDITOR_API void DisableEmitters(TArray<FGuid> EmitterHandlesIdsToDisable);
 
