- Refactor renderer widgets to share code and use small chip

diff --git a/Engine/Plugins/FX/Niagara/Source/Niagara/Private/NiagaraComponentRendererProperties.cpp b/Engine/Plugins/FX/Niagara/Source/Niagara/Private/NiagaraComponentRendererProperties.cpp
--- a/Engine/Plugins/FX/Niagara/Source/Niagara/Private/NiagaraComponentRendererProperties.cpp
+++ b/Engine/Plugins/FX/Niagara/Source/Niagara/Private/NiagaraComponentRendererProperties.cpp
@@ -578,12 +578,6 @@ void UNiagaraComponentRendererProperties::PostEditChangeProperty(struct FPropert
 	Super::PostEditChangeProperty(e);
 }
 
-void UNiagaraComponentRendererProperties::GetRendererWidgets(const FNiagaraEmitterInstance* InEmitter, TArray<TSharedPtr<SWidget>>& OutWidgets, TSharedPtr<FAssetThumbnailPool> InThumbnailPool) const
-{
-	TSharedRef<SWidget> Widget = SNew(SImage).Image(GetStackIcon());
-	OutWidgets.Add(Widget);
-}
-
 void UNiagaraComponentRendererProperties::GetRendererTooltipWidgets(const FNiagaraEmitterInstance* InEmitter, TArray<TSharedPtr<SWidget>>& OutWidgets, TSharedPtr<FAssetThumbnailPool> InThumbnailPool) const
 {
 	TSharedRef<SWidget> Tooltip = SNew(STextBlock)
diff --git a/Engine/Plugins/FX/Niagara/Source/Niagara/Private/NiagaraDecalRendererProperties.cpp b/Engine/Plugins/FX/Niagara/Source/Niagara/Private/NiagaraDecalRendererProperties.cpp
--- a/Engine/Plugins/FX/Niagara/Source/Niagara/Private/NiagaraDecalRendererProperties.cpp
+++ b/Engine/Plugins/FX/Niagara/Source/Niagara/Private/NiagaraDecalRendererProperties.cpp
@@ -285,12 +285,6 @@ const TArray<FNiagaraVariable>& UNiagaraDecalRendererProperties::GetOptionalAttr
 	return Attrs;
 }
 
-void UNiagaraDecalRendererProperties::GetRendererWidgets(const FNiagaraEmitterInstance* InEmitter, TArray<TSharedPtr<SWidget>>& OutWidgets, TSharedPtr<FAssetThumbnailPool> InThumbnailPool) const
-{
-	TSharedRef<SWidget> Widget = SNew(SImage).Image(GetStackIcon());
-	OutWidgets.Add(Widget);
-}
-
 const FSlateBrush* UNiagaraDecalRendererProperties::GetStackIcon() const
 {
 	return FSlateIconFinder::FindIconBrushForClass(UDecalComponent::StaticClass());
diff --git a/Engine/Plugins/FX/Niagara/Source/Niagara/Private/NiagaraLightRendererProperties.cpp b/Engine/Plugins/FX/Niagara/Source/Niagara/Private/NiagaraLightRendererProperties.cpp
--- a/Engine/Plugins/FX/Niagara/Source/Niagara/Private/NiagaraLightRendererProperties.cpp
+++ b/Engine/Plugins/FX/Niagara/Source/Niagara/Private/NiagaraLightRendererProperties.cpp
@@ -246,13 +246,6 @@ const TArray<FNiagaraVariable>& UNiagaraLightRendererProperties::GetOptionalAttr
 	return Attrs;
 }
 
-void UNiagaraLightRendererProperties::GetRendererWidgets(const FNiagaraEmitterInstance* InEmitter, TArray<TSharedPtr<SWidget>>& OutWidgets, TSharedPtr<FAssetThumbnailPool> InThumbnailPool) const
-{
-	TSharedRef<SWidget> LightWidget = SNew(SImage)
-		.Image(FSlateIconFinder::FindIconBrushForClass(GetClass()));
-	OutWidgets.Add(LightWidget);
-}
-
 void UNiagaraLightRendererProperties::GetRendererTooltipWidgets(const FNiagaraEmitterInstance* InEmitter, TArray<TSharedPtr<SWidget>>& OutWidgets, TSharedPtr<FAssetThumbnailPool> InThumbnailPool) const
 {
 	TSharedRef<SWidget> LightTooltip = SNew(STextBlock)
diff --git a/Engine/Plugins/FX/Niagara/Source/Niagara/Private/NiagaraMeshRendererProperties.cpp b/Engine/Plugins/FX/Niagara/Source/Niagara/Private/NiagaraMeshRendererProperties.cpp
--- a/Engine/Plugins/FX/Niagara/Source/Niagara/Private/NiagaraMeshRendererProperties.cpp
+++ b/Engine/Plugins/FX/Niagara/Source/Niagara/Private/NiagaraMeshRendererProperties.cpp
@@ -1145,31 +1145,17 @@ void UNiagaraMeshRendererProperties::GetAdditionalVariables(TArray<FNiagaraVaria
 
 void UNiagaraMeshRendererProperties::GetRendererWidgets(const FNiagaraEmitterInstance* InEmitter, TArray<TSharedPtr<SWidget>>& OutWidgets, TSharedPtr<FAssetThumbnailPool> InThumbnailPool) const
 {	
-	int32 ThumbnailSize = 32;
+	TArray<UObject*> Assets;
 	for(const FNiagaraMeshRendererMeshProperties& MeshProperties : Meshes)
 	{
-		TSharedRef<SWidget> DefaultThumbnailWidget = SNew(SImage)
-		.Image(FSlateIconFinder::FindIconBrushForClass(StaticClass()));
-		
-		TSharedPtr<SWidget> ThumbnailWidget = DefaultThumbnailWidget;
-	
 		UStaticMesh* Mesh = MeshProperties.Mesh;
 		if (Mesh && Mesh->HasValidRenderData())
 		{
-			TSharedPtr<FAssetThumbnail> AssetThumbnail = MakeShareable(new FAssetThumbnail(Mesh, ThumbnailSize, ThumbnailSize, InThumbnailPool));
-			ThumbnailWidget = AssetThumbnail->MakeThumbnailWidget();
+			Assets.Add(Mesh);
 		}
-		
-		OutWidgets.Add(ThumbnailWidget);
 	}
 	
-	if (Meshes.Num() == 0)
-	{
-		TSharedRef<SWidget> DefaultThumbnailWidget = SNew(SImage)
-		.Image(FSlateIconFinder::FindIconBrushForClass(StaticClass()));
-		
-		OutWidgets.Add(DefaultThumbnailWidget);
-	}
+	CreateRendererWidgetsForAssets(Assets, InThumbnailPool, OutWidgets);
 }
 
 void UNiagaraMeshRendererProperties::GetRendererTooltipWidgets(const FNiagaraEmitterInstance* InEmitter, TArray<TSharedPtr<SWidget>>& OutWidgets, TSharedPtr<FAssetThumbnailPool> InThumbnailPool) const
diff --git a/Engine/Plugins/FX/Niagara/Source/Niagara/Private/NiagaraRendererProperties.cpp b/Engine/Plugins/FX/Niagara/Source/Niagara/Private/NiagaraRendererProperties.cpp
--- a/Engine/Plugins/FX/Niagara/Source/Niagara/Private/NiagaraRendererProperties.cpp
+++ b/Engine/Plugins/FX/Niagara/Source/Niagara/Private/NiagaraRendererProperties.cpp
@@ -14,6 +14,7 @@
 #include "Stateless/NiagaraStatelessEmitter.h"
 #include "Stateless/Modules/NiagaraStatelessModule_DynamicMaterialParameters.h"
 
+#include "Algo/Copy.h"
 #include "Interfaces/ITargetPlatform.h"
 #include "Materials/Material.h"
 #include "Materials/MaterialInterface.h"
@@ -23,6 +24,11 @@
 #include "SceneUtils.h"
 #include "Styling/SlateIconFinder.h"
 
+#if WITH_EDITORONLY_DATA
+#include "Widgets/Images/SImage.h"
+#include "AssetThumbnail.h"
+#endif
+
 #include UE_INLINE_GENERATED_CPP_BY_NAME(NiagaraRendererProperties)
 
 #define LOCTEXT_NAMESPACE "UNiagaraRendererProperties"
@@ -765,6 +771,54 @@ FNiagaraVariable UNiagaraRendererProperties::GetBoundAttribute(const FNiagaraVar
 	return FNiagaraVariable();
 }
 
+void UNiagaraRendererProperties::CreateDefaultRendererWidget(TArray<TSharedPtr<SWidget>>& OutWidgets) const
+{
+	TSharedRef<SWidget> Widget = SNew(SImage)
+		.Image(GetStackIcon());
+	OutWidgets.Add(Widget);
+}
+
+void UNiagaraRendererProperties::CreateRendererWidgetsForAssets(TConstArrayView<UObject*> InAssets, TSharedPtr<FAssetThumbnailPool> InThumbnailPool, TArray<TSharedPtr<SWidget>>& OutWidgets) const
+{
+	constexpr int32 ThumbnailSize = 32;
+
+	bool bNeedsDefaultWidget = true;
+	for (UObject* Asset : InAssets)
+	{
+		if (Asset == nullptr)
+		{
+			continue;
+		}
+		bNeedsDefaultWidget = false;
+
+		TSharedPtr<FAssetThumbnail> AssetThumbnail = MakeShareable(new FAssetThumbnail(Asset, ThumbnailSize, ThumbnailSize, InThumbnailPool));
+
+		FAssetThumbnailConfig ThumbnailConfig;
+		ThumbnailConfig.SmallAssetChipBorder = true;
+
+		TSharedRef<SWidget> ThumbnailWidget = AssetThumbnail->MakeThumbnailWidget(ThumbnailConfig);
+		OutWidgets.Add(ThumbnailWidget);
+	}
+
+	if (bNeedsDefaultWidget)
+	{
+		CreateDefaultRendererWidget(OutWidgets);
+	}
+}
+
+void UNiagaraRendererProperties::CreateRendererWidgetsForAssets(TConstArrayView<UMaterialInterface*> InMaterials, TSharedPtr<FAssetThumbnailPool> InThumbnailPool, TArray<TSharedPtr<SWidget>>& OutWidgets) const
+{
+	TArray<UObject*> Assets;
+	Assets.Reserve(InMaterials.Num());
+	Algo::Copy(InMaterials, Assets);
+	CreateRendererWidgetsForAssets(Assets, InThumbnailPool, OutWidgets);
+}
+
+void UNiagaraRendererProperties::GetRendererWidgets(const FNiagaraEmitterInstance* InEmitter, TArray<TSharedPtr<SWidget>>& OutWidgets, TSharedPtr<FAssetThumbnailPool> InThumbnailPool) const
+{
+	CreateDefaultRendererWidget(OutWidgets);
+}
+
 void UNiagaraRendererProperties::GetRendererFeedback(const FVersionedNiagaraEmitter& InEmitter, TArray<FNiagaraRendererFeedback>& OutErrors, TArray<FNiagaraRendererFeedback>& OutWarnings,	TArray<FNiagaraRendererFeedback>& OutInfo) const
 {
 	TArray<FText> Errors;
diff --git a/Engine/Plugins/FX/Niagara/Source/Niagara/Private/NiagaraRibbonRendererProperties.cpp b/Engine/Plugins/FX/Niagara/Source/Niagara/Private/NiagaraRibbonRendererProperties.cpp
--- a/Engine/Plugins/FX/Niagara/Source/Niagara/Private/NiagaraRibbonRendererProperties.cpp
+++ b/Engine/Plugins/FX/Niagara/Source/Niagara/Private/NiagaraRibbonRendererProperties.cpp
@@ -616,26 +616,10 @@ FNiagaraVariable UNiagaraRibbonRendererProperties::GetBoundAttribute(const FNiag
 
 void UNiagaraRibbonRendererProperties::GetRendererWidgets(const FNiagaraEmitterInstance* InEmitter, TArray<TSharedPtr<SWidget>>& OutWidgets, TSharedPtr<FAssetThumbnailPool> InThumbnailPool) const
 {
-	TSharedRef<SWidget> ThumbnailWidget = SNullWidget::NullWidget;
-	int32 ThumbnailSize = 32;
 	TArray<UMaterialInterface*> Materials;
 	GetUsedMaterials(InEmitter, Materials);
-	for (UMaterialInterface* PreviewedMaterial : Materials)
-	{
-		TSharedPtr<FAssetThumbnail> AssetThumbnail = MakeShareable(new FAssetThumbnail(PreviewedMaterial, ThumbnailSize, ThumbnailSize, InThumbnailPool));
-		if (AssetThumbnail)
-		{
-			ThumbnailWidget = AssetThumbnail->MakeThumbnailWidget();
-		}
-		OutWidgets.Add(ThumbnailWidget);
-	}
 
-	if (Materials.Num() == 0)
-	{
-		TSharedRef<SWidget> SpriteWidget = SNew(SImage)
-			.Image(FSlateIconFinder::FindIconBrushForClass(GetClass()));
-		OutWidgets.Add(SpriteWidget);
-	}
+	CreateRendererWidgetsForAssets(Materials, InThumbnailPool, OutWidgets);
 }
 
 void UNiagaraRibbonRendererProperties::GetRendererTooltipWidgets(const FNiagaraEmitterInstance* InEmitter, TArray<TSharedPtr<SWidget>>& OutWidgets, TSharedPtr<FAssetThumbnailPool> InThumbnailPool) const
diff --git a/Engine/Plugins/FX/Niagara/Source/Niagara/Private/NiagaraSpriteRendererProperties.cpp b/Engine/Plugins/FX/Niagara/Source/Niagara/Private/NiagaraSpriteRendererProperties.cpp
--- a/Engine/Plugins/FX/Niagara/Source/Niagara/Private/NiagaraSpriteRendererProperties.cpp
+++ b/Engine/Plugins/FX/Niagara/Source/Niagara/Private/NiagaraSpriteRendererProperties.cpp
@@ -592,29 +592,12 @@ void UNiagaraSpriteRendererProperties::GetAdditionalVariables(TArray<FNiagaraVar
 
 void UNiagaraSpriteRendererProperties::GetRendererWidgets(const FNiagaraEmitterInstance* InEmitter, TArray<TSharedPtr<SWidget>>& OutWidgets, TSharedPtr<FAssetThumbnailPool> InThumbnailPool) const
 {
-	TSharedRef<SWidget> ThumbnailWidget = SNullWidget::NullWidget;
-	int32 ThumbnailSize = 32;
 	TArray<UMaterialInterface*> Materials;
 	GetUsedMaterials(InEmitter, Materials);
-	for (UMaterialInterface* PreviewedMaterial : Materials)
-	{
-		TSharedPtr<FAssetThumbnail> AssetThumbnail = MakeShareable(new FAssetThumbnail(PreviewedMaterial, ThumbnailSize, ThumbnailSize, InThumbnailPool));
-		if (AssetThumbnail)
-		{
-			ThumbnailWidget = AssetThumbnail->MakeThumbnailWidget();
-		}
-		OutWidgets.Add(ThumbnailWidget);
-	}
 
-	if (Materials.Num() == 0)
-	{
-		TSharedRef<SWidget> SpriteWidget = SNew(SImage)
-			.Image(FSlateIconFinder::FindIconBrushForClass(GetClass()));
-		OutWidgets.Add(SpriteWidget);
-	}
+	CreateRendererWidgetsForAssets(Materials, InThumbnailPool, OutWidgets);
 }
 
-
 void UNiagaraSpriteRendererProperties::GetRendererFeedback(const FVersionedNiagaraEmitter& InEmitter, TArray<FText>& OutErrors, TArray<FText>& OutWarnings, TArray<FText>& OutInfo) const
 {
 	Super::GetRendererFeedback(InEmitter, OutErrors, OutWarnings, OutInfo);
diff --git a/Engine/Plugins/FX/Niagara/Source/Niagara/Private/NiagaraVolumeRendererProperties.cpp b/Engine/Plugins/FX/Niagara/Source/Niagara/Private/NiagaraVolumeRendererProperties.cpp
--- a/Engine/Plugins/FX/Niagara/Source/Niagara/Private/NiagaraVolumeRendererProperties.cpp
+++ b/Engine/Plugins/FX/Niagara/Source/Niagara/Private/NiagaraVolumeRendererProperties.cpp
@@ -307,11 +307,6 @@ const TArray<FNiagaraVariable>& UNiagaraVolumeRendererProperties::GetOptionalAtt
 	return Attrs;
 }
 
-void UNiagaraVolumeRendererProperties::GetRendererWidgets(const FNiagaraEmitterInstance* InEmitter, TArray<TSharedPtr<SWidget>>& OutWidgets, TSharedPtr<FAssetThumbnailPool> InThumbnailPool) const
-{
-	OutWidgets.Add(SNew(SImage).Image(FSlateIconFinder::FindIconBrushForClass(GetClass())));
-}
-
 void UNiagaraVolumeRendererProperties::GetRendererTooltipWidgets(const FNiagaraEmitterInstance* InEmitter, TArray<TSharedPtr<SWidget>>& OutWidgets, TSharedPtr<FAssetThumbnailPool> InThumbnailPool) const
 {
 	OutWidgets.Add(SNew(STextBlock).Text(LOCTEXT("VolumeRenderer", "Volume Renderer")));
diff --git a/Engine/Plugins/FX/Niagara/Source/Niagara/Public/NiagaraComponentRendererProperties.h b/Engine/Plugins/FX/Niagara/Source/Niagara/Public/NiagaraComponentRendererProperties.h
--- a/Engine/Plugins/FX/Niagara/Source/Niagara/Public/NiagaraComponentRendererProperties.h
+++ b/Engine/Plugins/FX/Niagara/Source/Niagara/Public/NiagaraComponentRendererProperties.h
@@ -70,7 +70,6 @@ public:
 #if WITH_EDITORONLY_DATA
 	virtual void PostEditChangeProperty(struct FPropertyChangedEvent& e) override;
 
-	virtual void GetRendererWidgets(const FNiagaraEmitterInstance* InEmitter, TArray<TSharedPtr<SWidget>>& OutWidgets, TSharedPtr<FAssetThumbnailPool> InThumbnailPool) const override;
 	virtual void GetRendererTooltipWidgets(const FNiagaraEmitterInstance* InEmitter, TArray<TSharedPtr<SWidget>>& OutWidgets, TSharedPtr<FAssetThumbnailPool> InThumbnailPool) const override;
 	virtual void GetRendererFeedback(const FVersionedNiagaraEmitter& InEmitter, TArray<FNiagaraRendererFeedback>& OutErrors, TArray<FNiagaraRendererFeedback>& OutWarnings, TArray<FNiagaraRendererFeedback>& OutInfo) const override;
 	virtual const FSlateBrush* GetStackIcon() const override;
diff --git a/Engine/Plugins/FX/Niagara/Source/Niagara/Public/NiagaraDecalRendererProperties.h b/Engine/Plugins/FX/Niagara/Source/Niagara/Public/NiagaraDecalRendererProperties.h
--- a/Engine/Plugins/FX/Niagara/Source/Niagara/Public/NiagaraDecalRendererProperties.h
+++ b/Engine/Plugins/FX/Niagara/Source/Niagara/Public/NiagaraDecalRendererProperties.h
@@ -41,7 +41,6 @@ public:
 	virtual bool IsSimTargetSupported(ENiagaraSimTarget InSimTarget) const override { return InSimTarget == ENiagaraSimTarget::CPUSim; };
 #if WITH_EDITORONLY_DATA
 	virtual const TArray<FNiagaraVariable>& GetOptionalAttributes() override;
-	virtual void GetRendererWidgets(const FNiagaraEmitterInstance* InEmitter, TArray<TSharedPtr<SWidget>>& OutWidgets, TSharedPtr<FAssetThumbnailPool> InThumbnailPool) const override;
 	virtual const FSlateBrush* GetStackIcon() const override;
 	virtual void GetRendererTooltipWidgets(const FNiagaraEmitterInstance* InEmitter, TArray<TSharedPtr<SWidget>>& OutWidgets, TSharedPtr<FAssetThumbnailPool> InThumbnailPool) const override;
 	virtual void GetRendererFeedback(const FVersionedNiagaraEmitter& InEmitter, TArray<FNiagaraRendererFeedback>& OutErrors, TArray<FNiagaraRendererFeedback>& OutWarnings, TArray<FNiagaraRendererFeedback>& OutInfo) const override;
diff --git a/Engine/Plugins/FX/Niagara/Source/Niagara/Public/NiagaraLightRendererProperties.h b/Engine/Plugins/FX/Niagara/Source/Niagara/Public/NiagaraLightRendererProperties.h
--- a/Engine/Plugins/FX/Niagara/Source/Niagara/Public/NiagaraLightRendererProperties.h
+++ b/Engine/Plugins/FX/Niagara/Source/Niagara/Public/NiagaraLightRendererProperties.h
@@ -37,7 +37,6 @@ public:
 	virtual bool IsSimTargetSupported(ENiagaraSimTarget InSimTarget) const override { return (InSimTarget == ENiagaraSimTarget::CPUSim); };
 #if WITH_EDITORONLY_DATA
 	virtual const TArray<FNiagaraVariable>& GetOptionalAttributes() override;
-	virtual void GetRendererWidgets(const FNiagaraEmitterInstance* InEmitter, TArray<TSharedPtr<SWidget>>& OutWidgets, TSharedPtr<FAssetThumbnailPool> InThumbnailPool) const override;
 	virtual void GetRendererTooltipWidgets(const FNiagaraEmitterInstance* InEmitter, TArray<TSharedPtr<SWidget>>& OutWidgets, TSharedPtr<FAssetThumbnailPool> InThumbnailPool) const override;
 	virtual void GetRendererFeedback(const FVersionedNiagaraEmitter& InEmitter, TArray<FText>& OutErrors, TArray<FText>& OutWarnings, TArray<FText>& OutInfo) const override;
 #endif // WITH_EDITORONLY_DATA
diff --git a/Engine/Plugins/FX/Niagara/Source/Niagara/Public/NiagaraRendererProperties.h b/Engine/Plugins/FX/Niagara/Source/Niagara/Public/NiagaraRendererProperties.h
--- a/Engine/Plugins/FX/Niagara/Source/Niagara/Public/NiagaraRendererProperties.h
+++ b/Engine/Plugins/FX/Niagara/Source/Niagara/Public/NiagaraRendererProperties.h
@@ -372,7 +372,11 @@ public:
 		return CastChecked<UNiagaraRendererProperties>(Super::StaticDuplicateWithNewMergeIdInternal(InOuter));
 	}
 
-	NIAGARA_API virtual void GetRendererWidgets(const FNiagaraEmitterInstance* InEmitter, TArray<TSharedPtr<SWidget>>& OutWidgets, TSharedPtr<FAssetThumbnailPool> InThumbnailPool) const PURE_VIRTUAL(UNiagaraRendererProperties::GetRendererWidgets, );
+	NIAGARA_API void CreateDefaultRendererWidget(TArray<TSharedPtr<SWidget>>& OutWidgets) const;
+	NIAGARA_API void CreateRendererWidgetsForAssets(TConstArrayView<UObject*> InAssets, TSharedPtr<FAssetThumbnailPool> InThumbnailPool, TArray<TSharedPtr<SWidget>>& OutWidgets) const;
+	NIAGARA_API void CreateRendererWidgetsForAssets(TConstArrayView<UMaterialInterface*> InMaterials, TSharedPtr<FAssetThumbnailPool> InThumbnailPool, TArray<TSharedPtr<SWidget>>& OutWidgets) const;
+
+	NIAGARA_API virtual void GetRendererWidgets(const FNiagaraEmitterInstance* InEmitter, TArray<TSharedPtr<SWidget>>& OutWidgets, TSharedPtr<FAssetThumbnailPool> InThumbnailPool) const;
 	NIAGARA_API virtual void GetRendererTooltipWidgets(const FNiagaraEmitterInstance* InEmitter, TArray<TSharedPtr<SWidget>>& OutWidgets, TSharedPtr<FAssetThumbnailPool> InThumbnailPool) const PURE_VIRTUAL(UNiagaraRendererProperties::GetRendererTooltipWidgets, );
 	virtual void GetRendererFeedback(const FVersionedNiagaraEmitter& InEmitter, TArray<FText>& OutErrors, TArray<FText>& OutWarnings, TArray<FText>& OutInfo) const {}
 	NIAGARA_API virtual void GetRendererFeedback(const FVersionedNiagaraEmitter& InEmitter, TArray<FNiagaraRendererFeedback>& OutErrors, TArray<FNiagaraRendererFeedback>& OutWarnings, TArray<FNiagaraRendererFeedback>& OutInfo) const;
diff --git a/Engine/Plugins/FX/Niagara/Source/Niagara/Public/NiagaraVolumeRendererProperties.h b/Engine/Plugins/FX/Niagara/Source/Niagara/Public/NiagaraVolumeRendererProperties.h
--- a/Engine/Plugins/FX/Niagara/Source/Niagara/Public/NiagaraVolumeRendererProperties.h
+++ b/Engine/Plugins/FX/Niagara/Source/Niagara/Public/NiagaraVolumeRendererProperties.h
@@ -56,7 +56,6 @@ public:
 	virtual bool IsSimTargetSupported(ENiagaraSimTarget InSimTarget) const override { return true; };
 #if WITH_EDITORONLY_DATA
 	virtual const TArray<FNiagaraVariable>& GetOptionalAttributes() override;
-	virtual void GetRendererWidgets(const FNiagaraEmitterInstance* InEmitter, TArray<TSharedPtr<SWidget>>& OutWidgets, TSharedPtr<FAssetThumbnailPool> InThumbnailPool) const override;
 	virtual void GetRendererTooltipWidgets(const FNiagaraEmitterInstance* InEmitter, TArray<TSharedPtr<SWidget>>& OutWidgets, TSharedPtr<FAssetThumbnailPool> InThumbnailPool) const override;
 	virtual void GetRendererFeedback(const FVersionedNiagaraEmitter& InEmitter, TArray<FNiagaraRendererFeedback>& OutErrors, TArray<FNiagaraRendererFeedback>& OutWarnings, TArray<FNiagaraRendererFeedback>& OutInfo) const override;
 	TArray<FNiagaraVariable> GetBoundAttributes() const override;
diff --git a/Engine/Plugins/Runtime/GeometryCache/Source/GeometryCache/Private/NiagaraGeometryCacheRendererProperties.cpp b/Engine/Plugins/Runtime/GeometryCache/Source/GeometryCache/Private/NiagaraGeometryCacheRendererProperties.cpp
--- a/Engine/Plugins/Runtime/GeometryCache/Source/GeometryCache/Private/NiagaraGeometryCacheRendererProperties.cpp
+++ b/Engine/Plugins/Runtime/GeometryCache/Source/GeometryCache/Private/NiagaraGeometryCacheRendererProperties.cpp
@@ -363,28 +363,17 @@ const FSlateBrush* UNiagaraGeometryCacheRendererProperties::GetStackIcon() const
 
 void UNiagaraGeometryCacheRendererProperties::GetRendererWidgets(const FNiagaraEmitterInstance* InEmitter, TArray<TSharedPtr<SWidget>>& OutWidgets, TSharedPtr<FAssetThumbnailPool> InThumbnailPool) const
 {
-	TSharedRef<SWidget> DefaultThumbnailWidget = SNew(SImage)
-		.Image(FSlateIconFinder::FindIconBrushForClass(StaticClass()));
-
-	int32 ThumbnailSize = 32;
+	TArray<UObject*> Assets;
 	for(const FNiagaraGeometryCacheReference& Entry : GeometryCaches)
 	{
-		TSharedPtr<SWidget> ThumbnailWidget = DefaultThumbnailWidget;
-
 		UGeometryCache* GeometryCache = Entry.GeometryCache;
 		if (GeometryCache)
 		{
-			TSharedPtr<FAssetThumbnail> AssetThumbnail = MakeShareable(new FAssetThumbnail(GeometryCache, ThumbnailSize, ThumbnailSize, InThumbnailPool));
-			ThumbnailWidget = AssetThumbnail->MakeThumbnailWidget();
+			Assets.Add(GeometryCache);
 		}
-		
-		OutWidgets.Add(ThumbnailWidget);
 	}
 
-	if (GeometryCaches.Num() == 0)
-	{
-		OutWidgets.Add(DefaultThumbnailWidget);
-	}
+	CreateRendererWidgetsForAssets(Assets, InThumbnailPool, OutWidgets);
 }
 
 void UNiagaraGeometryCacheRendererProperties::GetRendererTooltipWidgets(const FNiagaraEmitterInstance* InEmitter, TArray<TSharedPtr<SWidget>>& OutWidgets, TSharedPtr<FAssetThumbnailPool> AssetThumbnailPool) const
