Niagara - Fix traversal cache initialization issues.
+ Move the initialization of the traversal cache instance from the editor module engine init callback to the editor module startup function.
+ Move the initialization of the traversal builder instance from the editor module to the traversal cache.

diff --git a/Engine/Plugins/FX/Niagara/Source/NiagaraEditor/Private/NiagaraEditorModule.cpp b/Engine/Plugins/FX/Niagara/Source/NiagaraEditor/Private/NiagaraEditorModule.cpp
--- a/Engine/Plugins/FX/Niagara/Source/NiagaraEditor/Private/NiagaraEditorModule.cpp
+++ b/Engine/Plugins/FX/Niagara/Source/NiagaraEditor/Private/NiagaraEditorModule.cpp
@@ -1624,6 +1624,9 @@ void FNiagaraEditorModule::StartupModule()
 	FAssetCompilingManager::Get().RegisterManager(&FNiagaraSystemCompilingManager::Get());
 	ILevelSequenceModule& LevelSequenceModule = FModuleManager::LoadModuleChecked<ILevelSequenceModule>("LevelSequence");
 	DefaultTrackHandle = LevelSequenceModule.OnNewActorTrackAdded().AddStatic(&FNiagaraSystemTrackEditor::AddDefaultSystemTracks);
+
+	TraversalCache = MakeShared<UE::Niagara::TraversalCache::FTraversalCache>();
+	TraversalCache->Initialize();
 }
 
 void FNiagaraEditorModule::ShutdownModule()
@@ -1756,7 +1759,6 @@ void FNiagaraEditorModule::ShutdownModule()
 	FAssetCompilingManager::Get().UnregisterManager(&FNiagaraSystemCompilingManager::Get());
 
 	TraversalCache.Reset();
-	TraversalBuilder.Reset();
 }
 
 void FNiagaraEditorModule::OnPostEngineInit()
@@ -1836,11 +1838,6 @@ void FNiagaraEditorModule::OnPostEngineInit()
 	const FAssetRegistryModule& AssetRegistryModule = FModuleManager::LoadModuleChecked<FAssetRegistryModule>(TEXT("AssetRegistry"));
 	IAssetRegistry& AssetRegistry = AssetRegistryModule.Get();
 	AssetRegistryOnLoadCompleteHandle = AssetRegistry.OnFilesLoaded().AddRaw(this, &FNiagaraEditorModule::OnAssetRegistryLoadComplete);
-
-	TraversalBuilder = MakeShared<UE::Niagara::TraversalCache::FTraversalBuilder>();
-	TraversalBuilder->Initialize();
-	TraversalCache = MakeShared<UE::Niagara::TraversalCache::FTraversalCache>();
-	TraversalCache->Initialize();
 }
 
 void FNiagaraEditorModule::OnDeviceProfileManagerUpdated()
diff --git a/Engine/Plugins/FX/Niagara/Source/NiagaraEditor/Private/TraversalCache/TraversalCache.cpp b/Engine/Plugins/FX/Niagara/Source/NiagaraEditor/Private/TraversalCache/TraversalCache.cpp
--- a/Engine/Plugins/FX/Niagara/Source/NiagaraEditor/Private/TraversalCache/TraversalCache.cpp
+++ b/Engine/Plugins/FX/Niagara/Source/NiagaraEditor/Private/TraversalCache/TraversalCache.cpp
@@ -32,6 +32,8 @@ void FTraversalCache::Initialize()
 		{
 			GEditor->RegisterForUndo(this);
 		}
+		TraversalBuilder = MakeShared<FTraversalBuilder>();
+		TraversalBuilder->Initialize();
 		bInitialized = true;
 	}
 }
diff --git a/Engine/Plugins/FX/Niagara/Source/NiagaraEditor/Private/TraversalCache/TraversalCache.h b/Engine/Plugins/FX/Niagara/Source/NiagaraEditor/Private/TraversalCache/TraversalCache.h
--- a/Engine/Plugins/FX/Niagara/Source/NiagaraEditor/Private/TraversalCache/TraversalCache.h
+++ b/Engine/Plugins/FX/Niagara/Source/NiagaraEditor/Private/TraversalCache/TraversalCache.h
@@ -19,6 +19,7 @@ namespace UE::Niagara::TraversalCache
 {
 
 class  FTraversal;
+class  FTraversalBuilder;
 struct FTraversalBuilderDebugData;
 struct FTraversalDebugData;
 
@@ -227,6 +228,8 @@ private:
 		number is updated it signals the traversals to discard any data collected by visitors so that it can be recomputed when requested. */
 	uint32 TraversalCachedDataSerialNumber = 0;
 
+	TSharedPtr<FTraversalBuilder> TraversalBuilder;
+
 	mutable TMap<FScriptTraversalCacheKey, FScriptTraversalCacheValue> CachedTraversals;
 	mutable TMap<FScriptTraversalCacheKey, FParsedRapidIterationParameterData> CachedParsedRapidIterationParameterData;
 	mutable TMap<FSoftObjectPath, TWeakObjectPtr<const UNiagaraScript>> SoftObjectPathToScriptCache;
diff --git a/Engine/Plugins/FX/Niagara/Source/NiagaraEditor/Public/NiagaraEditorModule.h b/Engine/Plugins/FX/Niagara/Source/NiagaraEditor/Public/NiagaraEditorModule.h
--- a/Engine/Plugins/FX/Niagara/Source/NiagaraEditor/Public/NiagaraEditorModule.h
+++ b/Engine/Plugins/FX/Niagara/Source/NiagaraEditor/Public/NiagaraEditorModule.h
@@ -68,7 +68,6 @@ class ITargetPlatform;
 
 namespace UE::Niagara::TraversalCache
 {
-	class FTraversalBuilder;
 	class FTraversalCache;
 }
 
@@ -315,8 +314,6 @@ public:
 
 	UPackage* GetTempPackage() { return TempPackage; }
 
-
-	UE::Niagara::TraversalCache::FTraversalBuilder& GetParameterTraversalBuilder() { return *TraversalBuilder.Get(); };
 	UE::Niagara::TraversalCache::FTraversalCache& GetParameterTraversalCache() { return *TraversalCache.Get(); }
 
 	struct FMoviePreviewObjects
@@ -575,6 +572,5 @@ private:
 
 	UPackage* TempPackage = nullptr;
 
-	TSharedPtr<UE::Niagara::TraversalCache::FTraversalBuilder> TraversalBuilder;
 	TSharedPtr<UE::Niagara::TraversalCache::FTraversalCache> TraversalCache;
 };
