- Downgrade MIC warning to log
- It didn't stall the cook so we should be able to remove in the future

#rnx

diff --git a/Engine/Plugins/FX/Niagara/Source/Niagara/Private/NiagaraRendererProperties.cpp b/Engine/Plugins/FX/Niagara/Source/Niagara/Private/NiagaraRendererProperties.cpp
--- a/Engine/Plugins/FX/Niagara/Source/Niagara/Private/NiagaraRendererProperties.cpp
+++ b/Engine/Plugins/FX/Niagara/Source/Niagara/Private/NiagaraRendererProperties.cpp
@@ -656,7 +656,7 @@ void UNiagaraRendererProperties::UpdateMaterialParametersMIC(const FNiagaraRende
 					const EInternalObjectFlags InternalFlags = ExistingObject->GetInternalFlags();
 					const EObjectFlags Flags = ExistingObject->GetFlags();
 
-					UE_LOG(LogNiagara, Warning, TEXT("While trying to allocate %s UNiagaraRendererProperties::UpdateMaterialParametersMIC() found a pre-existing object.  This will result in a re-allocation!  FullName: %s | InternalFlags: %x | Flags: %x"),
+					UE_LOG(LogNiagara, Log, TEXT("While trying to allocate %s UNiagaraRendererProperties::UpdateMaterialParametersMIC() found a pre-existing object.  This will result in a re-allocation!  FullName: %s | InternalFlags: %x | Flags: %x"),
 						*NameBuilder, *FullName, InternalFlags, Flags);
 
 					if (UMaterialInstanceConstant* ExistingMIC = Cast<UMaterialInstanceConstant>(ExistingObject))
