- Add menu entry when we have no stateless bindings to make it clear to the user what type we are looking for

#tests Test that both int and other binding types will display information when no bindings are found

diff --git a/Engine/Plugins/FX/Niagara/Source/NiagaraEditor/Private/Stateless/NiagaraDistributionIntPropertyCustomization.cpp b/Engine/Plugins/FX/Niagara/Source/NiagaraEditor/Private/Stateless/NiagaraDistributionIntPropertyCustomization.cpp
--- a/Engine/Plugins/FX/Niagara/Source/NiagaraEditor/Private/Stateless/NiagaraDistributionIntPropertyCustomization.cpp
+++ b/Engine/Plugins/FX/Niagara/Source/NiagaraEditor/Private/Stateless/NiagaraDistributionIntPropertyCustomization.cpp
@@ -377,18 +377,45 @@ class SNiagaraDistributionIntPropertyWidget : public SCompoundWidget
 	TSharedRef<SWidget> OnGetBindingMenuContent()
 	{
 		FMenuBuilder MenuBuilder(true, nullptr);
-		for (const FNiagaraVariableBase& Variable : GetAvailableBindings())
-		{
-			TSharedRef<SWidget> Widget = SNew(SNiagaraParameterName)
-				.ParameterName(Variable.GetName())
-				.IsReadOnly(true);
 
-			MenuBuilder.AddMenuEntry(
-				FUIAction(
-					FExecuteAction::CreateSP(this, &SNiagaraDistributionIntPropertyWidget::SetBindingValue, Variable)
-				),
-				Widget
-			);
+		const TArray<FNiagaraVariableBase> AvailableBindings = GetAvailableBindings();
+		if (AvailableBindings.Num() > 0)
+		{
+			for (const FNiagaraVariableBase& Variable : AvailableBindings)
+			{
+				TSharedRef<SWidget> Widget = SNew(SNiagaraParameterName)
+					.ParameterName(Variable.GetName())
+					.IsReadOnly(true);
+
+				MenuBuilder.AddMenuEntry(
+					FUIAction(
+						FExecuteAction::CreateSP(this, &SNiagaraDistributionIntPropertyWidget::SetBindingValue, Variable)
+					),
+					Widget
+				);
+			}
+		}
+		else
+		{
+			const FNiagaraTypeDefinition AllowedTypeDef = Distribution->GetBindingTypeDef();
+			if (AllowedTypeDef.IsValid())
+			{
+				MenuBuilder.AddMenuEntry(
+					FText::Format(LOCTEXT("NoBindings", "No bindings of type '{0}' found"), AllowedTypeDef.GetNameText()),
+					LOCTEXT("NoBindingsDesc", "No bindings for the property type were found."),
+					FSlateIcon(),
+					FUIAction()
+				);
+			}
+			else
+			{
+				MenuBuilder.AddMenuEntry(
+					LOCTEXT("InvalidTypeDef", "Bindings Unavailable"),
+					LOCTEXT("InvalidTypeDefDesc", "Bindings can not be set on this property."),
+					FSlateIcon(),
+					FUIAction()
+				);
+			}
 		}
 
 		return MenuBuilder.MakeWidget();
diff --git a/Engine/Plugins/FX/Niagara/Source/NiagaraEditor/Private/Widgets/SNiagaraDistributionEditor.cpp b/Engine/Plugins/FX/Niagara/Source/NiagaraEditor/Private/Widgets/SNiagaraDistributionEditor.cpp
--- a/Engine/Plugins/FX/Niagara/Source/NiagaraEditor/Private/Widgets/SNiagaraDistributionEditor.cpp
+++ b/Engine/Plugins/FX/Niagara/Source/NiagaraEditor/Private/Widgets/SNiagaraDistributionEditor.cpp
@@ -154,20 +154,46 @@ private:
 	TSharedRef<SWidget> OnGetMenuContent()
 	{
 		FMenuBuilder MenuBuilder(true, nullptr);
-		for (FNiagaraVariableBase Variable : DistributionAdapter->GetAvailableBindings())
-		{
-			TSharedRef<SWidget> Widget = SNew(SNiagaraParameterName)
-				.ParameterName(Variable.GetName())
-				.IsReadOnly(true);
 
-			MenuBuilder.AddMenuEntry(
-				FUIAction(
-					FExecuteAction::CreateSP(this, &SNiagaraDistributionBindingEditor::SetBinding, Variable)
-				),
-				Widget
-			);
+		const TArray<FNiagaraVariableBase> AvailableBindings = DistributionAdapter->GetAvailableBindings();
+		if (AvailableBindings.Num() > 0)
+		{
+			for (const FNiagaraVariableBase& Variable : AvailableBindings)
+			{
+				TSharedRef<SWidget> Widget = SNew(SNiagaraParameterName)
+					.ParameterName(Variable.GetName())
+					.IsReadOnly(true);
+
+				MenuBuilder.AddMenuEntry(
+					FUIAction(
+						FExecuteAction::CreateSP(this, &SNiagaraDistributionBindingEditor::SetBinding, Variable)
+					),
+					Widget
+				);
+			}
+		}
+		else
+		{
+			const FNiagaraTypeDefinition AllowedTypeDef = DistributionAdapter->GetExpressionTypeDef();
+			if (AllowedTypeDef.IsValid())
+			{
+				MenuBuilder.AddMenuEntry(
+					FText::Format(LOCTEXT("NoBindings", "No bindings of type '{0}' found"), AllowedTypeDef.GetNameText()),
+					LOCTEXT("NoBindingsDesc", "No bindings for the property type were found."),
+					FSlateIcon(),
+					FUIAction()
+				);
+			}
+			else
+			{
+				MenuBuilder.AddMenuEntry(
+					LOCTEXT("InvalidTypeDef", "Bindings Unavailable"),
+					LOCTEXT("InvalidTypeDefDesc", "Bindings can not be set on this property."),
+					FSlateIcon(),
+					FUIAction()
+				);
+			}
 		}
-
 		return MenuBuilder.MakeWidget();
 	}
 
