- Ensure UI refreshes when enabled / disabling selected emitters

#tests Have two emitters in different enabled states, then flip to enabled or disabled using keyboard shortcut
#rnx

diff --git a/Engine/Plugins/FX/Niagara/Source/NiagaraEditor/Private/ViewModels/NiagaraSystemViewModel.cpp b/Engine/Plugins/FX/Niagara/Source/NiagaraEditor/Private/ViewModels/NiagaraSystemViewModel.cpp
--- a/Engine/Plugins/FX/Niagara/Source/NiagaraEditor/Private/ViewModels/NiagaraSystemViewModel.cpp
+++ b/Engine/Plugins/FX/Niagara/Source/NiagaraEditor/Private/ViewModels/NiagaraSystemViewModel.cpp
@@ -2155,6 +2155,10 @@ void FNiagaraSystemViewModel::EnableEmitters(TArray<FGuid> EmitterHandlesIdsToEn
 		if (EmitterHandlesIdsToEnable.Contains(EmitterHandle->GetId()))
 		{
 			const bool bEnabled = EmitterHandle->GetEmitterHandle()->SetIsEnabled(true, GetSystem(), false);
+			if (bEnabled)
+			{
+				EmitterHandlePropertyChanged(EmitterHandle->GetId());
+			}
 			bAnyEmittersEnabled = bAnyEmittersEnabled || bEnabled;
 		}
 	}
@@ -2181,6 +2185,10 @@ void FNiagaraSystemViewModel::DisableEmitters(TArray<FGuid> EmitterHandlesIdsToD
 		if(EmitterHandlesIdsToDisable.Contains(EmitterHandle->GetId()))
 		{
 			const bool bDisabled = EmitterHandle->GetEmitterHandle()->SetIsEnabled(false, GetSystem(), false);
+			if (bDisabled)
+			{
+				EmitterHandlePropertyChanged(EmitterHandle->GetId());
+			}
 			bAnyEmittersDisabled = bAnyEmittersDisabled || bDisabled;
 		}
 	}
