- Disable low latency when anything other than the lit view mode is active
- Some of the view modes short circuit rendering so the data may not be ready in time

diff --git a/Engine/Plugins/FX/Niagara/Source/Niagara/Private/NiagaraRenderer.cpp b/Engine/Plugins/FX/Niagara/Source/Niagara/Private/NiagaraRenderer.cpp
--- a/Engine/Plugins/FX/Niagara/Source/Niagara/Private/NiagaraRenderer.cpp
+++ b/Engine/Plugins/FX/Niagara/Source/Niagara/Private/NiagaraRenderer.cpp
@@ -443,13 +443,11 @@ bool FNiagaraRenderer::UseLocalSpace(const FNiagaraSceneProxy* Proxy)const
 
 bool FNiagaraRenderer::ViewFamilySupportLowLatencyTranslucency(const FSceneViewFamily& ViewFamily)
 {
-	// when underwater, single layer water will render translucency before opaque
-	const bool bIncludesUnderwaterView = ViewFamily.Views.ContainsByPredicate([](const FSceneView* SceneView)
-	{
-		return SceneView->IsUnderwater();
-	});
-
-	return !bIncludesUnderwaterView;
+	return
+		// Disable if we are anything other than lit to be safe
+		ViewFamily.ViewMode == VMI_Lit &&
+		// when underwater, single layer water will render translucency before opaque
+		!ViewFamily.Views.ContainsByPredicate([](const FSceneView* SceneView) { return SceneView->IsUnderwater(); });
 }
 
 void FNiagaraRenderer::ProcessMaterialParameterBindings(const FNiagaraRendererMaterialParameters& MaterialParameters, const FNiagaraEmitterInstance* InEmitter, TConstArrayView<UMaterialInterface*> InMaterials) const
