Mobile: Explicitly disable instance order preservation for Niagara Mesh rendering, as this feature is not yet supported on mobile

diff --git a/Engine/Plugins/FX/Niagara/Source/Niagara/Private/NiagaraRendererMeshes.cpp b/Engine/Plugins/FX/Niagara/Source/Niagara/Private/NiagaraRendererMeshes.cpp
--- a/Engine/Plugins/FX/Niagara/Source/Niagara/Private/NiagaraRendererMeshes.cpp
+++ b/Engine/Plugins/FX/Niagara/Source/Niagara/Private/NiagaraRendererMeshes.cpp
@@ -1160,7 +1160,9 @@ void FNiagaraRendererMeshes::SetupElementForGPUScene(
 		OutMeshBatchElement.bForceInstanceCulling = true;
 
 		// Force it to preserve instance order if we are sorted so that GPU Scene's instance culling doesn't scramble them
-		OutMeshBatchElement.bPreserveInstanceOrder = ParticleMeshRenderData.bNeedsSort;
+		OutMeshBatchElement.bPreserveInstanceOrder = ParticleMeshRenderData.bNeedsSort 
+			// NOTE: For now, this feature is not supported for mobile platforms
+			&& (FeatureLevel > ERHIFeatureLevel::ES3_1);
 
 		GPUSceneRes.DynamicPrimitiveData.DataWriterGPU = FGPUSceneWriteDelegate::CreateLambda(		
 			ParticleMeshRenderData.Collector,
diff --git a/Engine/Plugins/FX/Niagara/Source/NiagaraVertexFactories/Private/NiagaraMeshVertexFactory.cpp b/Engine/Plugins/FX/Niagara/Source/NiagaraVertexFactories/Private/NiagaraMeshVertexFactory.cpp
--- a/Engine/Plugins/FX/Niagara/Source/NiagaraVertexFactories/Private/NiagaraMeshVertexFactory.cpp
+++ b/Engine/Plugins/FX/Niagara/Source/NiagaraVertexFactories/Private/NiagaraMeshVertexFactory.cpp
@@ -143,7 +143,6 @@ void FNiagaraMeshVertexFactory::InitRHI(FRHICommandListBase& RHICmdList)
 #if NIAGARA_ENABLE_GPU_SCENE_MESHES
 	if (bAddPrimitiveIDElement)
 	{
-		// TODO: Support GPU Scene on mobile? Maybe only for CPU particles?
 		AddPrimitiveIdStreamElement(EVertexInputStreamType::Default, Elements, 13, 13);
 	}
 #endif
