- Make comment more descriptive

#rnx

diff --git a/Engine/Plugins/FX/Niagara/Source/Niagara/Private/NiagaraRenderer.cpp b/Engine/Plugins/FX/Niagara/Source/Niagara/Private/NiagaraRenderer.cpp
--- a/Engine/Plugins/FX/Niagara/Source/Niagara/Private/NiagaraRenderer.cpp
+++ b/Engine/Plugins/FX/Niagara/Source/Niagara/Private/NiagaraRenderer.cpp
@@ -444,7 +444,8 @@ bool FNiagaraRenderer::UseLocalSpace(const FNiagaraSceneProxy* Proxy)const
 bool FNiagaraRenderer::ViewFamilySupportLowLatencyTranslucency(const FSceneViewFamily& ViewFamily)
 {
 	return
-		// Disable if we are anything other than lit to be safe
+		// Disable low latency when anything other than the lit view mode is active
+		// Some of the other view modes short circuit rendering so the data may not be ready in time resulting in corruption / GPU hangs
 		ViewFamily.ViewMode == VMI_Lit &&
 		// when underwater, single layer water will render translucency before opaque
 		!ViewFamily.Views.ContainsByPredicate([](const FSceneView* SceneView) { return SceneView->IsUnderwater(); });
