- Add wildcard name match exclusions to sim cache compare

diff --git a/Engine/Plugins/FX/Niagara/Source/Niagara/Classes/NiagaraSimCacheCompare.h b/Engine/Plugins/FX/Niagara/Source/Niagara/Classes/NiagaraSimCacheCompare.h
--- a/Engine/Plugins/FX/Niagara/Source/Niagara/Classes/NiagaraSimCacheCompare.h
+++ b/Engine/Plugins/FX/Niagara/Source/Niagara/Classes/NiagaraSimCacheCompare.h
@@ -31,6 +31,9 @@ public:
 	/** Add an attribute to exclude from the comparison. */
 	NIAGARA_API void AddExcludeAttribute(FNiagaraVariableBase Variable);
 
+	/** Adds an exclude rule to match against attreibute names. */
+	NIAGARA_API void AddExcludeAttributesByWildcard(FString Wildcard);
+
 	/** Setup comparison tolerances for variables that match the provided type. */
 	NIAGARA_API void SetToleranceByType(const UNiagaraSimCache& SimCache, const FNiagaraTypeDefinition& TypeDef, float Tolerance);
 
@@ -50,6 +53,7 @@ private:
 private:
 	TSet<FNiagaraVariableBase>			AttributesToExclude_SystemScript;
 	TSet<FNiagaraVariableBase>			AttributesToExclude_ParticleScript;
+	TArray<FString>						AttributesToExclude_ByWildcard;
 	TMap<FNiagaraVariableBase, float>	VariableToFloatTolerances_SystemScript;
 	TMap<FNiagaraVariableBase, float>	VariableToFloatTolerances_ParticleScript;
 };
diff --git a/Engine/Plugins/FX/Niagara/Source/Niagara/Private/NiagaraSimCacheCompare.cpp b/Engine/Plugins/FX/Niagara/Source/Niagara/Private/NiagaraSimCacheCompare.cpp
--- a/Engine/Plugins/FX/Niagara/Source/Niagara/Private/NiagaraSimCacheCompare.cpp
+++ b/Engine/Plugins/FX/Niagara/Source/Niagara/Private/NiagaraSimCacheCompare.cpp
@@ -29,14 +29,32 @@ namespace NiagaraSimCacheCompareInternal
 		return FString::FromInt(Num);
 	}
 
-	TArray<FNiagaraSimCacheVariable> GatherEmitterVariables(const UNiagaraSimCache& SimCache, int32 EmitterIndex, const TSet<FNiagaraVariableBase>& ExcludeVariables)
+	bool MatchesWildcardArray(TConstArrayView<FString> Wildcards, FName Name)
+	{
+		if (Wildcards.Num() == 0)
+		{
+			return false;
+		}
+
+		const FString NameString = Name.ToString();
+		for (const FString& Wildcard : Wildcards)
+		{
+			if (NameString.MatchesWildcard(Wildcard))
+			{
+				return true;
+			}
+		}
+		return false;
+	}
+
+	TArray<FNiagaraSimCacheVariable> GatherEmitterVariables(const UNiagaraSimCache& SimCache, int32 EmitterIndex, const TSet<FNiagaraVariableBase>& ExcludeVariables, TConstArrayView<FString> WildcardMatches)
 	{
 		TArray<FNiagaraSimCacheVariable> OutVariables;
 		SimCache.ForEachEmitterAttribute(
 			EmitterIndex,
-			[&OutVariables, &ExcludeVariables](const FNiagaraSimCacheVariable& Variable) -> bool
+			[&OutVariables, &ExcludeVariables, &WildcardMatches](const FNiagaraSimCacheVariable& Variable) -> bool
 			{
-				if (!ExcludeVariables.Contains(Variable.Variable))
+				if (!ExcludeVariables.Contains(Variable.Variable) && !MatchesWildcardArray(WildcardMatches, Variable.Variable.GetName()))
 				{
 					OutVariables.Add(Variable);
 				}
@@ -171,6 +189,11 @@ void FNiagaraSimCacheCompare::AddExcludeAttribute(FNiagaraVariableBase Variable)
 	}
 }
 
+void FNiagaraSimCacheCompare::AddExcludeAttributesByWildcard(FString Wildcard)
+{
+	AttributesToExclude_ByWildcard.AddUnique(Wildcard);
+}
+
 void FNiagaraSimCacheCompare::SetToleranceByType(const UNiagaraSimCache& SimCache, const FNiagaraTypeDefinition& TypeDef, float Tolerance)
 {
 	using namespace NiagaraSimCacheCompareInternal;
@@ -290,6 +313,11 @@ bool FNiagaraSimCacheCompare::Compare(const UNiagaraSimCache& LhsCache, const UN
 				{
 					continue;
 				}
+				
+				if ( MatchesWildcardArray(AttributesToExclude_ByWildcard, DIVariable.GetName()) )
+				{
+					continue;
+				}
 
 				UClass* Class = DIVariable.GetType().GetClass();
 				INiagaraSimCacheCustomStorageInterface* DataInterfaceCDO = Class ? Cast<INiagaraSimCacheCustomStorageInterface>(DIVariable.GetType().GetClass()->GetDefaultObject()) : nullptr;
@@ -363,10 +391,10 @@ bool FNiagaraSimCacheCompare::CompareEmitter(const UNiagaraSimCache& LhsCache, c
 	}
 
 	// Gather all the attributes we want to compare
-	const TArray<FNiagaraSimCacheVariable> CompareVariables = GatherEmitterVariables(LhsCache, EmitterIndex, AttributesToExclude);
+	const TArray<FNiagaraSimCacheVariable> CompareVariables = GatherEmitterVariables(LhsCache, EmitterIndex, AttributesToExclude, AttributesToExclude_ByWildcard);
 	{
 		bool bAttributesValid = true;
-		const TArray<FNiagaraSimCacheVariable> RhsCompareVariables = GatherEmitterVariables(RhsCache, EmitterIndex, AttributesToExclude);
+		const TArray<FNiagaraSimCacheVariable> RhsCompareVariables = GatherEmitterVariables(RhsCache, EmitterIndex, AttributesToExclude, AttributesToExclude_ByWildcard);
 		if (CompareVariables.Num() != RhsCompareVariables.Num())
 		{
 			AddError(OutDifferences, FString::Printf(TEXT("%s attribute count different frame %d.  Expected: %d, Actual %d"), *EmitterNameString, FrameIndex, CompareVariables.Num(), RhsCompareVariables.Num()));
