Niagara - Revert a change to the compiler which changed the caching behavior for non-exportable static parameters.  This change fixed an issue with incorrect compile errors when binding static parameters but it has caused other worse issues and needs further investigation.

diff --git a/Engine/Plugins/FX/Niagara/Source/NiagaraEditor/Private/NiagaraCompiler.cpp b/Engine/Plugins/FX/Niagara/Source/NiagaraEditor/Private/NiagaraCompiler.cpp
--- a/Engine/Plugins/FX/Niagara/Source/NiagaraEditor/Private/NiagaraCompiler.cpp
+++ b/Engine/Plugins/FX/Niagara/Source/NiagaraEditor/Private/NiagaraCompiler.cpp
@@ -1608,10 +1608,7 @@ TSharedPtr<FNiagaraGraphCachedDataBase, ESPMode::ThreadSafe> FNiagaraEditorModul
 
 			for (int32 BuilderVarIndex = 0; BuilderVarIndex < Builder.StaticVariables.Num() && BuilderVarIndex < Builder.StaticVariableExportable.Num(); ++BuilderVarIndex)
 			{
-				if (Builder.StaticVariableExportable[BuilderVarIndex])
-				{
-					StaticVariables.AddUnique(Builder.StaticVariables[BuilderVarIndex]);
-				}
+				StaticVariables.AddUnique(Builder.StaticVariables[BuilderVarIndex]);
 			}
 		}
 	}
