Refactored renderer thumbnails for emitter into dedicated widget

- Made sure invalid renderer entries still get a thumbnail generated. This fixes a crash where renderer thumbnails & generated tooltip counts have to match, and is consistent with previous display of thumbnails
- Removed 'on tick' refresh of the top content bar (it was broken anyways)
- Added hover hint button to the thumbnail renderers to make it obvious you can click them

diff --git a/Engine/Plugins/FX/Niagara/Source/Niagara/Private/NiagaraMeshRendererProperties.cpp b/Engine/Plugins/FX/Niagara/Source/Niagara/Private/NiagaraMeshRendererProperties.cpp
--- a/Engine/Plugins/FX/Niagara/Source/Niagara/Private/NiagaraMeshRendererProperties.cpp
+++ b/Engine/Plugins/FX/Niagara/Source/Niagara/Private/NiagaraMeshRendererProperties.cpp
@@ -1146,6 +1146,7 @@ void UNiagaraMeshRendererProperties::GetAdditionalVariables(TArray<FNiagaraVaria
 void UNiagaraMeshRendererProperties::GetRendererWidgets(const FNiagaraEmitterInstance* InEmitter, TArray<TSharedPtr<SWidget>>& OutWidgets, TSharedPtr<FAssetThumbnailPool> InThumbnailPool) const
 {	
 	TArray<UObject*> Assets;
+	// We create a thumbnail per mesh entry even if it's invalid to highlight that there are unused mesh entries
 	for(const FNiagaraMeshRendererMeshProperties& MeshProperties : Meshes)
 	{
 		UStaticMesh* Mesh = MeshProperties.Mesh;
@@ -1153,6 +1154,10 @@ void UNiagaraMeshRendererProperties::GetRendererWidgets(const FNiagaraEmitterIns
 		{
 			Assets.Add(Mesh);
 		}
+		else
+		{
+			Assets.Add(nullptr);
+		}
 	}
 	
 	CreateRendererWidgetsForAssets(Assets, InThumbnailPool, OutWidgets);
diff --git a/Engine/Plugins/FX/Niagara/Source/Niagara/Private/NiagaraRendererProperties.cpp b/Engine/Plugins/FX/Niagara/Source/Niagara/Private/NiagaraRendererProperties.cpp
--- a/Engine/Plugins/FX/Niagara/Source/Niagara/Private/NiagaraRendererProperties.cpp
+++ b/Engine/Plugins/FX/Niagara/Source/Niagara/Private/NiagaraRendererProperties.cpp
@@ -781,15 +781,14 @@ void UNiagaraRendererProperties::CreateDefaultRendererWidget(TArray<TSharedPtr<S
 void UNiagaraRendererProperties::CreateRendererWidgetsForAssets(TConstArrayView<UObject*> InAssets, TSharedPtr<FAssetThumbnailPool> InThumbnailPool, TArray<TSharedPtr<SWidget>>& OutWidgets) const
 {
 	constexpr int32 ThumbnailSize = 32;
-
-	bool bNeedsDefaultWidget = true;
+	
 	for (UObject* Asset : InAssets)
 	{
 		if (Asset == nullptr)
 		{
+			CreateDefaultRendererWidget(OutWidgets);
 			continue;
 		}
-		bNeedsDefaultWidget = false;
 
 		TSharedPtr<FAssetThumbnail> AssetThumbnail = MakeShareable(new FAssetThumbnail(Asset, ThumbnailSize, ThumbnailSize, InThumbnailPool));
 
@@ -799,11 +798,6 @@ void UNiagaraRendererProperties::CreateRendererWidgetsForAssets(TConstArrayView<
 		TSharedRef<SWidget> ThumbnailWidget = AssetThumbnail->MakeThumbnailWidget(ThumbnailConfig);
 		OutWidgets.Add(ThumbnailWidget);
 	}
-
-	if (bNeedsDefaultWidget)
-	{
-		CreateDefaultRendererWidget(OutWidgets);
-	}
 }
 
 void UNiagaraRendererProperties::CreateRendererWidgetsForAssets(TConstArrayView<UMaterialInterface*> InMaterials, TSharedPtr<FAssetThumbnailPool> InThumbnailPool, TArray<TSharedPtr<SWidget>>& OutWidgets) const
diff --git a/Engine/Plugins/FX/Niagara/Source/NiagaraEditorWidgets/Private/SNiagaraEmitterRendererThumbnails.cpp b/Engine/Plugins/FX/Niagara/Source/NiagaraEditorWidgets/Private/SNiagaraEmitterRendererThumbnails.cpp
new file mode 100644
--- /dev/null
+++ b/Engine/Plugins/FX/Niagara/Source/NiagaraEditorWidgets/Private/SNiagaraEmitterRendererThumbnails.cpp
@@ -0,0 +1,238 @@
+﻿// Copyright Epic Games, Inc. All Rights Reserved.
+
+
+#include "SNiagaraEmitterRendererThumbnails.h"
+
+#include "NiagaraEmitterHandle.h"
+#include "SlateOptMacros.h"
+#include "ViewModels/NiagaraEmitterViewModel.h"
+#include "ViewModels/Stack/NiagaraStackRenderItemGroup.h"
+#include "ViewModels/Stack/NiagaraStackRoot.h"
+#include "ViewModels/Stack/NiagaraStackViewModel.h"
+#include "Widgets/SToolTip.h"
+#include "Materials/Material.h"
+#include "ViewModels/NiagaraOverviewGraphViewModel.h"
+#include "ViewModels/NiagaraSystemSelectionViewModel.h"
+#include "ViewModels/NiagaraSystemViewModel.h"
+#include "Widgets/Layout/SScaleBox.h"
+
+BEGIN_SLATE_FUNCTION_BUILD_OPTIMIZATION
+
+static constexpr float RendererThumbnailSize = 24.f;
+
+#define LOCTEXT_NAMESPACE "NiagaraEmitterRendererThumbnails"
+
+FNiagaraRendererStackEntryObserver::FNiagaraRendererStackEntryObserver()
+{
+	
+}
+
+FNiagaraRendererStackEntryObserver::~FNiagaraRendererStackEntryObserver()
+{
+	if(EmitterHandleViewModelWeak.IsValid())
+	{
+		EmitterHandleViewModelWeak.Pin()->GetEmitterStackViewModel()->OnDataObjectChanged().RemoveAll(this);
+		EmitterHandleViewModelWeak.Pin()->GetEmitterStackViewModel()->OnStructureChanged().RemoveAll(this);
+	}
+
+	UMaterial::OnMaterialCompilationFinished().RemoveAll(this);
+}
+
+void FNiagaraRendererStackEntryObserver::Initialize(TSharedRef<FNiagaraEmitterHandleViewModel> InEmitterHandleViewModel)
+{
+	EmitterHandleViewModelWeak = InEmitterHandleViewModel;
+
+	if(UNiagaraStackViewModel* StackViewModel = EmitterHandleViewModelWeak.Pin()->GetEmitterStackViewModel())
+	{
+		StackViewModel->OnDataObjectChanged().AddSP(this, &FNiagaraRendererStackEntryObserver::OnStackDataObjectChanged);
+		StackViewModel->OnStructureChanged().AddSP(this, &FNiagaraRendererStackEntryObserver::OnRenderGroupStructureChanged);
+	}
+	
+	UMaterial::OnMaterialCompilationFinished().AddSP(this, &FNiagaraRendererStackEntryObserver::OnMaterialCompiled);
+}
+
+void FNiagaraRendererStackEntryObserver::OnRenderGroupStructureChanged(ENiagaraStructureChangedFlags NiagaraStructureChangedFlags) const
+{
+	OnRenderersChanged.Broadcast();
+}
+
+void FNiagaraRendererStackEntryObserver::OnMaterialCompiled(class UMaterialInterface* MaterialInterface) const
+{
+	if (EmitterHandleViewModelWeak.IsValid())
+	{
+		bool bUsingThisMaterial = false;
+		TArray<UNiagaraStackEntry*> RendererPreviewStackEntries;
+		EmitterHandleViewModelWeak.Pin()->GetRendererEntries(RendererPreviewStackEntries);
+		FNiagaraEmitterInstance* InInstance = EmitterHandleViewModelWeak.Pin()->GetEmitterViewModel()->GetSimulation().IsValid() ? EmitterHandleViewModelWeak.Pin()->GetEmitterViewModel()->GetSimulation().Pin().Get() : nullptr;
+		for (UNiagaraStackEntry* Entry : RendererPreviewStackEntries)
+		{
+			if (UNiagaraStackRendererItem* RendererItem = Cast<UNiagaraStackRendererItem>(Entry))
+			{
+				TArray<UMaterialInterface*> Materials;
+				RendererItem->GetRendererProperties()->GetUsedMaterials(InInstance, Materials);
+				if (Materials.Contains(MaterialInterface))
+				{
+					bUsingThisMaterial = true;
+					break;
+				}
+			}
+		}
+
+		if (bUsingThisMaterial)
+		{
+			OnRenderersChanged.Broadcast();
+		}
+	}
+}
+
+void FNiagaraRendererStackEntryObserver::OnStackDataObjectChanged(TArray<UObject*> ChangedObjects, ENiagaraDataObjectChange NiagaraDataObjectChange) const
+{
+	for (UObject* ChangedObject : ChangedObjects)
+	{
+		if (ChangedObject->IsA<UNiagaraRendererProperties>())
+		{
+			OnRenderersChanged.Broadcast();
+			break;
+		}
+	}
+}
+
+void SNiagaraEmitterRendererThumbnails::Construct(const FArguments& InArgs, TSharedRef<FNiagaraEmitterHandleViewModel> InEmitterHandleViewModel)
+{
+	EmitterHandleViewModelWeak = InEmitterHandleViewModel;
+	OnRendererThumbnailClickedDelegate = InArgs._OnRendererThumbnailClicked;
+	
+	RendererObserver = MakeShared<FNiagaraRendererStackEntryObserver>();
+	RendererObserver->Initialize(InEmitterHandleViewModel);
+	RendererObserver->OnRenderersChanged.AddSP(this, &SNiagaraEmitterRendererThumbnails::RefreshThumbnails);
+	
+	ChildSlot
+	[
+		SAssignNew(Content, SBox)
+	];
+
+	RefreshThumbnails();
+}
+
+void SNiagaraEmitterRendererThumbnails::RefreshThumbnails()
+{	
+	TArray<UNiagaraStackEntry*> RendererPreviewStackEntries;
+	EmitterHandleViewModelWeak.Pin()->GetRendererEntries(RendererPreviewStackEntries);
+	FNiagaraEmitterInstance* InInstance = EmitterHandleViewModelWeak.Pin()->GetEmitterViewModel()->GetSimulation().IsValid() ? EmitterHandleViewModelWeak.Pin()->GetEmitterViewModel()->GetSimulation().Pin().Get() : nullptr;
+	
+	FToolBarBuilder ToolBarBuilder(nullptr, FMultiBoxCustomization::None, nullptr, true);
+	ToolBarBuilder.SetLabelVisibility(EVisibility::Collapsed);
+		
+	for (int32 StackEntryIndex = 0; StackEntryIndex < RendererPreviewStackEntries.Num(); StackEntryIndex++)
+	{
+		UNiagaraStackEntry* Entry = RendererPreviewStackEntries[StackEntryIndex];
+		if (UNiagaraStackRendererItem* RendererItem = Cast<UNiagaraStackRendererItem>(Entry))
+		{
+			TArray<TSharedPtr<SWidget>> Widgets;
+			RendererItem->GetRendererProperties()->GetRendererWidgets(InInstance, Widgets, UThumbnailManager::Get().GetSharedThumbnailPool());
+			TArray<TSharedPtr<SWidget>> TooltipWidgets;
+			RendererItem->GetRendererProperties()->GetRendererTooltipWidgets(InInstance, TooltipWidgets, UThumbnailManager::Get().GetSharedThumbnailPool());
+
+			// Fallback widget & tooltip
+			if(Widgets.Num() == 0)
+			{
+				RendererItem->GetRendererProperties()->CreateDefaultRendererWidget(Widgets);
+			}
+
+			if(TooltipWidgets.Num() == 0)
+			{
+				FText BaseText = LOCTEXT("DefaultRendererTooltip", "{0}");
+				TooltipWidgets.Add(SNew(STextBlock).Text(FText::FormatOrdered(BaseText, RendererItem->GetRendererProperties()->GetClass()->GetDisplayNameText())));
+			}
+			
+			check(Widgets.Num() == TooltipWidgets.Num());
+			for (int32 WidgetIndex = 0; WidgetIndex < Widgets.Num(); WidgetIndex++)
+			{
+				SetupRendererThumbnailTooltip(Widgets[WidgetIndex], TooltipWidgets[WidgetIndex]);
+				
+				// We wrap the thumbnail in a scalebox so it will scale to whatever its parent dictates
+				TSharedPtr<SWidget> RendererThumbnailWidget = SNew(SBox)
+					.VAlign(VAlign_Center)
+					.HAlign(HAlign_Center)
+					.HeightOverride(RendererThumbnailSize)
+					.WidthOverride(RendererThumbnailSize)
+					[
+						SNew(SScaleBox)
+						.Stretch(EStretch::ScaleToFit)
+						[
+							Widgets[WidgetIndex].ToSharedRef()	
+						]
+					];
+				
+				// If we have on click functionality, we wrap the thumbnail in a button
+				if(OnRendererThumbnailClickedDelegate.IsBound())
+				{
+					RendererThumbnailWidget = SNew(SButton)
+						.ButtonStyle(&FAppStyle::GetWidgetStyle<FButtonStyle>("HoverHintOnly"))
+						.ContentPadding(FMargin(1.f))
+						.OnClicked(this, &SNiagaraEmitterRendererThumbnails::OnRendererThumbnailClicked, RendererItem)
+						[
+							RendererThumbnailWidget.ToSharedRef()
+						];
+				}
+				
+				ToolBarBuilder.AddWidget(
+					RendererThumbnailWidget.ToSharedRef()
+				);
+			}
+	
+			// if we had a widget for this entry, add a separator for the next entry's widgets, except for the last entry
+			if(Widgets.Num() > 0 && StackEntryIndex < RendererPreviewStackEntries.Num() - 1)
+			{
+				ToolBarBuilder.AddSeparator();
+			}
+		}
+	}
+	
+	Content->SetContent(ToolBarBuilder.MakeWidget());
+}
+
+void SNiagaraEmitterRendererThumbnails::SetupRendererThumbnailTooltip(TSharedPtr<SWidget> InWidget, TSharedPtr<SWidget> InTooltipWidget)
+{
+	TSharedPtr<SToolTip> ThumbnailTooltipWidget;
+	// If this is just text, don't constrain the size
+	if (InTooltipWidget->GetType() == TEXT("STextBlock"))
+	{
+		ThumbnailTooltipWidget = SNew(SToolTip)
+			.Content()
+			[
+				InTooltipWidget.ToSharedRef()
+			];
+	}
+	else
+	{
+		ThumbnailTooltipWidget = SNew(SToolTip)
+			.Content()
+			[
+				SNew(SBox)
+				.MaxDesiredHeight(64.0f)
+				.MinDesiredHeight(64.0f)
+				.MaxDesiredWidth(64.0f)
+				.MinDesiredWidth(64.0f)
+				[
+					InTooltipWidget.ToSharedRef()
+				]
+			];
+	}
+	
+	InWidget->SetToolTip(ThumbnailTooltipWidget);
+}
+
+FReply SNiagaraEmitterRendererThumbnails::OnRendererThumbnailClicked(UNiagaraStackRendererItem* StackRendererItem) const
+{
+	if(OnRendererThumbnailClickedDelegate.IsBound())
+	{
+		return OnRendererThumbnailClickedDelegate.Execute(StackRendererItem);
+	}
+
+	return FReply::Unhandled();
+}
+
+#undef LOCTEXT_NAMESPACE
+
+END_SLATE_FUNCTION_BUILD_OPTIMIZATION
diff --git a/Engine/Plugins/FX/Niagara/Source/NiagaraEditorWidgets/Private/SNiagaraEmitterRendererThumbnails.h b/Engine/Plugins/FX/Niagara/Source/NiagaraEditorWidgets/Private/SNiagaraEmitterRendererThumbnails.h
new file mode 100644
--- /dev/null
+++ b/Engine/Plugins/FX/Niagara/Source/NiagaraEditorWidgets/Private/SNiagaraEmitterRendererThumbnails.h
@@ -0,0 +1,59 @@
+﻿// Copyright Epic Games, Inc. All Rights Reserved.
+
+#pragma once
+
+#include "CoreMinimal.h"
+#include "ViewModels/NiagaraEmitterHandleViewModel.h"
+#include "ViewModels/Stack/NiagaraStackRendererItem.h"
+#include "ViewModels/Stack/NiagaraStackRenderItemGroup.h"
+#include "Widgets/SCompoundWidget.h"
+
+
+struct FNiagaraRendererStackEntryObserver : public TSharedFromThis<FNiagaraRendererStackEntryObserver>
+{
+public:
+	FNiagaraRendererStackEntryObserver();
+	~FNiagaraRendererStackEntryObserver();
+
+	void Initialize(TSharedRef<FNiagaraEmitterHandleViewModel> InEmitterHandleViewModel);
+	
+	DECLARE_MULTICAST_DELEGATE(FOnRendererGroupStructureChanged);
+
+	FOnRendererGroupStructureChanged OnRenderersChanged;
+	
+private:	
+	void OnRenderGroupStructureChanged(ENiagaraStructureChangedFlags NiagaraStructureChangedFlags) const;
+	void OnMaterialCompiled(class UMaterialInterface* MaterialInterface) const;
+	void OnStackDataObjectChanged(TArray<UObject*> ChangedObjects, ENiagaraDataObjectChange NiagaraDataObjectChange) const;
+
+private:
+	TWeakPtr<FNiagaraEmitterHandleViewModel> EmitterHandleViewModelWeak;
+};
+
+/**
+ * 
+ */
+class NIAGARAEDITORWIDGETS_API SNiagaraEmitterRendererThumbnails : public SCompoundWidget
+{
+public:
+	DECLARE_DELEGATE_RetVal_OneParam(FReply, FOnRendererThumbnailClicked, UNiagaraStackRendererItem* StackRendererItem);
+	
+	SLATE_BEGIN_ARGS(SNiagaraEmitterRendererThumbnails)
+		{}
+		SLATE_EVENT(FOnRendererThumbnailClicked, OnRendererThumbnailClicked)
+	SLATE_END_ARGS()
+
+	/** Constructs this widget with InArgs */
+	void Construct(const FArguments& InArgs, TSharedRef<FNiagaraEmitterHandleViewModel> InEmitterHandleViewModel);
+private:
+	void RefreshThumbnails();
+
+	static void SetupRendererThumbnailTooltip(TSharedPtr<SWidget> InWidget, TSharedPtr<SWidget> InTooltipWidget);
+	FReply OnRendererThumbnailClicked(UNiagaraStackRendererItem* StackRendererItem) const;
+
+private:
+	TSharedPtr<SBox> Content;
+	TSharedPtr<FNiagaraRendererStackEntryObserver> RendererObserver;
+	TWeakPtr<FNiagaraEmitterHandleViewModel> EmitterHandleViewModelWeak;
+	FOnRendererThumbnailClicked OnRendererThumbnailClickedDelegate;
+};
diff --git a/Engine/Plugins/FX/Niagara/Source/NiagaraEditorWidgets/Private/SNiagaraOverviewStackNode.cpp b/Engine/Plugins/FX/Niagara/Source/NiagaraEditorWidgets/Private/SNiagaraOverviewStackNode.cpp
--- a/Engine/Plugins/FX/Niagara/Source/NiagaraEditorWidgets/Private/SNiagaraOverviewStackNode.cpp
+++ b/Engine/Plugins/FX/Niagara/Source/NiagaraEditorWidgets/Private/SNiagaraOverviewStackNode.cpp
@@ -47,11 +47,11 @@
 #include "Widgets/Layout/SScaleBox.h"
 #include "Widgets/Layout/SSpacer.h"
 #include "AssetThumbnail.h"
+#include "SNiagaraEmitterRendererThumbnails.h"
 #include "Framework/Application/SlateApplication.h"
 
 #define LOCTEXT_NAMESPACE "NiagaraOverviewStackNode"
 
-constexpr float RendererThumbnailSize = 24.f;
 constexpr float EmitterThumbnailSize = 200.f;
 constexpr float NodeTitleMaxSize = 125.f;
 
@@ -61,8 +61,6 @@ void SNiagaraOverviewStackNode::Construct(const FArguments& InArgs, UNiagaraOver
 	OverviewStackNode = InNode;
 	StackViewModel = nullptr;
 	OverviewSelectionViewModel = nullptr;
-	bIsHoveringThumbnail = false;
-	bTopContentBarRefreshPending = true;
 	CurrentIssueIndex = -1;
 	// we are reducing the left margin to better place our widgets
 	TitleBorderMargin = FMargin(2.f, 5.f, 30.f, 3.f);
@@ -71,7 +69,7 @@ void SNiagaraOverviewStackNode::Construct(const FArguments& InArgs, UNiagaraOver
 	
 	if (OverviewStackNode->GetOwningSystem() != nullptr)
 	{
-		FNiagaraEditorModule& NiagaraEditorModule = FModuleManager::Get().LoadModuleChecked<FNiagaraEditorModule>("NiagaraEditor");
+		FNiagaraEditorModule& NiagaraEditorModule = FNiagaraEditorModule::Get();
 		TSharedPtr<FNiagaraSystemViewModel> OwningSystemViewModel = NiagaraEditorModule.GetExistingViewModelForSystem(OverviewStackNode->GetOwningSystem());
 		if (OwningSystemViewModel.IsValid())
 		{
@@ -89,12 +87,7 @@ void SNiagaraOverviewStackNode::Construct(const FArguments& InArgs, UNiagaraOver
 					StackViewModel = EmitterHandleViewModelWeak.Pin()->GetEmitterStackViewModel();
 				}
 			}
-			if (StackViewModel)
-			{
-				StackViewModel->OnStructureChanged().AddSP(this, &SNiagaraOverviewStackNode::StackViewModelStructureChanged);
-				StackViewModel->OnDataObjectChanged().AddSP(this, &SNiagaraOverviewStackNode::StackViewModelDataObjectChanged);
-			}
-			UMaterial::OnMaterialCompilationFinished().AddSP(this, &SNiagaraOverviewStackNode::OnMaterialCompiled);
+			
 			OverviewSelectionViewModel = OwningSystemViewModel->GetSelectionViewModel();
 			ScalabilityViewModel = OwningSystemViewModel->GetScalabilityViewModel();
 
@@ -317,42 +310,6 @@ FReply SNiagaraOverviewStackNode::OnCycleThroughIssues()
 	return FReply::Handled();
 }
 
-TSharedRef<SWidget> SNiagaraOverviewStackNode::CreateRendererThumbnailWidget(UNiagaraStackEntry* InData, TSharedPtr<SWidget> InWidget, TSharedPtr<SWidget> InTooltipWidget)
-{
-	TSharedPtr<SToolTip> ThumbnailTooltipWidget;
-	// If this is just text, don't constrain the size
-	if (InTooltipWidget->GetType() == TEXT("STextBlock"))
-	{
-		ThumbnailTooltipWidget = SNew(SToolTip)
-			.Content()
-			[
-				InTooltipWidget.ToSharedRef()
-			];
-	}
-	else
-	{
-
-		ThumbnailTooltipWidget = SNew(SToolTip)
-			.Content()
-			[
-				SNew(SBox)
-				.MaxDesiredHeight(64.0f)
-				.MinDesiredHeight(64.0f)
-				.MaxDesiredWidth(64.0f)
-				.MinDesiredWidth(64.0f)
-				[
-					InTooltipWidget.ToSharedRef()
-				]
-			];
-	}
-	InWidget->SetOnMouseButtonDown(FPointerEventHandler::CreateSP(this, &SNiagaraOverviewStackNode::OnClickedRenderingPreview, InData));
-	InWidget->SetOnMouseEnter(FNoReplyPointerEventHandler::CreateSP(this, &SNiagaraOverviewStackNode::SetIsHoveringThumbnail, true));
-	InWidget->SetOnMouseLeave(FSimpleNoReplyPointerEventHandler::CreateSP(this, &SNiagaraOverviewStackNode::SetIsHoveringThumbnail, false));
-	InWidget->SetToolTip(ThumbnailTooltipWidget);
-
-	return InWidget.ToSharedRef();
-}
-
 FReply SNiagaraOverviewStackNode::OnCaptureThumbnailButtonClicked()
 {
 	if(EmitterHandleViewModelWeak.IsValid())
@@ -365,19 +322,15 @@ FReply SNiagaraOverviewStackNode::OnCaptureThumbnailButtonClicked()
 	return FReply::Unhandled();
 }
 
-FReply SNiagaraOverviewStackNode::OnClickedRenderingPreview(const FGeometry& InGeometry, const FPointerEvent& InEvent, UNiagaraStackEntry* InEntry)
+FReply SNiagaraOverviewStackNode::OnRendererThumbnailClicked(UNiagaraStackRendererItem* StackRendererItem)
 {
-	if (InEvent.GetEffectingButton() == EKeys::LeftMouseButton)
-	{
-		TArray<UNiagaraStackEntry*> SelectedEntries;
-		SelectedEntries.Add(InEntry);
-		TArray<UNiagaraStackEntry*> DeselectedEntries;
-		OverviewSelectionViewModel->UpdateSelectedEntries(SelectedEntries, DeselectedEntries, true);
+	TArray<UNiagaraStackEntry*> SelectedEntries;
+	SelectedEntries.Add(StackRendererItem);
+	TArray<UNiagaraStackEntry*> DeselectedEntries;
+	OverviewSelectionViewModel->UpdateSelectedEntries(SelectedEntries, DeselectedEntries, true);
 
-		CurrentIssueIndex = -1;
-		return FReply::Handled();
-	}
-	return FReply::Unhandled();
+	CurrentIssueIndex = -1;
+	return FReply::Handled();
 }
 
 FReply SNiagaraOverviewStackNode::OnPropertiesButtonClicked() const
@@ -430,12 +383,6 @@ void SNiagaraOverviewStackNode::Tick(const FGeometry& AllottedGeometry, const do
 			OverviewStackNode->RenameStarted();
 		}
 
-		if (bTopContentBarRefreshPending && TopContentBarSlot != nullptr)
-		{
-			CreateTopContentBar();
-			bTopContentBarRefreshPending = false;
-		}
-
 		if (ScalabilityWrapper != nullptr && !UseLowDetailNodeContent())
 		{
 			LastHighDetailSize = ScalabilityWrapper->GetTickSpaceGeometry().Size;
@@ -448,34 +395,6 @@ void SNiagaraOverviewStackNode::Tick(const FGeometry& AllottedGeometry, const do
 	}
 }
 
-void SNiagaraOverviewStackNode::OnMaterialCompiled(class UMaterialInterface* MaterialInterface)
-{
-	if (EmitterHandleViewModelWeak.IsValid())
-	{
-		bool bUsingThisMaterial = false;
-		EmitterHandleViewModelWeak.Pin()->GetRendererEntries(RendererPreviewStackEntries);
-		FNiagaraEmitterInstance* InInstance = EmitterHandleViewModelWeak.Pin()->GetEmitterViewModel()->GetSimulation().IsValid() ? EmitterHandleViewModelWeak.Pin()->GetEmitterViewModel()->GetSimulation().Pin().Get() : nullptr;
-		for (UNiagaraStackEntry* Entry : RendererPreviewStackEntries)
-		{
-			if (UNiagaraStackRendererItem* RendererItem = Cast<UNiagaraStackRendererItem>(Entry))
-			{
-				TArray<UMaterialInterface*> Materials;
-				RendererItem->GetRendererProperties()->GetUsedMaterials(InInstance, Materials);
-				if (Materials.Contains(MaterialInterface))
-				{
-					bUsingThisMaterial = true;
-					break;
-				}
-			}
-		}
-
-		if (bUsingThisMaterial)
-		{
-			bTopContentBarRefreshPending = true;
-		}
-	}
-}
-
 TSharedRef<SWidget> SNiagaraOverviewStackNode::CreateTitleWidget_Default(TSharedPtr<SNodeTitle> NodeTitle)
 {
 	TSharedRef<SWidget> DefaultTitle = SGraphNode::CreateTitleWidget(NodeTitle);
@@ -600,7 +519,6 @@ TSharedRef<SWidget> SNiagaraOverviewStackNode::CreateNodeContentArea_Default()
 			.HAlign(HAlign_Fill)
 			.VAlign(VAlign_Center)
 			.Padding(2.f)
-			.Expose(TopContentBarSlot)
 			[
 				CreateTopContentBar()
 			]
@@ -700,54 +618,13 @@ TSharedRef<SWidget> SNiagaraOverviewStackNode::CreateTopContentBar_Default()
 				CreateIsolateButton()
 			];
 
-		EmitterHandleViewModelWeak.Pin()->GetRendererEntries(RendererPreviewStackEntries);
-		FNiagaraEmitterInstance* InInstance = EmitterHandleViewModelWeak.Pin()->GetEmitterViewModel()->GetSimulation().IsValid() ? EmitterHandleViewModelWeak.Pin()->GetEmitterViewModel()->GetSimulation().Pin().Get() : nullptr;
-
-		FToolBarBuilder ToolBarBuilder(nullptr, FMultiBoxCustomization::None, nullptr, true);
-		ToolBarBuilder.SetLabelVisibility(EVisibility::Collapsed);
-		
-		for (int32 StackEntryIndex = 0; StackEntryIndex < RendererPreviewStackEntries.Num(); StackEntryIndex++)
-		{
-			UNiagaraStackEntry* Entry = RendererPreviewStackEntries[StackEntryIndex];
-			if (UNiagaraStackRendererItem* RendererItem = Cast<UNiagaraStackRendererItem>(Entry))
-			{
-				TArray<TSharedPtr<SWidget>> Widgets;
-				RendererItem->GetRendererProperties()->GetRendererWidgets(InInstance, Widgets, UThumbnailManager::Get().GetSharedThumbnailPool());
-				TArray<TSharedPtr<SWidget>> TooltipWidgets;
-				RendererItem->GetRendererProperties()->GetRendererTooltipWidgets(InInstance, TooltipWidgets, UThumbnailManager::Get().GetSharedThumbnailPool());
-				check(Widgets.Num() == TooltipWidgets.Num());
-				for (int32 WidgetIndex = 0; WidgetIndex < Widgets.Num(); WidgetIndex++)
-				{
-					ToolBarBuilder.AddWidget(
-						SNew(SBox)
-						.VAlign(VAlign_Center)
-						.HAlign(HAlign_Center)
-						.MinDesiredHeight(RendererThumbnailSize)
-						.MinDesiredWidth(RendererThumbnailSize)
-						.MaxDesiredHeight(RendererThumbnailSize)
-						.MaxDesiredWidth(RendererThumbnailSize)
-						.Visibility(this, &SNiagaraOverviewStackNode::GetEnabledCheckBoxVisibility)
-						[
-							CreateRendererThumbnailWidget(Entry, Widgets[WidgetIndex], TooltipWidgets[WidgetIndex])
-						]
-					);
-				}
-
-				// if we had a widget for this entry, add a separator for the next entry's widgets, except for the last entry
-				if(Widgets.Num() > 0 && StackEntryIndex < RendererPreviewStackEntries.Num() - 1)
-				{
-					ToolBarBuilder.AddSeparator();
-				}
-			}
-		}
-
 		TopContentBar->AddSlot()
 		.AutoWidth()
 		.HAlign(HAlign_Left)
 		.VAlign(VAlign_Center)
 		.MaxWidth(300.f)
 		[
-			ToolBarBuilder.MakeWidget()
+			CreateEmitterRendererThumbnails()
 		];
 	}
 
@@ -1089,6 +966,17 @@ TSharedRef<SWidget> SNiagaraOverviewStackNode::CreateEmitterThumbnail()
 	return SNullWidget::NullWidget;
 }
 
+TSharedRef<SWidget> SNiagaraOverviewStackNode::CreateEmitterRendererThumbnails()
+{
+	if(ensureMsgf(EmitterHandleViewModelWeak.IsValid(), TEXT("This should only ever be called on emitter nodes")))
+	{
+		return SNew(SNiagaraEmitterRendererThumbnails, EmitterHandleViewModelWeak.Pin().ToSharedRef())
+			.OnRendererThumbnailClicked(this, &SNiagaraOverviewStackNode::OnRendererThumbnailClicked);
+	}
+
+	return SNullWidget::NullWidget;
+}
+
 TSharedRef<SWidget> SNiagaraOverviewStackNode::CreateEnabledCheckbox()
 {
 	return SNew(SCheckBox)
@@ -1233,23 +1121,6 @@ TSharedRef<SWidget> SNiagaraOverviewStackNode::CreatePropertiesButton()
 	];
 }
 
-void SNiagaraOverviewStackNode::StackViewModelStructureChanged(ENiagaraStructureChangedFlags Flags)
-{
-	bTopContentBarRefreshPending = true;
-}
-
-void SNiagaraOverviewStackNode::StackViewModelDataObjectChanged(TArray<UObject*> ChangedObjects, ENiagaraDataObjectChange ChangeType)
-{
-	for (UObject* ChangedObject : ChangedObjects)
-	{
-		if (ChangedObject->IsA<UNiagaraRendererProperties>())
-		{
-			bTopContentBarRefreshPending = true;
-			break;
-		}
-	}
-}
-
 EVisibility SNiagaraOverviewStackNode::GetIssueIconVisibility() const
 {
 	return StackViewModel->HasIssues() ? EVisibility::Visible : EVisibility::Collapsed;
diff --git a/Engine/Plugins/FX/Niagara/Source/NiagaraEditorWidgets/Private/SNiagaraOverviewStackNode.h b/Engine/Plugins/FX/Niagara/Source/NiagaraEditorWidgets/Private/SNiagaraOverviewStackNode.h
--- a/Engine/Plugins/FX/Niagara/Source/NiagaraEditorWidgets/Private/SNiagaraOverviewStackNode.h
+++ b/Engine/Plugins/FX/Niagara/Source/NiagaraEditorWidgets/Private/SNiagaraOverviewStackNode.h
@@ -53,9 +53,6 @@ protected:
 	TSharedRef<SWidget> CreateTopContentBar_Summary();
 
 	virtual bool IsHidingPinWidgets() const override { return UseLowDetailNodeContent(); }
-	void StackViewModelStructureChanged(ENiagaraStructureChangedFlags Flags);
-	void StackViewModelDataObjectChanged(TArray<UObject*> ChangedObjects, ENiagaraDataObjectChange ChangeType);
-	void OnMaterialCompiled(class UMaterialInterface* MaterialInterface);
 	
 	bool UseLowDetailNodeContent() const;
 	FVector2D GetLowDetailDesiredSize() const;
@@ -66,6 +63,7 @@ protected:
 	TSharedRef<SWidget> CreateEnabledCheckbox();
 	TSharedRef<SButton> CreateIsolateButton();
 	TSharedRef<SWidget> CreateEmitterThumbnail();
+	TSharedRef<SWidget> CreateEmitterRendererThumbnails();
 	TSharedRef<SWidget> CreateSummaryViewToggle();
 	TSharedRef<SWidget> CreateLocalSpaceToggle();
 	TSharedRef<SButton> CreateCaptureThumbnailButton();
@@ -88,9 +86,8 @@ private:
 	const FSlateBrush* GetEnabledImage() const;
 	void OnEnabledCheckStateChanged(ECheckBoxState InCheckState);
 	EVisibility GetShouldShowSummaryControls() const;
-	TSharedRef<SWidget> CreateRendererThumbnailWidget(UNiagaraStackEntry* InData, TSharedPtr<SWidget> InWidget, TSharedPtr<SWidget> InTooltipWidget);
 	FReply OnCaptureThumbnailButtonClicked();
-	FReply OnClickedRenderingPreview(const FGeometry& InGeometry, const FPointerEvent& InEvent, class UNiagaraStackEntry* InEntry);
+	FReply OnRendererThumbnailClicked(class UNiagaraStackRendererItem* StackRendererItem);
 	FReply OnPropertiesButtonClicked() const;
 	
 	FText GetToggleIsolateToolTip() const;
@@ -104,18 +101,6 @@ private:
 	void OnScalabilityModeChanged(bool bActive);
 	EVisibility ShowExcludedOverlay() const;
 	float GetGraphZoomDistanceAlphaMultiplier() const;
-	void SetIsHoveringThumbnail(const FGeometry& InGeometry, const FPointerEvent& InEvent, const bool bInHoveringThumbnail)
-	{
-		SetIsHoveringThumbnail(InEvent, bInHoveringThumbnail);
-	}
-	void SetIsHoveringThumbnail(const FPointerEvent& InEvent, const bool bInHoveringThumbnail)
-	{
-		bIsHoveringThumbnail = bInHoveringThumbnail;
-	}
-	bool IsHoveringThumbnail()
-	{
-		return bIsHoveringThumbnail;
-	}
 
 	FText GetSpawnCountScaleText() const;
 	FText GetSpawnCountScaleTooltip() const;
@@ -141,8 +126,6 @@ private:
 	TSharedPtr<SInvalidationPanel> ContentAreaWidget;
 
 	TArray<UNiagaraStackEntry*> RendererPreviewStackEntries;
-	bool bIsHoveringThumbnail = false;
-	bool bTopContentBarRefreshPending = false;
 	int32 CurrentIssueIndex = 0;
 	bool bScalabilityModeActive = false;
 
@@ -155,6 +138,4 @@ private:
 
 	TSharedPtr<SBox> ThumbnailContainer;
 	TSharedPtr<FAssetThumbnail> PreviewThumbnail;
-	
-	SVerticalBox::FSlot* TopContentBarSlot = nullptr;
 };
\ No newline at end of file
