Minimal change to get proper completion when calling DeactivateImmediate() on a culled system.

diff --git a/Engine/Plugins/FX/Niagara/Source/Niagara/Private/NiagaraComponent.cpp b/Engine/Plugins/FX/Niagara/Source/Niagara/Private/NiagaraComponent.cpp
--- a/Engine/Plugins/FX/Niagara/Source/Niagara/Private/NiagaraComponent.cpp
+++ b/Engine/Plugins/FX/Niagara/Source/Niagara/Private/NiagaraComponent.cpp
@@ -617,6 +617,7 @@ UNiagaraComponent::UNiagaraComponent(const FObjectInitializer& ObjectInitializer
 	, bAllowScalability(true)
 	, bForceLocalPlayerEffect(false)
 	, bIsCulledByScalability(false)
+	, bIsFullyComplete(true)
 	, bDuringUpdateContextReset(false)
 	, bDesiredPauseState(false)
 	, bRecachePSOs(false)
@@ -1335,6 +1336,7 @@ void UNiagaraComponent::ActivateInternal(bool bReset /* = false */, bool bIsScal
 		return;
 	}
 
+	bIsFullyComplete = false;
 	bIsCulledByScalability = false;
 	
 	//Ensure we're not paused. Non-scalability activate calls should override pause state.
@@ -1579,8 +1581,6 @@ void UNiagaraComponent::DeactivateImmediateInternal(bool bIsScalabilityCull)
 
 	Super::Deactivate();
 
-	bool bWasCulledByScalability = bIsCulledByScalability;
-
 	//UE_LOG(LogNiagara, Log, TEXT("DeactivateImmediate: %p - %s - %s"), this, *Asset->GetName(), bIsScalabilityCull ? TEXT("Scalability") : TEXT(""));
 
 	//UE_LOG(LogNiagara, Log, TEXT("Deactivate %s"), *GetName());
@@ -1602,7 +1602,8 @@ void UNiagaraComponent::DeactivateImmediateInternal(bool bIsScalabilityCull)
 	{
 		SystemInstanceController->Deactivate(true);
 	}
-	else if (bWasCulledByScalability && !bIsCulledByScalability)//We were culled by scalability but no longer, ensure we've handled completion correctly. E.g. returned to the pool etc.
+	
+	if (bIsFullyComplete == false)//In some flows we can get here but not yet be fully complete.
 	{
 		OnSystemComplete(true);
 	}
@@ -1771,6 +1772,7 @@ void UNiagaraComponent::OnSystemComplete(bool bExternalCompletion)
 	//We also want to skip this work if we're destroying during and update context reset.
 	if (bIsCulledByScalability == false && bDuringUpdateContextReset == false)
 	{
+		bIsFullyComplete = true;
 		DestroyCullProxy();
 
 		//UE_LOG(LogNiagara, Log, TEXT("OnSystemFinished.Broadcast(this);: { %p -  %p - %s"), this, SystemInstance.Get(), *Asset->GetName());
diff --git a/Engine/Plugins/FX/Niagara/Source/Niagara/Public/NiagaraComponent.h b/Engine/Plugins/FX/Niagara/Source/Niagara/Public/NiagaraComponent.h
--- a/Engine/Plugins/FX/Niagara/Source/Niagara/Public/NiagaraComponent.h
+++ b/Engine/Plugins/FX/Niagara/Source/Niagara/Public/NiagaraComponent.h
@@ -841,6 +841,9 @@ private:
 	/** True if this component has been culled by the scalability manager. */
 	uint32 bIsCulledByScalability : 1;
 
+	/** True if this component has been fully completed. If false then we are still pending a final call to OnSystemComplete(). */
+	uint32 bIsFullyComplete : 1;
+
 	/** Flag to mark us as currently changing auto attachment as part of Activate/Deactivate so we don't reset in the OnAttachmentChanged() callback. */
 	//uint32 bIsChangingAutoAttachment : 1;
 
