Added additional safety check during cleanup

diff --git a/Engine/Plugins/FX/Niagara/Source/NiagaraEditorWidgets/Private/SNiagaraEmitterRendererThumbnails.cpp b/Engine/Plugins/FX/Niagara/Source/NiagaraEditorWidgets/Private/SNiagaraEmitterRendererThumbnails.cpp
--- a/Engine/Plugins/FX/Niagara/Source/NiagaraEditorWidgets/Private/SNiagaraEmitterRendererThumbnails.cpp
+++ b/Engine/Plugins/FX/Niagara/Source/NiagaraEditorWidgets/Private/SNiagaraEmitterRendererThumbnails.cpp
@@ -31,8 +31,11 @@ FNiagaraRendererStackEntryObserver::~FNiagaraRendererStackEntryObserver()
 {
 	if(EmitterHandleViewModelWeak.IsValid())
 	{
-		EmitterHandleViewModelWeak.Pin()->GetEmitterStackViewModel()->OnDataObjectChanged().RemoveAll(this);
-		EmitterHandleViewModelWeak.Pin()->GetEmitterStackViewModel()->OnStructureChanged().RemoveAll(this);
+		if(UNiagaraStackViewModel* StackViewModel = EmitterHandleViewModelWeak.Pin()->GetEmitterStackViewModel())
+		{
+			StackViewModel->OnDataObjectChanged().RemoveAll(this);
+			StackViewModel->OnStructureChanged().RemoveAll(this);
+		}
 	}
 
 	UMaterial::OnMaterialCompilationFinished().RemoveAll(this);
