Move bGenerateMetaDataOnCompile from NiagaraEditorSettings to NiagaraSettings as the editor version doesn't want to properly serialize bool values.  Enabled the flag for now in some projects to try to chase down an RI issue, and bumps the niagara script compile version to ensure scripts are flushed.

diff --git a/Engine/Plugins/FX/Niagara/Source/Niagara/Public/NiagaraSettings.h b/Engine/Plugins/FX/Niagara/Source/Niagara/Public/NiagaraSettings.h
--- a/Engine/Plugins/FX/Niagara/Source/Niagara/Public/NiagaraSettings.h
+++ b/Engine/Plugins/FX/Niagara/Source/Niagara/Public/NiagaraSettings.h
@@ -290,6 +290,11 @@ class UNiagaraSettings : public UDeveloperSettings
 	UPROPERTY(config, EditAnywhere, Category = Scalability)
 	TArray<FNiagaraPlatformSetRedirect> PlatformSetRedirects;
 
+#if WITH_EDITORONLY_DATA
+	UPROPERTY(config, EditAnywhere, Category = Niagara)
+	bool bGenerateMetaDataOnCompile = false;
+#endif // WITH_EDITORONLY_DATA
+
 	// Begin UDeveloperSettings Interface
 	NIAGARA_API virtual FName GetCategoryName() const override;
 #if WITH_EDITOR
diff --git a/Engine/Plugins/FX/Niagara/Source/NiagaraEditor/Private/NiagaraCompilationTasks.cpp b/Engine/Plugins/FX/Niagara/Source/NiagaraEditor/Private/NiagaraCompilationTasks.cpp
--- a/Engine/Plugins/FX/Niagara/Source/NiagaraEditor/Private/NiagaraCompilationTasks.cpp
+++ b/Engine/Plugins/FX/Niagara/Source/NiagaraEditor/Private/NiagaraCompilationTasks.cpp
@@ -9,11 +9,11 @@
 #include "Misc/PathViews.h"
 #include "NiagaraCompilationTypes.h"
 #include "NiagaraDigestDatabase.h"
-#include "NiagaraEditorSettings.h"
 #include "NiagaraEmitter.h"
 #include "NiagaraGraphDigest.h"
 #include "NiagaraCompilationBridge.h"
 #include "NiagaraScriptSource.h"
+#include "NiagaraSettings.h"
 #include "NiagaraShader.h"
 #include "NiagaraShaderParametersBuilder.h"
 #include "NiagaraSimulationStageBase.h"
@@ -621,7 +621,7 @@ FNiagaraSystemCompilationTask::FNiagaraSystemCompilationTask(FNiagaraCompilation
 	, DDCRequestOwner(UE::DerivedData::EPriority::Normal)
 	, System_GT(InSystem)
 {
-	bGenerateCompileMetaData = GetDefault<UNiagaraEditorSettings>()->GetGenerateMetaDataOnCompile();
+	bGenerateCompileMetaData = GetDefault<UNiagaraSettings>()->bGenerateMetaDataOnCompile;
 
 	CompilationTaskId = FGuid::NewGuid();
 }
diff --git a/Engine/Plugins/FX/Niagara/Source/NiagaraEditor/Private/NiagaraEditorSettings.cpp b/Engine/Plugins/FX/Niagara/Source/NiagaraEditor/Private/NiagaraEditorSettings.cpp
--- a/Engine/Plugins/FX/Niagara/Source/NiagaraEditor/Private/NiagaraEditorSettings.cpp
+++ b/Engine/Plugins/FX/Niagara/Source/NiagaraEditor/Private/NiagaraEditorSettings.cpp
@@ -919,10 +919,5 @@ bool UNiagaraEditorSettings::GetUpdateStackValuesOnCommitOnly() const
 	return bUpdateStackValuesOnCommitOnly;
 }
 
-bool UNiagaraEditorSettings::GetGenerateMetaDataOnCompile() const
-{
-	return bGenerateMetaDataOnCompile;
-}
-
 #undef LOCTEXT_NAMESPACE
 
diff --git a/Engine/Plugins/FX/Niagara/Source/NiagaraEditor/Public/NiagaraEditorSettings.h b/Engine/Plugins/FX/Niagara/Source/NiagaraEditor/Public/NiagaraEditorSettings.h
--- a/Engine/Plugins/FX/Niagara/Source/NiagaraEditor/Public/NiagaraEditorSettings.h
+++ b/Engine/Plugins/FX/Niagara/Source/NiagaraEditor/Public/NiagaraEditorSettings.h
@@ -579,10 +579,6 @@ private:
 	UPROPERTY(config, EditAnywhere, Category = Niagara)
 	bool bUpdateStackValuesOnCommitOnly = false;
 
-	/** */
-	UPROPERTY(config, EditAnywhere, Category = Niagara)
-	bool bGenerateMetaDataOnCompile = false;
-
 	/** Speeds used for slowing down and speeding up the playback speeds */
 	UPROPERTY(config, EditAnywhere, Category = Niagara)
 	TArray<float> PlaybackSpeeds;
@@ -642,6 +638,5 @@ public:
 	NIAGARAEDITOR_API void SetShowMemoryInfo(bool bInShowInfo);
 	NIAGARAEDITOR_API bool IsShowStatelessInfo() const;
 	NIAGARAEDITOR_API void SetShowStatelessInfo(bool bInShowInfo);
-	NIAGARAEDITOR_API bool GetGenerateMetaDataOnCompile() const;
 
 };
diff --git a/Engine/Source/Runtime/Core/Private/UObject/UE5ReleaseStreamObjectVersion.cpp b/Engine/Source/Runtime/Core/Private/UObject/UE5ReleaseStreamObjectVersion.cpp
--- a/Engine/Source/Runtime/Core/Private/UObject/UE5ReleaseStreamObjectVersion.cpp
+++ b/Engine/Source/Runtime/Core/Private/UObject/UE5ReleaseStreamObjectVersion.cpp
@@ -12,7 +12,7 @@ TMap<FGuid, FGuid> FUE5ReleaseStreamObjectVersion::GetSystemGuids()
 	SystemGuids.Add(DevGuids.MATERIALSHADERMAP_DERIVEDDATA_VER, FGuid("79090A66F9D94B5BB285D49E5D39468E"));
 	SystemGuids.Add(DevGuids.NANITE_DERIVEDDATA_VER, FGuid("AC88CFBDDA614A9C8488F64D2DAF699D"));
 	SystemGuids.Add(DevGuids.NIAGARASHADERMAP_DERIVEDDATA_VER, FGuid("7BCE96ABAA6E46D689306FE673DED4DE"));
-	SystemGuids.Add(DevGuids.Niagara_LatestScriptCompileVersion, FGuid("5CA8BABDBA114F41A3303EBE474C35DE"));
+	SystemGuids.Add(DevGuids.Niagara_LatestScriptCompileVersion, FGuid("14BEE7BD97194A7BBB8374D34DC49915"));
 	SystemGuids.Add(DevGuids.SkeletalMeshDerivedDataVersion, FGuid("BA91E71F642E4813A1B74B80D0448928"));
 	SystemGuids.Add(DevGuids.STATICMESH_DERIVEDDATA_VER, FGuid("14B86848953D46989C69A8B195AD7F6C"));
 
