Don't forcefully load niagara system when getting asset registry tags to prevent deadlocks during emitter merging

diff --git a/Engine/Plugins/FX/Niagara/Source/Niagara/Classes/NiagaraSystem.h b/Engine/Plugins/FX/Niagara/Source/Niagara/Classes/NiagaraSystem.h
--- a/Engine/Plugins/FX/Niagara/Source/Niagara/Classes/NiagaraSystem.h
+++ b/Engine/Plugins/FX/Niagara/Source/Niagara/Classes/NiagaraSystem.h
@@ -970,6 +970,9 @@ protected:
 
 	/** A multicast delegate that is called whenever the effect type or the effect type values are changed */
 	FOnScalabilityChanged OnScalabilityChangedDelegate;
+
+	// If true then we need to manually update the asset registry after the system is fully loaded
+	mutable bool bUpdateTagsAfterLoad = false;
 #endif
 
 	/** The fixed bounding box value for the whole system. When placed in the level and the bounding box is not visible to the camera, the effect is culled from rendering. */
diff --git a/Engine/Plugins/FX/Niagara/Source/Niagara/Private/NiagaraSystem.cpp b/Engine/Plugins/FX/Niagara/Source/Niagara/Private/NiagaraSystem.cpp
--- a/Engine/Plugins/FX/Niagara/Source/Niagara/Private/NiagaraSystem.cpp
+++ b/Engine/Plugins/FX/Niagara/Source/Niagara/Private/NiagaraSystem.cpp
@@ -815,6 +815,14 @@ void UNiagaraSystem::UpdateSystemAfterLoad()
 			}
 		}
 	}
+
+	if (bUpdateTagsAfterLoad && HasAnyFlags(RF_ClassDefaultObject) == false)
+	{
+		if (IAssetRegistry* AssetRegistry = IAssetRegistry::Get())
+		{
+			AssetRegistry->AssetUpdateTags(this, EAssetRegistryTagsCaller::FullUpdate);
+		}
+	}
 #endif
 }
 
@@ -1649,9 +1657,15 @@ void UNiagaraSystem::GetAssetRegistryTags(TArray<FAssetRegistryTag>& OutTags) co
 
 void UNiagaraSystem::GetAssetRegistryTags(FAssetRegistryTagsContext Context) const
 {
-	if (HasAnyFlags(RF_ClassDefaultObject) == false)
+	if (HasAnyFlags(RF_ClassDefaultObject) || !bFullyLoaded)
 	{
-		EnsureFullyLoaded();
+		// We're currently not fully loaded and can't do so in the context of asset tags, as it could cause a deadlock with emitter merging.
+		// The system will update the asset registry manually once fully loaded.
+#if WITH_EDITORONLY_DATA
+		bUpdateTagsAfterLoad = true;
+#endif
+		Super::GetAssetRegistryTags(Context);
+		return;
 	}
 
 #if WITH_EDITOR
