- Remove emitter mode switching between LW & Standard

#rnx

diff --git a/Engine/Plugins/FX/Niagara/Source/NiagaraEditor/Private/NiagaraEditorUtilities.cpp b/Engine/Plugins/FX/Niagara/Source/NiagaraEditor/Private/NiagaraEditorUtilities.cpp
--- a/Engine/Plugins/FX/Niagara/Source/NiagaraEditor/Private/NiagaraEditorUtilities.cpp
+++ b/Engine/Plugins/FX/Niagara/Source/NiagaraEditor/Private/NiagaraEditorUtilities.cpp
@@ -2626,50 +2626,6 @@ void FNiagaraEditorUtilities::ToggleSelectedEmittersIsolated(TSharedRef<FNiagara
 	SystemViewModel->IsolateEmitters(EmittersToIsolate);
 }
 
-ECheckBoxState FNiagaraEditorUtilities::GetSelectedEmittersEmitterModeCheckState(TSharedRef<FNiagaraSystemViewModel> SystemViewModel, ENiagaraEmitterMode EmitterMode)
-{
-	TOptional<ENiagaraEmitterMode> CurrentEmitterMode;
-	const TArray<FGuid>& SelectedEmitterHandleIds = SystemViewModel->GetSelectionViewModel()->GetSelectedEmitterHandleIds();
-	for (const TSharedRef<FNiagaraEmitterHandleViewModel>& EmitterHandleViewModel : SystemViewModel->GetEmitterHandleViewModels())
-	{
-		if (SelectedEmitterHandleIds.Contains(EmitterHandleViewModel->GetId()) == false)
-		{
-			continue;
-		}
-		if (CurrentEmitterMode.IsSet() == false)
-		{
-			CurrentEmitterMode = EmitterHandleViewModel->GetEmitterHandle()->GetEmitterMode();
-		}
-		else
-		{
-			if (CurrentEmitterMode.GetValue() != EmitterHandleViewModel->GetEmitterHandle()->GetEmitterMode())
-			{
-				return ECheckBoxState::Undetermined;
-			}
-		}
-	}
-	if (CurrentEmitterMode.IsSet())
-	{
-		return CurrentEmitterMode.GetValue() == EmitterMode ? ECheckBoxState::Checked : ECheckBoxState::Unchecked;
-	}
-	return ECheckBoxState::Unchecked;
-}
-
-void FNiagaraEditorUtilities::SetSelectedEmittersEmitterMode(TSharedRef<FNiagaraSystemViewModel> SystemViewModel, ENiagaraEmitterMode EmitterMode)
-{
-	FScopedTransaction ScopedTransaction(LOCTEXT("ChangeEmitterMode", "Change emitter mode"));
-	SystemViewModel->GetSystem().Modify();
-
-	const TArray<FGuid>& SelectedEmitterHandleIds = SystemViewModel->GetSelectionViewModel()->GetSelectedEmitterHandleIds();
-	for (const TSharedRef<FNiagaraEmitterHandleViewModel>& EmitterHandleViewModel : SystemViewModel->GetEmitterHandleViewModels())
-	{
-		if (SelectedEmitterHandleIds.Contains(EmitterHandleViewModel->GetId()))
-		{
-			EmitterHandleViewModel->GetEmitterHandle()->SetEmitterMode(SystemViewModel->GetSystem(), EmitterMode);
-		}
-	}
-}
-
 void FNiagaraEditorUtilities::CreateAssetFromEmitter(TSharedRef<FNiagaraEmitterHandleViewModel> EmitterHandleViewModel)
 {
 	TSharedRef<FNiagaraSystemViewModel> SystemViewModel = EmitterHandleViewModel->GetOwningSystemViewModel();
diff --git a/Engine/Plugins/FX/Niagara/Source/NiagaraEditor/Private/NiagaraOverviewNode.cpp b/Engine/Plugins/FX/Niagara/Source/NiagaraEditor/Private/NiagaraOverviewNode.cpp
--- a/Engine/Plugins/FX/Niagara/Source/NiagaraEditor/Private/NiagaraOverviewNode.cpp
+++ b/Engine/Plugins/FX/Niagara/Source/NiagaraEditor/Private/NiagaraOverviewNode.cpp
@@ -299,44 +299,6 @@ void UNiagaraOverviewNode::GetNodeContextMenuActions(class UToolMenu* Menu, clas
 						)
 					)
 				);
-
-				Section.AddSubMenu(
-					"EmitterMode",
-					LOCTEXT("EmitterModeSubMenuLabel", "Emitter Mode..."),
-					FText(),
-					FNewToolMenuDelegate::CreateLambda([OwningSystemViewModel](UToolMenu* InMenu)
-					{
-						if (OwningSystemViewModel.IsValid())
-						{
-							bool bCanExecute = OwningSystemViewModel->GetSelectionViewModel()->GetSelectedEmitterHandleIds().Num() > 0;
-							FToolMenuSection& SubMenuSection = InMenu->AddSection("EmitterModeOptions", LOCTEXT("EmitterModeSubMenuHeader", "Emitter Mode"));
-							SubMenuSection.AddMenuEntry(
-								"SetEmitterModeStandard",
-								LOCTEXT("SetEmitterStandardModeMenuLabel", "Standard Mode"),
-								LOCTEXT("SetEmitterStandardModeMenuToolTip", "Set this emitter to use standard mode."),
-								FSlateIcon(),
-								FUIAction(
-									FExecuteAction::CreateStatic(&FNiagaraEditorUtilities::SetSelectedEmittersEmitterMode, OwningSystemViewModel.ToSharedRef(), ENiagaraEmitterMode::Standard),
-									FCanExecuteAction::CreateLambda([bCanExecute]() { return bCanExecute; }),
-									FGetActionCheckState::CreateStatic(&FNiagaraEditorUtilities::GetSelectedEmittersEmitterModeCheckState, OwningSystemViewModel.ToSharedRef(), ENiagaraEmitterMode::Standard)
-								),
-								EUserInterfaceActionType::RadioButton
-							);
-							SubMenuSection.AddMenuEntry(
-								"SetEmitterModeStateless",
-								LOCTEXT("SetEmitterStatelessModeMenuLabel", "Lightweight Mode"),
-								LOCTEXT("SetEmitterStatelessModeMenuToolTip", "Set this emitter to use lightweight mode."),
-								FSlateIcon(),
-								FUIAction(
-									FExecuteAction::CreateStatic(&FNiagaraEditorUtilities::SetSelectedEmittersEmitterMode, OwningSystemViewModel.ToSharedRef(), ENiagaraEmitterMode::Stateless),
-									FCanExecuteAction::CreateLambda([bCanExecute]() { return bCanExecute; }),
-									FGetActionCheckState::CreateStatic(&FNiagaraEditorUtilities::GetSelectedEmittersEmitterModeCheckState, OwningSystemViewModel.ToSharedRef(), ENiagaraEmitterMode::Stateless)
-								),
-								EUserInterfaceActionType::RadioButton
-							);
-						}
-					})
-				);
 			}
 		}
 	}
diff --git a/Engine/Plugins/FX/Niagara/Source/NiagaraEditor/Public/NiagaraEditorUtilities.h b/Engine/Plugins/FX/Niagara/Source/NiagaraEditor/Public/NiagaraEditorUtilities.h
--- a/Engine/Plugins/FX/Niagara/Source/NiagaraEditor/Public/NiagaraEditorUtilities.h
+++ b/Engine/Plugins/FX/Niagara/Source/NiagaraEditor/Public/NiagaraEditorUtilities.h
@@ -309,9 +309,6 @@ namespace FNiagaraEditorUtilities
 	ECheckBoxState GetSelectedEmittersIsolatedCheckState(TSharedRef<FNiagaraSystemViewModel> SystemViewModel);
 	void ToggleSelectedEmittersIsolated(TSharedRef<FNiagaraSystemViewModel> SystemViewModel);
 
-	ECheckBoxState GetSelectedEmittersEmitterModeCheckState(TSharedRef<FNiagaraSystemViewModel> SystemViewModel, ENiagaraEmitterMode EmitterMode);
-	void SetSelectedEmittersEmitterMode(TSharedRef<FNiagaraSystemViewModel> SystemViewModel, ENiagaraEmitterMode EmitterMode);
-
 	void CreateAssetFromEmitter(TSharedRef<FNiagaraEmitterHandleViewModel> EmitterHandleViewModel);
 
 	NIAGARAEDITOR_API void WarnWithToastAndLog(FText WarningMessage);
