Fixing potential mem leaks when the work is paused and we manually tick FX.
Adding OnWorldTickEnd callback to world manager and moving a bunch of things from PostActorTick() that should happen even when paused.
Including flushing the deferred deletion queues.

diff --git a/Engine/Plugins/FX/Niagara/Source/Niagara/Private/NiagaraWorldManager.cpp b/Engine/Plugins/FX/Niagara/Source/Niagara/Private/NiagaraWorldManager.cpp
--- a/Engine/Plugins/FX/Niagara/Source/Niagara/Private/NiagaraWorldManager.cpp
+++ b/Engine/Plugins/FX/Niagara/Source/Niagara/Private/NiagaraWorldManager.cpp
@@ -251,6 +251,7 @@ FDelegateHandle FNiagaraWorldManager::OnPostWorldCleanupHandle;
 FDelegateHandle FNiagaraWorldManager::OnPreWorldFinishDestroyHandle;
 FDelegateHandle FNiagaraWorldManager::OnWorldBeginTearDownHandle;
 FDelegateHandle FNiagaraWorldManager::TickWorldStartHandle;
+FDelegateHandle FNiagaraWorldManager::TickWorldEndHandle;
 FDelegateHandle FNiagaraWorldManager::TickWorldHandle;
 FDelegateHandle FNiagaraWorldManager::OnWorldPreActorTickHandle;
 FDelegateHandle FNiagaraWorldManager::OnWorldPreSendAllEndOfFrameUpdatesHandle;
@@ -463,6 +464,7 @@ void FNiagaraWorldManager::OnStartup()
 	OnPreWorldFinishDestroyHandle = FWorldDelegates::OnPreWorldFinishDestroy.AddStatic(&FNiagaraWorldManager::OnPreWorldFinishDestroy);
 	OnWorldBeginTearDownHandle = FWorldDelegates::OnWorldBeginTearDown.AddStatic(&FNiagaraWorldManager::OnWorldBeginTearDown);
 	TickWorldStartHandle = FWorldDelegates::OnWorldTickStart.AddStatic(&FNiagaraWorldManager::TickWorldStart);
+	TickWorldEndHandle = FWorldDelegates::OnWorldTickEnd.AddStatic(&FNiagaraWorldManager::TickWorldEnd);
 	TickWorldHandle = FWorldDelegates::OnWorldPostActorTick.AddStatic(&FNiagaraWorldManager::TickWorld);
 	OnWorldPreActorTickHandle = FWorldDelegates::OnWorldPreActorTick.AddStatic(&FNiagaraWorldManager::OnWorldPreActorTick);
 	OnWorldPreSendAllEndOfFrameUpdatesHandle = FWorldDelegates::OnWorldPreSendAllEndOfFrameUpdates.AddLambda(
@@ -1049,6 +1051,11 @@ void FNiagaraWorldManager::TickWorldStart(UWorld* World, ELevelTick TickType, fl
 	Get(World)->TickStart(DeltaSeconds);
 }
 
+void FNiagaraWorldManager::TickWorldEnd(UWorld* World, ELevelTick TickType, float DeltaSeconds)
+{
+	Get(World)->TickEnd(DeltaSeconds);
+}
+
 void FNiagaraWorldManager::TickWorld(UWorld* World, ELevelTick TickType, float DeltaSeconds)
 {
 	Get(World)->PostActorTick(DeltaSeconds);
@@ -1117,6 +1124,41 @@ void FNiagaraWorldManager::TickStart(float DeltaSeconds)
 #endif
 }
 
+void FNiagaraWorldManager::TickEnd(float DeltaSeconds)
+{
+	CSV_SCOPED_TIMING_STAT_EXCLUSIVE(Effects);
+	LLM_SCOPE(ELLMTag::Niagara);
+	SCOPE_CYCLE_COUNTER(STAT_NiagaraWorldManTick);
+
+	//Do any end of tick bookkeeping that must occur even when paused.	
+	FNiagaraSharedObject::FlushDeletionList();
+	NiagaraWorldManagerInternal::ExecuteGlobalDeferredCallbacks();
+
+	// Delete any instances that were pending deletion
+	//-TODO: This could be done after each system sim has run
+	DeferredDeletionQueue.Empty();
+
+#if WITH_NIAGARA_DEBUGGER
+	// Tick debug HUD for the world
+	if (NiagaraDebugHud != nullptr)
+	{
+		QUICK_SCOPE_CYCLE_COUNTER(STAT_NiagaraDebugHUD);
+		CSV_SCOPED_TIMING_STAT(Particles, CoreSystems_NiagaraDebugHUD);
+		NiagaraDebugHud->GatherSystemInfo();
+	}
+#endif
+
+	if (DebugPlaybackMode == ENiagaraDebugPlaybackMode::Step)
+	{
+		RequestedDebugPlaybackMode = ENiagaraDebugPlaybackMode::Paused;
+		DebugPlaybackMode = ENiagaraDebugPlaybackMode::Paused;
+	}
+
+	HandleCSVStats(DeltaSeconds);
+
+	DataChannelManager->EndFrame(DeltaSeconds);
+}
+
 void FNiagaraWorldManager::PreActorTick(ELevelTick InLevelTick, float InDeltaSeconds)
 {	
 	NiagaraWorldManagerInternal::ExecuteGlobalDeferredCallbacks();
@@ -1195,26 +1237,6 @@ void FNiagaraWorldManager::PostActorTick(float DeltaSeconds)
 	// Clear cached player view location list, it should never be used outside of the world tick
 	CachedViewInfo.Reset();
 
-	// Delete any instances that were pending deletion
-	//-TODO: This could be done after each system sim has run
-	DeferredDeletionQueue.Empty();
-
-#if WITH_NIAGARA_DEBUGGER
-	// Tick debug HUD for the world
-	if (NiagaraDebugHud != nullptr)
-	{
-		QUICK_SCOPE_CYCLE_COUNTER(STAT_NiagaraDebugHUD);
-		CSV_SCOPED_TIMING_STAT(Particles, CoreSystems_NiagaraDebugHUD);
-		NiagaraDebugHud->GatherSystemInfo();
-	}
-#endif
-
-	if ( DebugPlaybackMode == ENiagaraDebugPlaybackMode::Step )
-	{
-		RequestedDebugPlaybackMode = ENiagaraDebugPlaybackMode::Paused;
-		DebugPlaybackMode = ENiagaraDebugPlaybackMode::Paused;
-	}
-
 	for (auto it = CullProxyMap.CreateIterator(); it; ++it)
 	{
 		UNiagaraSystem* System = it->Key;
@@ -1232,10 +1254,6 @@ void FNiagaraWorldManager::PostActorTick(float DeltaSeconds)
 			it.RemoveCurrent();
 		}
 	}
-
-	HandleCSVStats(DeltaSeconds);
-
-	DataChannelManager->EndFrame(DeltaSeconds);
 }
 
 void FNiagaraWorldManager::HandleCSVStats(float DeltaSeconds)
@@ -1459,8 +1477,6 @@ void FNiagaraWorldManager::ExecutePreTickWork(ETickingGroup TickGroup, float Del
 			bPoolIsPrimed = true;
 		}
 
-		FNiagaraSharedObject::FlushDeletionList();
-
 #if WITH_EDITOR //PLATFORM_DESKTOP
 		{
 			QUICK_SCOPE_CYCLE_COUNTER(STAT_Niagara_IsThisApplicationForeground);
diff --git a/Engine/Plugins/FX/Niagara/Source/Niagara/Public/NiagaraWorldManager.h b/Engine/Plugins/FX/Niagara/Source/Niagara/Public/NiagaraWorldManager.h
--- a/Engine/Plugins/FX/Niagara/Source/Niagara/Public/NiagaraWorldManager.h
+++ b/Engine/Plugins/FX/Niagara/Source/Niagara/Public/NiagaraWorldManager.h
@@ -135,8 +135,11 @@ public:
 	void MarkSimulationForPostActorWork(FNiagaraSystemSimulation* SystemSimulation);
 	void MarkSimulationsForEndOfFrameWait(FNiagaraSystemSimulation* SystemSimulation);
 
-	/** Called at the very beginning of world ticking. */
+	/** Called at the very beginning of world ticking, even when paused. */
 	void TickStart(float DeltaSeconds);
+	
+	/** Called at the very end of world ticking, even when paused. */
+	void TickEnd(float DeltaSeconds);
 
 	void Tick(ENiagaraWorldManagerTickPass TickPass, ETickingGroup TickGroup, float DeltaSeconds, const FGraphEventRef& MyCompletionGraphEvent);
 private:
@@ -335,9 +338,12 @@ private:
 
 	static void OnWorldPreActorTick(UWorld* InWorld, ELevelTick InLevelTick, float InDeltaSeconds);
 
-	// Callback for when a world is ticked.
+	// Callback for when a world tick is begun.
 	static void TickWorldStart(UWorld* World, ELevelTick TickType, float DeltaSeconds);
 
+	// Callback for when a world tick is ended.
+	static void TickWorldEnd(UWorld* World, ELevelTick TickType, float DeltaSeconds);
+
 	// Callback for when a world is ticked.
 	static void TickWorld(UWorld* World, ELevelTick TickType, float DeltaSeconds);
 
@@ -379,6 +385,7 @@ private:
 	static FDelegateHandle OnPreWorldFinishDestroyHandle;
 	static FDelegateHandle OnWorldBeginTearDownHandle;
 	static FDelegateHandle TickWorldStartHandle;
+	static FDelegateHandle TickWorldEndHandle;
 	static FDelegateHandle TickWorldHandle;
 	static FDelegateHandle OnWorldPreActorTickHandle;
 	static FDelegateHandle OnWorldPreSendAllEndOfFrameUpdatesHandle;
