- Remove dead code

#rnx

diff --git a/Engine/Plugins/FX/Niagara/Source/NiagaraEditor/Private/ViewModels/Stack/NiagaraStackFunctionInput.cpp b/Engine/Plugins/FX/Niagara/Source/NiagaraEditor/Private/ViewModels/Stack/NiagaraStackFunctionInput.cpp
--- a/Engine/Plugins/FX/Niagara/Source/NiagaraEditor/Private/ViewModels/Stack/NiagaraStackFunctionInput.cpp
+++ b/Engine/Plugins/FX/Niagara/Source/NiagaraEditor/Private/ViewModels/Stack/NiagaraStackFunctionInput.cpp
@@ -845,36 +845,6 @@ void UNiagaraStackFunctionInput::RefreshChildrenInternal(const TArray<UNiagaraSt
 	RefreshFromMetaData(NewIssues);
 	RefreshValues();
 
-	// If we keep around rapid iteration parameters that should go away, it bloats the parameter store.
-	if (RapidIterationParameter.IsValid() && InputValues.Mode != EValueMode::Local)
-	{
-		bool bHasStaleRapidIterationVar = false;
-		for (TWeakObjectPtr<UNiagaraScript> Script : AffectedScripts)
-		{
-			if (Script->RapidIterationParameters.IndexOf(RapidIterationParameter) != INDEX_NONE)
-			{
-				bHasStaleRapidIterationVar = true;
-			}
-		}
-
-		if (bHasStaleRapidIterationVar)
-		{
-			/*
-			NewIssues.Add(FStackIssue(
-				EStackIssueSeverity::Warning,
-				LOCTEXT("InvalidRapidIterationOverrideShort", "Invalid Rapid Iteration override"),
-				LOCTEXT("InvalidRapidIterationOverrideLong", "There is a rapid iteration value left around that could make the UI not match the actual state. Hit 'Fix issue' to clear out the old, bad value."),
-				GetStackEditorDataKey(),
-				false,
-				{ FStackIssueFix(
-					LOCTEXT("ResetRapidIterationInputFix", "Fix stale value"),
-					FStackIssueFixDelegate::CreateLambda([this]() { this->RemoveRapidIterationParametersForAffectedScripts(true); }))
-				}
-			));
-			*/
-		}
-	}
-
 	// we check for local ints acting as enums whether the local value is allowed with the current metadata
 	// it can happen if a module version changes or a module in use is modified to lose options that were in use
 	if(InputMetaData.IsSet() && InputMetaData->WidgetCustomization.WidgetType == ENiagaraInputWidgetType::EnumStyle && GetInputType() == FNiagaraTypeDefinition::GetIntDef() && GetValueMode() == EValueMode::Local)
