- Fix for performance mode button not working

#rnx

diff --git a/Engine/Plugins/FX/Niagara/Source/NiagaraEditor/Private/Toolkits/NiagaraSystemToolkit.cpp b/Engine/Plugins/FX/Niagara/Source/NiagaraEditor/Private/Toolkits/NiagaraSystemToolkit.cpp
--- a/Engine/Plugins/FX/Niagara/Source/NiagaraEditor/Private/Toolkits/NiagaraSystemToolkit.cpp
+++ b/Engine/Plugins/FX/Niagara/Source/NiagaraEditor/Private/Toolkits/NiagaraSystemToolkit.cpp
@@ -72,8 +72,6 @@ DECLARE_CYCLE_STAT(TEXT("Niagara - SystemToolkit - OnApply"), STAT_NiagaraEditor
 const FName FNiagaraSystemToolkit::DefaultModeName(TEXT("Default"));
 const FName FNiagaraSystemToolkit::ScalabilityModeName(TEXT("Scalability"));
 
-IConsoleVariable* FNiagaraSystemToolkit::VmStatEnabledVar = IConsoleManager::Get().FindConsoleVariable(TEXT("vm.DetailedVMScriptStats"));
-
 static int32 GbLogNiagaraSystemChanges = 0;
 static FAutoConsoleVariableRef CVarLogNiagaraSystemChanges(
 	TEXT("fx.LogNiagaraSystemChanges"),
@@ -1042,15 +1040,12 @@ void FNiagaraSystemToolkit::ClearStatPerformance()
 void FNiagaraSystemToolkit::ToggleStatPerformance()
 {
 	const bool IsEnabled = IsStatPerformanceChecked();
-	if (VmStatEnabledVar)
+	if (SystemViewModel.IsValid())
 	{
-		VmStatEnabledVar->Set(!IsEnabled);
+		SystemViewModel->SetPerformanceStatsEnabled(!IsEnabled);
 
-		if (SystemViewModel.IsValid())
-		{
-			const bool bCompileForEdit = !IsEnabled ? bProfileInEditMode : true;
-			SystemViewModel->SetCompileForEdit(bCompileForEdit);
-		}
+		const bool bCompileForEdit = !IsEnabled ? bProfileInEditMode : true;
+		SystemViewModel->SetCompileForEdit(bCompileForEdit);
 	}
 	if (IsStatPerformanceGPUChecked() == IsEnabled)
 	{
@@ -1100,7 +1095,7 @@ bool FNiagaraSystemToolkit::IsStatPerformanceModeAbsolute()
 
 bool FNiagaraSystemToolkit::IsStatPerformanceChecked()
 {
-	return VmStatEnabledVar ? VmStatEnabledVar->GetBool() : false;
+	return SystemViewModel.IsValid() ? SystemViewModel->IsPerformanceStatsEnabled() : false;
 }
 
 void FNiagaraSystemToolkit::ToggleStatPerformanceGPU()
diff --git a/Engine/Plugins/FX/Niagara/Source/NiagaraEditor/Public/ViewModels/NiagaraSystemViewModel.h b/Engine/Plugins/FX/Niagara/Source/NiagaraEditor/Public/ViewModels/NiagaraSystemViewModel.h
--- a/Engine/Plugins/FX/Niagara/Source/NiagaraEditor/Public/ViewModels/NiagaraSystemViewModel.h
+++ b/Engine/Plugins/FX/Niagara/Source/NiagaraEditor/Public/ViewModels/NiagaraSystemViewModel.h
@@ -353,6 +353,11 @@ public:
 	/** Returns what level of performance stats we have available in the editor. */
 	NIAGARAEDITOR_API ENiagaraPerformanceStateScopeMode GetPerformanceStateScopeMode() const;
 
+	/** Are the performance stats enabled? */
+	bool IsPerformanceStatsEnabled() const { return bPerformanceStatsEnabled; }
+	/** Set if performance stats should be enabled or not */
+	void SetPerformanceStatsEnabled(bool bEnabled) { bPerformanceStatsEnabled = bEnabled; }
+
 	/** Clear the captures stats for the current system. */
 	void ClearSystemStats();
 
@@ -834,6 +839,9 @@ private:
 	/** Specifies that this view model is for data processing only and will not be displayed in the UI. */
 	bool bIsForDataProcessingOnly;
 
+	/** If performance stats should be enable in the editor or not. */
+	bool bPerformanceStatsEnabled = false;
+
 	/** ViewModel for caching the selected graph state and notifying when it changes. */
 	TSharedPtr<FNiagaraSystemGraphSelectionViewModel> SystemGraphSelectionViewModel;
 
diff --git a/Engine/Plugins/FX/Niagara/Source/NiagaraEditorWidgets/Private/Stack/SNiagaraStackRowPerfWidget.cpp b/Engine/Plugins/FX/Niagara/Source/NiagaraEditorWidgets/Private/Stack/SNiagaraStackRowPerfWidget.cpp
--- a/Engine/Plugins/FX/Niagara/Source/NiagaraEditorWidgets/Private/Stack/SNiagaraStackRowPerfWidget.cpp
+++ b/Engine/Plugins/FX/Niagara/Source/NiagaraEditorWidgets/Private/Stack/SNiagaraStackRowPerfWidget.cpp
@@ -22,7 +22,7 @@
 
 #define LOCTEXT_NAMESPACE "SNiagaraStackRowPerfWidget"
 
-IConsoleVariable* SNiagaraStackRowPerfWidget::StatEnabledVar = IConsoleManager::Get().FindConsoleVariable(TEXT("vm.DetailedVMScriptStats"));
+IConsoleVariable* SNiagaraStackRowPerfWidget::StatEnabledVar = IConsoleManager::Get().FindConsoleVariable(TEXT("fx.Niagara.EditorPerformanceStatsEnabled"));
 
 void SNiagaraStackRowPerfWidget::Construct(const FArguments& InArgs, UNiagaraStackEntry* InStackEntry, ENiagaraPerformanceStateScopeMode InStatScopeMode)
 {
@@ -241,7 +241,7 @@ bool SNiagaraStackRowPerfWidget::HasPerformanceData() const
 		return false;
 	}
 
-	bool IsPerfCaptureEnabled = StatEnabledVar && StatEnabledVar->GetBool() && (StackEntry->GetSystemViewModel()->GetPerformanceStateScopeMode() >= StatScopeMode);
+	const bool IsPerfCaptureEnabled = StackEntry->GetSystemViewModel()->IsPerformanceStatsEnabled() && (StackEntry->GetSystemViewModel()->GetPerformanceStateScopeMode() >= StatScopeMode);
 	if (IsPerfCaptureEnabled && (IsSystemStack() || IsEmitterStack()))
 	{
 		return (IsGroupHeaderEntry() || IsModuleEntry()) && (StackEntry->GetExecutionSubcategoryName() != UNiagaraStackEntry::FExecutionSubcategoryNames::Settings);
