Fixed drag & drop of parameter onto stack for 'set existing parameter' module not working

diff --git a/Engine/Plugins/FX/Niagara/Source/NiagaraEditor/Private/ViewModels/Stack/NiagaraStackScriptItemGroup.cpp b/Engine/Plugins/FX/Niagara/Source/NiagaraEditor/Private/ViewModels/Stack/NiagaraStackScriptItemGroup.cpp
--- a/Engine/Plugins/FX/Niagara/Source/NiagaraEditor/Private/ViewModels/Stack/NiagaraStackScriptItemGroup.cpp
+++ b/Engine/Plugins/FX/Niagara/Source/NiagaraEditor/Private/ViewModels/Stack/NiagaraStackScriptItemGroup.cpp
@@ -148,7 +148,7 @@ public:
 		return MakeShareable(new FScriptGroupAddAction(CreationArguments));
 	}
 
-	static TSharedRef<FScriptGroupAddAction> CreateExistingParameterModuleAction(FNiagaraVariable ParameterVariable)
+	static TSharedRef<FScriptGroupAddAction> CreateSetExistingParameterModuleAction(FNiagaraVariable ParameterVariable)
 	{
 		FText Category = LOCTEXT("ExistingParameterModuleCategory", "Set Specific Parameters");
 
@@ -161,6 +161,9 @@ public:
 		CreationArguments.Description = Description;
 		CreationArguments.Categories = { Category.ToString() };
 		CreationArguments.bIsInLibrary = true;
+		CreationArguments.ModuleParameterVariable = ParameterVariable;
+		CreationArguments.bRenameParameterOnAdd = false;
+		CreationArguments.SpecificActionType = ESpecificActionType::NewSetSpecificParameters;
 		return MakeShareable(new FScriptGroupAddAction(CreationArguments));
 	}
 
@@ -303,7 +306,7 @@ public:
 		return SpecificActionType.IsSet() && SpecificActionType.GetValue() == ESpecificActionType::NewScratchModule;
 	}
 
-	bool IsNewSetSpecificModuleAction() const
+	bool IsNewSetSpecificParametersModuleAction() const
 	{
 		return SpecificActionType.IsSet() && SpecificActionType.GetValue() == ESpecificActionType::NewSetSpecificParameters;
 	}
@@ -441,10 +444,6 @@ public:
 
 			NewModuleNode = FNiagaraStackGraphUtilities::AddScriptModuleToStack(AddScriptModuleToStackArgs);
 		}
-		else if (ScriptGroupAddAction->GetModuleParameterVariable().IsValid())
-		{
-			NewModuleNode = AddParameterModule(ScriptGroupAddAction->GetModuleParameterVariable(), ScriptGroupAddAction->GetRenameParameterOnAdd(), TargetIndex);
-		}
 		else if (ScriptGroupAddAction->GetModuleScript() != nullptr)
 		{
 			NewModuleNode = FNiagaraStackGraphUtilities::AddScriptModuleToStack(ScriptGroupAddAction->GetModuleScript(), *OutputNode, TargetIndex);
@@ -510,9 +509,16 @@ public:
 				ScratchPadScriptViewModel->SetIsPendingRename(true);
 			}
 		}
-		else if (ScriptGroupAddAction->IsNewSetSpecificModuleAction())
+		else if (ScriptGroupAddAction->IsNewSetSpecificParametersModuleAction())
 		{
-			NewModuleNode = AddParameterModule(FNiagaraVariable(), false, TargetIndex);
+			TOptional<const FNiagaraVariable> InitializeVariable;
+
+			if(ScriptGroupAddAction->GetModuleParameterVariable().IsValid())
+			{
+				InitializeVariable = ScriptGroupAddAction->GetModuleParameterVariable();
+			}
+			
+			NewModuleNode = AddParameterModule(InitializeVariable, ScriptGroupAddAction->GetRenameParameterOnAdd(), TargetIndex);
 		}
 
 		if (NewModuleNode != nullptr)
@@ -531,15 +537,15 @@ private:
 		return FNiagaraStackGraphUtilities::AddScriptModuleToStack(AddScriptModuleToStackArgs);
 	}
 
-	UNiagaraNodeFunctionCall* AddParameterModule(const FNiagaraVariable& ParameterVariable, bool bRenameParameterOnAdd, int32 TargetIndex)
+	UNiagaraNodeFunctionCall* AddParameterModule(TOptional<const FNiagaraVariable>& ParameterVariable, bool bRenameParameterOnAdd, int32 TargetIndex)
 	{
 		TArray<FNiagaraVariable> Vars;
 		TArray<FString> DefaultVals;
 
-		if (ParameterVariable.IsValid())
+		if (ParameterVariable.IsSet() && ParameterVariable.GetValue().IsValid())
 		{
-			Vars.Add(ParameterVariable);
-			DefaultVals.Add(FNiagaraConstants::GetAttributeDefaultValue(ParameterVariable));
+			Vars.Add(ParameterVariable.GetValue());
+			DefaultVals.Add(FNiagaraConstants::GetAttributeDefaultValue(ParameterVariable.GetValue()));
 		}
 
 		UNiagaraNodeAssignment* NewAssignmentModule = FNiagaraStackGraphUtilities::AddParameterModuleToStack(Vars, *OutputNode, TargetIndex,DefaultVals );
@@ -1394,7 +1400,7 @@ TOptional<UNiagaraStackEntry::FDropRequestResponse> UNiagaraStackScriptItemGroup
 	TSharedPtr<FNiagaraParameterAction> ParameterAction = StaticCastSharedPtr<FNiagaraParameterAction>(ParameterDragDropOp->GetSourceAction());
 
 	int32 TargetIndex = GetTargetIndexForTargetEntry(this, TargetEntry, DropRequest.DropZone);
-	TSharedRef<FScriptGroupAddAction> AddAction = FScriptGroupAddAction::CreateExistingParameterModuleAction(ParameterAction->GetParameter());
+	TSharedRef<FScriptGroupAddAction> AddAction = FScriptGroupAddAction::CreateSetExistingParameterModuleAction(ParameterAction->GetParameter());
 	AddUtilities->ExecuteAddAction(AddAction, TargetIndex);
 
 	return FDropRequestResponse(DropRequest.DropZone);
