- Add ID support to Sim Cache Reader

#tests Ran ET which includes testing ID

diff --git a/Engine/Plugins/FX/Niagara/Shaders/Private/NiagaraDataInterfaceSimCacheReaderTemplate.ush b/Engine/Plugins/FX/Niagara/Shaders/Private/NiagaraDataInterfaceSimCacheReaderTemplate.ush
--- a/Engine/Plugins/FX/Niagara/Shaders/Private/NiagaraDataInterfaceSimCacheReaderTemplate.ush
+++ b/Engine/Plugins/FX/Niagara/Shaders/Private/NiagaraDataInterfaceSimCacheReaderTemplate.ush
@@ -7,7 +7,7 @@ Buffer<uint>	{ParameterName}_CacheData;
 
 struct {ParameterName}_EmitterFrameData
 {
-	uint	NumInstances;
+	int		NumInstances;
 	uint	Int32DataOffset;
 	uint	FloatDataOffset;
 	uint	HalfDataOffset;
@@ -39,7 +39,7 @@ void GetEmitterIndex_{ParameterName}(out int EmitterIndex)
 	if ( FrameIndex >= 0 && FrameIndex < {ParameterName}_NumFrames && EmitterIndex >= 0 && EmitterIndex < {ParameterName}_NumEmitters )
 	{
 		uint FrameDataOffset		= (FrameIndex * 4 * {ParameterName}_NumEmitters) + EmitterIndex;
-		FrameData.NumInstances		= {ParameterName}_CacheData[FrameDataOffset + 0];
+		FrameData.NumInstances		= int({ParameterName}_CacheData[FrameDataOffset + 0]);
 		FrameData.Int32DataOffset	= {ParameterName}_CacheData[FrameDataOffset + 1];
 		FrameData.FloatDataOffset	= {ParameterName}_CacheData[FrameDataOffset + 2];
 		FrameData.HalfDataOffset	= {ParameterName}_CacheData[FrameDataOffset + 3];
@@ -123,3 +123,18 @@ void GetInstanceFloat4_{ParameterName}(int FrameIndex, int EmitterIndex, int Att
 		Value.w = asfloat({ParameterName}_CacheData[FrameData.FloatDataOffset + ((AttributeIndex + 3) * FrameData.NumInstances) + InstanceIndex]);
 	}
 }
+
+void GetInstanceID_{ParameterName}(int FrameIndex, int EmitterIndex, int AttributeIndex, int InstanceIndex, out bool bValid, out NiagaraID Value)
+{
+	bValid = false;
+	Value.Index = -1;
+	Value.AcquireTag = -1;
+
+	{ParameterName}_EmitterFrameData FrameData = GetEmitterFrameData_{ParameterName}(FrameIndex, EmitterIndex);
+	if ( InstanceIndex >= 0 && InstanceIndex < FrameData.NumInstances )
+	{
+		bValid = true;
+		Value.Index = asint({ParameterName}_CacheData[FrameData.Int32DataOffset + ((AttributeIndex + 0) * FrameData.NumInstances) + InstanceIndex]);
+		Value.AcquireTag = asint({ParameterName}_CacheData[FrameData.Int32DataOffset + ((AttributeIndex + 1) * FrameData.NumInstances) + InstanceIndex]);
+	}
+}
diff --git a/Engine/Plugins/FX/Niagara/Source/Niagara/Private/DataInterface/NiagaraDataInterfaceSimCacheReader.cpp b/Engine/Plugins/FX/Niagara/Source/Niagara/Private/DataInterface/NiagaraDataInterfaceSimCacheReader.cpp
--- a/Engine/Plugins/FX/Niagara/Source/Niagara/Private/DataInterface/NiagaraDataInterfaceSimCacheReader.cpp
+++ b/Engine/Plugins/FX/Niagara/Source/Niagara/Private/DataInterface/NiagaraDataInterfaceSimCacheReader.cpp
@@ -33,6 +33,7 @@ namespace NDISimCacheReaderLocal
 	static const FName GetInstanceColor("GetInstanceColor");
 	static const FName GetInstanceQuat("GetInstanceQuat");
 	static const FName GetInstancePosition("GetInstancePosition");
+	static const FName GetInstanceID("GetInstanceID");
 
 	static const FName NAME_Attribute("Attribute");
 
@@ -78,14 +79,15 @@ namespace NDISimCacheReaderLocal
 	{
 		static const FGetFunctionInfo Infos[] =
 		{
-			FGetFunctionInfo(GetInstanceInt,		FNiagaraTypeDefinition::GetIntDef(),		TEXT("int"),	TEXT("GetInstanceInt")),
-			FGetFunctionInfo(GetInstanceFloat,		FNiagaraTypeDefinition::GetFloatDef(),		TEXT("float"),	TEXT("GetInstanceFloat")),
-			FGetFunctionInfo(GetInstanceVector2,	FNiagaraTypeDefinition::GetVec2Def(),		TEXT("float2"),	TEXT("GetInstanceFloat2")),
-			FGetFunctionInfo(GetInstanceVector3,	FNiagaraTypeDefinition::GetVec3Def(),		TEXT("float3"),	TEXT("GetInstanceFloat3")),
-			FGetFunctionInfo(GetInstanceVector4,	FNiagaraTypeDefinition::GetVec4Def(),		TEXT("float4"),	TEXT("GetInstanceFloat4")),
-			FGetFunctionInfo(GetInstanceColor,		FNiagaraTypeDefinition::GetColorDef(),		TEXT("float4"),	TEXT("GetInstanceFloat4")),
-			FGetFunctionInfo(GetInstanceQuat,		FNiagaraTypeDefinition::GetQuatDef(),		TEXT("float4"),	TEXT("GetInstanceFloat4")),
-			FGetFunctionInfo(GetInstancePosition,	FNiagaraTypeDefinition::GetPositionDef(),	TEXT("float3"),	TEXT("GetInstanceFloat3")),
+			FGetFunctionInfo(GetInstanceInt,		FNiagaraTypeDefinition::GetIntDef(),		TEXT("int"),		TEXT("GetInstanceInt")),
+			FGetFunctionInfo(GetInstanceFloat,		FNiagaraTypeDefinition::GetFloatDef(),		TEXT("float"),		TEXT("GetInstanceFloat")),
+			FGetFunctionInfo(GetInstanceVector2,	FNiagaraTypeDefinition::GetVec2Def(),		TEXT("float2"),		TEXT("GetInstanceFloat2")),
+			FGetFunctionInfo(GetInstanceVector3,	FNiagaraTypeDefinition::GetVec3Def(),		TEXT("float3"),		TEXT("GetInstanceFloat3")),
+			FGetFunctionInfo(GetInstanceVector4,	FNiagaraTypeDefinition::GetVec4Def(),		TEXT("float4"),		TEXT("GetInstanceFloat4")),
+			FGetFunctionInfo(GetInstanceColor,		FNiagaraTypeDefinition::GetColorDef(),		TEXT("float4"),		TEXT("GetInstanceFloat4")),
+			FGetFunctionInfo(GetInstanceQuat,		FNiagaraTypeDefinition::GetQuatDef(),		TEXT("float4"),		TEXT("GetInstanceFloat4")),
+			FGetFunctionInfo(GetInstancePosition,	FNiagaraTypeDefinition::GetPositionDef(),	TEXT("float3"),		TEXT("GetInstanceFloat3")),
+			FGetFunctionInfo(GetInstanceID,			FNiagaraTypeDefinition::GetIDDef(),			TEXT("NiagaraID"),	TEXT("GetInstanceID")),
 		};
 		return MakeArrayView(Infos);
 	}
@@ -115,7 +117,7 @@ namespace NDISimCacheReaderLocal
 		const FNiagaraVariableBase Variable(FunctionNameToTypeDef(FunctionName), AttributeName);
 		if ( const FNiagaraSimCacheVariable* CacheVariable = FNiagaraSimCacheAttributeReaderHelper::FindVariable(SimCache, Variable, EmitterIndex) )
 		{
-			if (Variable.GetType() == FNiagaraTypeDefinition::GetIntDef())
+			if (Variable.GetType() == FNiagaraTypeDefinition::GetIntDef() || Variable.GetType() == FNiagaraTypeDefinition::GetIDDef())
 			{
 				return CacheVariable->Int32Offset;
 			}
@@ -456,6 +458,11 @@ void UNiagaraDataInterfaceSimCacheReader::GetVMExternalFunction(const FVMExterna
 		const int32 AttributeIndex = AttributeIndexFromVariable(InstanceSimCache, BindingInfo.Name, BindingInfo.FunctionSpecifiers[0].Value, InstanceData->EmitterIndex);
 		OutFunc = FVMExternalFunction::CreateUObject(this, &UNiagaraDataInterfaceSimCacheReader::VMGetInstancePosition, AttributeIndex, false);
 	}
+	else if (BindingInfo.Name == GetInstanceID)
+	{
+		const int32 AttributeIndex = AttributeIndexFromVariable(InstanceSimCache, BindingInfo.Name, BindingInfo.FunctionSpecifiers[0].Value, InstanceData->EmitterIndex);
+		OutFunc = FVMExternalFunction::CreateUObject(this, &UNiagaraDataInterfaceSimCacheReader::VMGetInstanceID, AttributeIndex, false);
+	}
 }
 
 #if WITH_EDITORONLY_DATA
@@ -669,4 +676,10 @@ void UNiagaraDataInterfaceSimCacheReader::VMGetInstancePosition(FVectorVMExterna
 	VMInstanceGetFunction<FVector3f>(Context, FVector3f::ZeroVector, AttributeIndex, [](FNiagaraSimCacheAttributeReaderHelper& AttributeReader, int32 InstanceIndex) { return AttributeReader.ReadFloat3f(InstanceIndex); });
 }
 
+void UNiagaraDataInterfaceSimCacheReader::VMGetInstanceID(FVectorVMExternalFunctionContext& Context, int AttributeIndex, bool bRebase)
+{
+	using namespace NDISimCacheReaderLocal;
+	VMInstanceGetFunction<FNiagaraID>(Context, FNiagaraID(), AttributeIndex, [](FNiagaraSimCacheAttributeReaderHelper& AttributeReader, int32 InstanceIndex) { return AttributeReader.ReadID(InstanceIndex); });
+}
+
 #undef LOCTEXT_NAMESPACE
diff --git a/Engine/Plugins/FX/Niagara/Source/Niagara/Private/DataInterface/NiagaraDataInterfaceSimCacheReader.h b/Engine/Plugins/FX/Niagara/Source/Niagara/Private/DataInterface/NiagaraDataInterfaceSimCacheReader.h
--- a/Engine/Plugins/FX/Niagara/Source/Niagara/Private/DataInterface/NiagaraDataInterfaceSimCacheReader.h
+++ b/Engine/Plugins/FX/Niagara/Source/Niagara/Private/DataInterface/NiagaraDataInterfaceSimCacheReader.h
@@ -83,4 +83,5 @@ private:
 	void VMGetInstanceColor(FVectorVMExternalFunctionContext& Context, int AttributeIndex);
 	void VMGetInstanceQuat(FVectorVMExternalFunctionContext& Context, int AttributeIndex, bool bRebase);
 	void VMGetInstancePosition(FVectorVMExternalFunctionContext& Context, int AttributeIndex, bool bRebase);
+	void VMGetInstanceID(FVectorVMExternalFunctionContext& Context, int AttributeIndex, bool bRebase);
 };
