- Fix for incorrect indexing attributes in the sim cache reader DI
- Using the ForEachFunction does not always match against the generated function order

#tests Run ET for cache reader, added additional none attribute functions before the other attribute functions, etc.
#rnx

diff --git a/Engine/Plugins/FX/Niagara/Source/Niagara/Private/DataInterface/NiagaraDataInterfaceSimCacheReader.cpp b/Engine/Plugins/FX/Niagara/Source/Niagara/Private/DataInterface/NiagaraDataInterfaceSimCacheReader.cpp
--- a/Engine/Plugins/FX/Niagara/Source/Niagara/Private/DataInterface/NiagaraDataInterfaceSimCacheReader.cpp
+++ b/Engine/Plugins/FX/Niagara/Source/Niagara/Private/DataInterface/NiagaraDataInterfaceSimCacheReader.cpp
@@ -37,6 +37,14 @@ namespace NDISimCacheReaderLocal
 
 	static const FName NAME_Attribute("Attribute");
 
+	struct FShaderStorage : public FNiagaraDataInterfaceParametersCS
+	{
+		DECLARE_TYPE_LAYOUT(FShaderStorage, NonVirtual);
+
+		LAYOUT_FIELD(TMemoryImageArray<FMemoryImageName>, AttributeReadNames);
+	};
+	IMPLEMENT_TYPE_LAYOUT(FShaderStorage);
+
 	struct FInstanceData_GameThread
 	{
 		FNiagaraParameterDirectBinding<UObject*>	SimCacheBinding;
@@ -48,7 +56,7 @@ namespace NDISimCacheReaderLocal
 	{
 		int32							EmitterIndex = INDEX_NONE;
 		FNiagaraSimCacheGpuResourcePtr	SimCacheResource;
-		TArray<int32>					ComponentOffsets;
+		TMap<uintptr_t, TArray<int32>>	StorageToComponentOffsets;
 	};
 
 	struct FNDIProxy : public FNiagaraDataInterfaceProxy
@@ -106,13 +114,13 @@ namespace NDISimCacheReaderLocal
 		return FNiagaraTypeDefinition::GetFloatDef();
 	}
 
-	static int AttributeIndexFromVariable(UNiagaraSimCache* SimCache, FName FunctionName, FName AttributeName, int EmitterIndex)
+	static int AttributeIndexFromVariable(const UNiagaraSimCache* SimCache, FName FunctionName, FName AttributeName, int EmitterIndex)
 	{
 		const FNiagaraVariableBase Variable(FunctionNameToTypeDef(FunctionName), AttributeName);
 		return FNiagaraSimCacheAttributeReaderHelper::FindVariableIndex(SimCache, Variable, EmitterIndex);
 	}
 
-	static int ComponentIndexFromVariable(UNiagaraSimCache* SimCache, FName FunctionName, FName AttributeName, int EmitterIndex)
+	static int ComponentIndexFromVariable(const UNiagaraSimCache* SimCache, FName FunctionName, FName AttributeName, int EmitterIndex)
 	{
 		const FNiagaraVariableBase Variable(FunctionNameToTypeDef(FunctionName), AttributeName);
 		if ( const FNiagaraSimCacheVariable* CacheVariable = FNiagaraSimCacheAttributeReaderHelper::FindVariable(SimCache, Variable, EmitterIndex) )
@@ -279,41 +287,15 @@ bool UNiagaraDataInterfaceSimCacheReader::PerInstanceTick(void* PerInstanceData,
 	{
 		InstanceData_GT->EmitterIndex = InstanceSimCache->GetEmitterIndex(EmitterBinding);
 
-		TArray<int32> ComponentOffsets;
-		FNiagaraDataInterfaceUtilities::ForEachGpuFunction(
-			this, SystemInstance,
-			[&](const UNiagaraScript* Scirpt, const FNiagaraDataInterfaceGeneratedFunction& FunctionInfo)
-			{
-				if (FunctionInfo.Specifiers.Num() != 1)
-				{
-					return true;
-				}
-
-				check(FunctionInfo.Specifiers[0].Key == NAME_Attribute);
-
-				const int32 ComponentIndex = ComponentIndexFromVariable(InstanceSimCache, FunctionInfo.DefinitionName, FunctionInfo.Specifiers[0].Value, InstanceData_GT->EmitterIndex);
-				ComponentOffsets.Add(ComponentIndex);
-
-				return true;
-			}
-		);
-
-		const int32 AlignedAttributes = ComponentOffsets.Num() == 0 ? 4 : Align(ComponentOffsets.Num(), 4);
-		for ( int32 i= ComponentOffsets.Num(); i < AlignedAttributes; ++i )
-		{
-			ComponentOffsets.Add(INDEX_NONE);
-		}
-
 		FNiagaraSimCacheGpuResourcePtr SimCacheResource = FNiagaraSimCacheGpuResource::CreateResource(InstanceSimCache);
 
 		ENQUEUE_RENDER_COMMAND(NDISimCacheReader_InitRT)
 		(
-			[Proxy_RT=GetProxyAs<FNDIProxy>(), InstanceID_RT=SystemInstance->GetId(), EmitterIndex_RT=InstanceData_GT->EmitterIndex, SimCacheResource_RT=MoveTemp(SimCacheResource), ComponentOffsets_RT=MoveTemp(ComponentOffsets)](FRHICommandList& CmdList)
+			[Proxy_RT=GetProxyAs<FNDIProxy>(), InstanceID_RT=SystemInstance->GetId(), EmitterIndex_RT=InstanceData_GT->EmitterIndex, SimCacheResource_RT=MoveTemp(SimCacheResource)](FRHICommandList& CmdList)
 			{
 				FInstanceData_RenderThread* InstanceData_RT = &Proxy_RT->PerInstanceData_RenderThread.FindOrAdd(InstanceID_RT);
 				InstanceData_RT->SimCacheResource = SimCacheResource_RT;
 				InstanceData_RT->EmitterIndex = EmitterIndex_RT;
-				InstanceData_RT->ComponentOffsets = ComponentOffsets_RT;
 			}
 		);
 	}
@@ -468,9 +450,12 @@ void UNiagaraDataInterfaceSimCacheReader::GetVMExternalFunction(const FVMExterna
 #if WITH_EDITORONLY_DATA
 bool UNiagaraDataInterfaceSimCacheReader::AppendCompileHash(FNiagaraCompileHashVisitor* InVisitor) const
 {
+	using namespace NDISimCacheReaderLocal;
+
 	bool bSuccess = Super::AppendCompileHash(InVisitor);
 	InVisitor->UpdateShaderFile(NDISimCacheReaderLocal::TemplateShaderFile);
 	InVisitor->UpdateShaderParameters<FShaderParameters>();
+	InVisitor->UpdateString(TEXT("FShaderStorage"), TEXT("V0"));
 	return bSuccess;
 }
 
@@ -510,6 +495,34 @@ bool UNiagaraDataInterfaceSimCacheReader::GetFunctionHLSL(const FNiagaraDataInte
 }
 #endif
 
+FNiagaraDataInterfaceParametersCS* UNiagaraDataInterfaceSimCacheReader::CreateShaderStorage(const FNiagaraDataInterfaceGPUParamInfo& ParameterInfo, const FShaderParameterMap& ParameterMap) const
+{
+	using namespace NDISimCacheReaderLocal;
+
+	FShaderStorage* ShaderStorage = new FShaderStorage();
+	ShaderStorage->AttributeReadNames.Reserve(ParameterInfo.GeneratedFunctions.Num() * 2);
+	for (const FNiagaraDataInterfaceGeneratedFunction& GeneratedFunction : ParameterInfo.GeneratedFunctions)
+	{
+		if (GeneratedFunction.Specifiers.Num() == 1)
+		{
+			ShaderStorage->AttributeReadNames.Add(GeneratedFunction.DefinitionName);
+			ShaderStorage->AttributeReadNames.Add(GeneratedFunction.Specifiers[0].Value);
+		}
+		else
+		{
+			ShaderStorage->AttributeReadNames.Add(NAME_None);
+			ShaderStorage->AttributeReadNames.Add(NAME_None);
+		}
+	}
+	return ShaderStorage;
+}
+
+const FTypeLayoutDesc* UNiagaraDataInterfaceSimCacheReader::GetShaderStorageType() const
+{
+	using namespace NDISimCacheReaderLocal;
+	return &StaticGetTypeLayoutDesc<FShaderStorage>();
+}
+
 void UNiagaraDataInterfaceSimCacheReader::BuildShaderParameters(FNiagaraShaderParametersBuilder& ShaderParametersBuilder) const
 {
 	ShaderParametersBuilder.AddNestedStruct<FShaderParameters>();
@@ -522,8 +535,8 @@ void UNiagaraDataInterfaceSimCacheReader::SetShaderParameters(const FNiagaraData
 {
 	using namespace NDISimCacheReaderLocal;
 
-	const FNDIProxy& DIProxy = Context.GetProxy<FNDIProxy>();
-	const FInstanceData_RenderThread* InstanceData_RT = DIProxy.PerInstanceData_RenderThread.Find(Context.GetSystemInstanceID());
+	FNDIProxy& DIProxy = Context.GetProxy<FNDIProxy>();
+	FInstanceData_RenderThread* InstanceData_RT = DIProxy.PerInstanceData_RenderThread.Find(Context.GetSystemInstanceID());
 
 	FShaderParameters* ShaderParameters = Context.GetParameterNestedStruct<FShaderParameters>();
 	if (InstanceData_RT && InstanceData_RT->SimCacheResource->GetNumFrames() > 0)
@@ -533,10 +546,26 @@ void UNiagaraDataInterfaceSimCacheReader::SetShaderParameters(const FNiagaraData
 		ShaderParameters->EmitterIndex	= InstanceData_RT->EmitterIndex;
 		ShaderParameters->CacheData		= InstanceData_RT->SimCacheResource->GetBufferSRV();
 
-		check((InstanceData_RT->ComponentOffsets.Num() % 4) == 0 && InstanceData_RT->ComponentOffsets.Num() > 0);
-		const int32 ComponentOffsetsSize = InstanceData_RT->ComponentOffsets.Num() / 4;
-		TArrayView<FIntVector4> ComponentOffsets = Context.GetParameterLooseArray<FIntVector4>(ComponentOffsetsSize);
-		FMemory::Memcpy(ComponentOffsets.GetData(), InstanceData_RT->ComponentOffsets.GetData(), InstanceData_RT->ComponentOffsets.Num() * InstanceData_RT->ComponentOffsets.GetTypeSize());
+		const FShaderStorage& ShaderStorage = Context.GetShaderStorage<FShaderStorage>();
+		const uint32 NumAttributes4 = FMath::DivideAndRoundUp(FMath::Max(ShaderStorage.AttributeReadNames.Num() / 2, 1), 4);
+
+		TArray<int32>& ComponentOffsets = InstanceData_RT->StorageToComponentOffsets.FindOrAdd(uintptr_t(&ShaderStorage));
+		if (ComponentOffsets.IsEmpty())
+		{
+			const UNiagaraSimCache* InstanceSimCache = InstanceData_RT->SimCacheResource->GetSimCache();
+
+			ComponentOffsets.AddZeroed(NumAttributes4 * 4);
+			for (int32 i = 0; i < ShaderStorage.AttributeReadNames.Num() / 2; ++i)
+			{
+				const FName FunctionName = ShaderStorage.AttributeReadNames[(i * 2) + 0];
+				const FName AttributeName = ShaderStorage.AttributeReadNames[(i * 2) + 1];
+				const int32 ComponentIndex = FunctionName.IsNone() ? INDEX_NONE : ComponentIndexFromVariable(InstanceSimCache, FunctionName, AttributeName, InstanceData_RT->EmitterIndex);
+				ComponentOffsets[i] = ComponentIndex;
+			}
+		}
+
+		TArrayView<FIntVector4> ShaderComponentOffsets = Context.GetParameterLooseArray<FIntVector4>(NumAttributes4);
+		FMemory::Memcpy(ShaderComponentOffsets.GetData(), ComponentOffsets.GetData(), ComponentOffsets.Num() * ComponentOffsets.GetTypeSize());
 	}
 	else
 	{
diff --git a/Engine/Plugins/FX/Niagara/Source/Niagara/Private/DataInterface/NiagaraDataInterfaceSimCacheReader.h b/Engine/Plugins/FX/Niagara/Source/Niagara/Private/DataInterface/NiagaraDataInterfaceSimCacheReader.h
--- a/Engine/Plugins/FX/Niagara/Source/Niagara/Private/DataInterface/NiagaraDataInterfaceSimCacheReader.h
+++ b/Engine/Plugins/FX/Niagara/Source/Niagara/Private/DataInterface/NiagaraDataInterfaceSimCacheReader.h
@@ -61,6 +61,8 @@ public:
 	virtual void GetParameterDefinitionHLSL(const FNiagaraDataInterfaceGPUParamInfo& ParamInfo, FString& OutHLSL) override;
 	virtual bool GetFunctionHLSL(const FNiagaraDataInterfaceGPUParamInfo& ParamInfo, const FNiagaraDataInterfaceGeneratedFunction& FunctionInfo, int FunctionInstanceIndex, FString& OutHLSL) override;
 #endif
+	virtual FNiagaraDataInterfaceParametersCS* CreateShaderStorage(const FNiagaraDataInterfaceGPUParamInfo& ParameterInfo, const FShaderParameterMap& ParameterMap) const override;
+	virtual const FTypeLayoutDesc* GetShaderStorageType() const override;
 	virtual void BuildShaderParameters(FNiagaraShaderParametersBuilder& ShaderParametersBuilder) const override;
 	virtual void SetShaderParameters(const FNiagaraDataInterfaceSetShaderParametersContext& Context) const override;
 	//UNiagaraDataInterface Interface
diff --git a/Engine/Plugins/FX/Niagara/Source/Niagara/Private/NiagaraSimCacheAttributeReaderHelper.h b/Engine/Plugins/FX/Niagara/Source/Niagara/Private/NiagaraSimCacheAttributeReaderHelper.h
--- a/Engine/Plugins/FX/Niagara/Source/Niagara/Private/NiagaraSimCacheAttributeReaderHelper.h
+++ b/Engine/Plugins/FX/Niagara/Source/Niagara/Private/NiagaraSimCacheAttributeReaderHelper.h
@@ -110,7 +110,7 @@ struct FNiagaraSimCacheAttributeReaderHelper
 		Variable = &DataBuffersLayout->Variables[AttributeIndex];
 	}
 
-	static int32 FindVariableIndex(UNiagaraSimCache* SimCache, FNiagaraVariableBase Variable, int EmitterIndex)
+	static int32 FindVariableIndex(const UNiagaraSimCache* SimCache, const FNiagaraVariableBase& Variable, int EmitterIndex)
 	{
 		if (SimCache->CacheLayout.EmitterLayouts.IsValidIndex(EmitterIndex))
 		{
@@ -122,7 +122,7 @@ struct FNiagaraSimCacheAttributeReaderHelper
 		return INDEX_NONE;
 	}
 
-	static const FNiagaraSimCacheVariable* FindVariable(UNiagaraSimCache* SimCache, FNiagaraVariableBase Variable, int EmitterIndex)
+	static const FNiagaraSimCacheVariable* FindVariable(const UNiagaraSimCache* SimCache, const FNiagaraVariableBase& Variable, int EmitterIndex)
 	{
 		const int32 VariableIndex = FindVariableIndex(SimCache, Variable, EmitterIndex);
 		if (VariableIndex != INDEX_NONE)
diff --git a/Engine/Plugins/FX/Niagara/Source/Niagara/Private/NiagaraSimCacheGpuResource.h b/Engine/Plugins/FX/Niagara/Source/Niagara/Private/NiagaraSimCacheGpuResource.h
--- a/Engine/Plugins/FX/Niagara/Source/Niagara/Private/NiagaraSimCacheGpuResource.h
+++ b/Engine/Plugins/FX/Niagara/Source/Niagara/Private/NiagaraSimCacheGpuResource.h
@@ -21,6 +21,8 @@ private:
 public:
 	static TSharedPtr<FNiagaraSimCacheGpuResource> CreateResource(UNiagaraSimCache* SimCache);
 
+	const UNiagaraSimCache* GetSimCache() const { return WeakSimCache.Get(); }
+
 	int32 GetNumFrames() const { return NumFrames; }
 	int32 GetNumEmitters() const { return NumEmitters; }
 	int32 GetEmitterIndex() const { return EmitterIndex; }
