- Add shape rotation & coordinate space selection to stateless shape location
- Rework stateless debug draw to make it cleaner when dealing with different spaces

diff --git a/Engine/Plugins/FX/Niagara/Shaders/Private/Stateless/Modules/NiagaraStatelessModule_ShapeLocation.ush b/Engine/Plugins/FX/Niagara/Shaders/Private/Stateless/Modules/NiagaraStatelessModule_ShapeLocation.ush
--- a/Engine/Plugins/FX/Niagara/Shaders/Private/Stateless/Modules/NiagaraStatelessModule_ShapeLocation.ush
+++ b/Engine/Plugins/FX/Niagara/Shaders/Private/Stateless/Modules/NiagaraStatelessModule_ShapeLocation.ush
@@ -1,10 +1,12 @@
 // Copyright Epic Games, Inc. All Rights Reserved.
 
 #include "../NiagaraStatelessCommon.ush"
+#include "/Plugin/FX/Niagara/Private/NiagaraQuaternionUtils.ush"
 
 uint4	ShapeLocation_Mode;
 float4	ShapeLocation_Parameters0;
 float4	ShapeLocation_Parameters1;
+float4	ShapeLocation_Parameters2;
 
 float3 ShapeLocation_GetLocation()
 {
@@ -13,6 +15,7 @@ float3 ShapeLocation_GetLocation()
 	if ( ShapeLocation_Mode.x == 0 )
 	{
 		const float3 BoxSize = ShapeLocation_Parameters0.xyz;
+		const float4 Rotation = ShapeLocation_Parameters2;
 		const bool bOnSurface = ShapeLocation_Mode.y == 1;
 
 		const float3 P0 = RandomFloat3(0) - 0.5f;
@@ -32,39 +35,48 @@ float3 ShapeLocation_GetLocation()
 			Location += LocalBoxSize * FaceDirB[iAxis] * P0.y;
 			Location += LocalBoxSize * FaceDirC[iAxis] * 0.5f;
 				
-			return Location;
+			return RotateVectorByQuat(Location, Rotation);
 		}
-		return P0 * BoxSize;
+		return RotateVectorByQuat(P0 * BoxSize, Rotation);
 	}
 
 	// ENSM_ShapePrimitive::Cylinder:
 	BRANCH
 	if ( ShapeLocation_Mode.x == 1 )
 	{
+		const float HeightScale	= ShapeLocation_Parameters0.w;
+		const float HeightBias	= ShapeLocation_Parameters1.w;
+		const float Radius		= ShapeLocation_Parameters2.w;
+
+		const float3 AxisX		= ShapeLocation_Parameters0.xyz;
+		const float3 AxisY		= ShapeLocation_Parameters1.xyz;
+		const float3 AxisZ		= ShapeLocation_Parameters2.xyz;
+
 		const float4 Random = RandomFloat4(0);
-		const float HeightScale = ShapeLocation_Parameters0.x;
-		const float HeightBias = ShapeLocation_Parameters0.y;
-		const float Radius = ShapeLocation_Parameters0.z;
+		const float2 XYLocation = SafeNormalize(Random.xy - 0.5f) * Radius * Random.z;
 
-		return float3(
-			SafeNormalize(Random.xy - 0.5f) * Radius * Random.z,
-			Random.w * HeightScale + HeightBias
-		);
+		float3 Location;
+		Location  = AxisX * XYLocation.x;
+		Location += AxisY * XYLocation.y;
+		Location += AxisZ * (Random.w * HeightScale + HeightBias);
+		return Location;
 	}
 
 	// ENSM_ShapePrimitive::Ring:
 	BRANCH
 	if ( ShapeLocation_Mode.x == 2 )
 	{
-		const float RadiusScale = ShapeLocation_Parameters0.x;
-		const float RadiusBias = ShapeLocation_Parameters0.y;
-		const float UDistributionScale = ShapeLocation_Parameters0.z;
-		const float UDistributionBias = ShapeLocation_Parameters0.w;
+		const float RadiusScale			= ShapeLocation_Parameters0.x;
+		const float RadiusBias			= ShapeLocation_Parameters0.y;
+		const float UDistributionScale	= ShapeLocation_Parameters0.z;
+		const float UDistributionBias	= ShapeLocation_Parameters0.w;
+		const float3 AxisX				= ShapeLocation_Parameters1.xyz;
+		const float3 AxisY				= ShapeLocation_Parameters2.xyz;
 
 		const float Radius = RandomScaleBiasFloat(0, RadiusScale, RadiusBias);
 		const float U = RandomScaleBiasFloat(1, UDistributionScale, UDistributionBias);
 
-		return float3(cos(U) * Radius, sin(U) * Radius, 0.0f);
+		return (AxisX * cos(U) * Radius) + (AxisY * sin(U) * Radius);
 	}
 
 	// ENSM_ShapePrimitive::Sphere:
diff --git a/Engine/Plugins/FX/Niagara/Source/Niagara/Internal/Stateless/Modules/NiagaraStatelessModule_ShapeLocation.h b/Engine/Plugins/FX/Niagara/Source/Niagara/Internal/Stateless/Modules/NiagaraStatelessModule_ShapeLocation.h
--- a/Engine/Plugins/FX/Niagara/Source/Niagara/Internal/Stateless/Modules/NiagaraStatelessModule_ShapeLocation.h
+++ b/Engine/Plugins/FX/Niagara/Source/Niagara/Internal/Stateless/Modules/NiagaraStatelessModule_ShapeLocation.h
@@ -73,6 +73,12 @@ public:
 	UPROPERTY(EditAnywhere, Category = "Parameters", meta=(EditConditionHides, EditCondition="ShapePrimitive == ENSM_ShapePrimitive::Sphere"))
 	FNiagaraDistributionRangeFloat SphereRadius = FNiagaraDistributionRangeFloat(100.0f);
 
+	UPROPERTY(EditAnywhere, Category = "Parameters", meta = (SegmentedDisplay))
+	ENiagaraCoordinateSpace CoordinateSpace = ENiagaraCoordinateSpace::Simulation;
+
+	UPROPERTY(EditAnywhere, Category = "Parameters")
+	FRotator3f ShapeRotation = FRotator3f::ZeroRotator;
+
 private:
 #if WITH_EDITORONLY_DATA
 	UPROPERTY()
diff --git a/Engine/Plugins/FX/Niagara/Source/Niagara/Internal/Stateless/NiagaraStatelessDrawDebugContext.h b/Engine/Plugins/FX/Niagara/Source/Niagara/Internal/Stateless/NiagaraStatelessDrawDebugContext.h
--- a/Engine/Plugins/FX/Niagara/Source/Niagara/Internal/Stateless/NiagaraStatelessDrawDebugContext.h
+++ b/Engine/Plugins/FX/Niagara/Source/Niagara/Internal/Stateless/NiagaraStatelessDrawDebugContext.h
@@ -2,20 +2,38 @@
 
 #pragma once
 
-#include "NiagaraCommon.h"
+#include "NiagaraStatelessCommon.h"
 
 #if WITH_EDITOR
 struct FNiagaraStatelessDrawDebugContext
 {
-	UWorld*		World = nullptr;
-	FTransform	LocalToWorldTransform = FTransform::Identity;
-	FTransform	WorldTransform = FTransform::Identity;
+public:
+	FNiagaraStatelessDrawDebugContext(UWorld* InWorld, const FTransform& LocalToWorldIn, bool bIsLocalSpace);
 
+	// Debug draw functions that all operate in world space
 	NIAGARA_API void DrawArrow(const FVector& Origin, const FVector& DirectionWithLength, const FColor& Color = FColor::White) const;
-	NIAGARA_API void DrawBox(const FVector& Center, const FVector& Extent, const FColor& Color = FColor::White, const FQuat& Rotation = FQuat::Identity) const;
+	NIAGARA_API void DrawBox(const FVector& Center, const FQuat& Rotation, const FVector& Extent, const FColor& Color = FColor::White) const;
 	NIAGARA_API void DrawCone(const FVector& Origin, const FQuat& Rotation, float Angle, float Length, const FColor& Color = FColor::White) const;
-	NIAGARA_API void DrawCylinder(float CylinderHeight, float CylinderRadius, float CylinderHeightMidpoint, const FColor& Color = FColor::White) const;
-	NIAGARA_API void DrawCircle(const FVector& Center, const float Radius, const FColor& Color = FColor::White) const;
+	NIAGARA_API void DrawCylinder(const FVector& Center, const FQuat& Rotation, float CylinderHeight, float CylinderRadius, float CylinderHeightMidpoint, const FColor& Color = FColor::White) const;
+	NIAGARA_API void DrawCircle(const FVector& Center, const FQuat& Rotation, const float Radius, const FColor& Color = FColor::White) const;
+	NIAGARA_API void DrawRing(const FVector& Center, const FQuat& Rotation, const float InnerRadius, const float OuterRadius, float Distribution, const FColor& Color = FColor::White) const;
 	NIAGARA_API void DrawSphere(const FVector& Center, const float Radius, const FColor& Color = FColor::White) const;
+
+	// Helper functions to transform into World space for debug rendering
+	NIAGARA_API FQuat   TransformRotation(ENiagaraCoordinateSpace SourceSpace, const FQuat4f& Rotation) const;
+	NIAGARA_API FVector TransformPosition(ENiagaraCoordinateSpace SourceSpace, const FVector3f& Position) const;
+	NIAGARA_API FVector TransformVector(ENiagaraCoordinateSpace SourceSpace, const FVector3f& Vector) const;
+	NIAGARA_API FVector TransformVectorNoScale(ENiagaraCoordinateSpace SourceSpace, const FVector3f& Vector) const;
+
+	FQuat   TransformRotation(const FQuat4f& Rotation) const { return TransformRotation(ENiagaraCoordinateSpace::Local, Rotation); }
+	FVector TransformPosition(const FVector3f& Position) const { return TransformPosition(ENiagaraCoordinateSpace::Local, Position); }
+	FVector TransformVector(const FVector3f& Vector) const { return TransformVector(ENiagaraCoordinateSpace::Local, Vector); }
+	FVector TransformVectorNoScale(const FVector3f& Vector) const { return TransformVectorNoScale(ENiagaraCoordinateSpace::Local, Vector); }
+
+private:
+	UWorld*		World = nullptr;
+	FTransform	LocalToWorld = FTransform::Identity;
+	FVector		LWCTileOffset = FVector::ZeroVector;
+	bool		bApplyLocalToWorld[3] = {};
 };
 #endif
diff --git a/Engine/Plugins/FX/Niagara/Source/Niagara/Private/Stateless/Modules/NiagaraStatelessModule_AddVelocity.cpp b/Engine/Plugins/FX/Niagara/Source/Niagara/Private/Stateless/Modules/NiagaraStatelessModule_AddVelocity.cpp
--- a/Engine/Plugins/FX/Niagara/Source/Niagara/Private/Stateless/Modules/NiagaraStatelessModule_AddVelocity.cpp
+++ b/Engine/Plugins/FX/Niagara/Source/Niagara/Private/Stateless/Modules/NiagaraStatelessModule_AddVelocity.cpp
@@ -44,6 +44,8 @@ void UNiagaraStatelessModule_AddVelocity::BuildEmitterData(const FNiagaraStatele
 #if WITH_EDITOR
 void UNiagaraStatelessModule_AddVelocity::DrawDebug(const FNiagaraStatelessDrawDebugContext& DrawDebugContext) const
 {
+	const FVector WorldOrigin = DrawDebugContext.TransformPosition(FVector3f::ZeroVector);
+
 	switch (VelocityType)
 	{
 		case ENSM_VelocityType::Linear:
@@ -56,14 +58,14 @@ void UNiagaraStatelessModule_AddVelocity::DrawDebug(const FNiagaraStatelessDrawD
 
 			if (!FMath::IsNearlyEqual(MinDir.Length(), MaxDir.Length()))
 			{
-				DrawDebugContext.DrawArrow(FVector::ZeroVector, MaxDir);
+				DrawDebugContext.DrawArrow(WorldOrigin, MaxDir);
 			}
 			break;
 		}
 
 		case ENSM_VelocityType::InCone:
 		{
-			const FQuat Quat = ConeRotation.Quaternion();
+			const FQuat ConeQuat = DrawDebugContext.TransformRotation(FQuat4f(ConeRotation.Quaternion()));
 			const float ConeHAngle = ConeAngle / 2.0f;
 
 			TOptional<float> InnerConeHAngle;
@@ -73,18 +75,18 @@ void UNiagaraStatelessModule_AddVelocity::DrawDebug(const FNiagaraStatelessDrawD
 			}
 
 			const FNiagaraStatelessRangeFloat ConeVelocityRange = ConeVelocityDistribution.CalculateRange(0.0f);
-			DrawDebugContext.DrawCone(FVector::ZeroVector, Quat, ConeHAngle, ConeVelocityRange.Min);
+			DrawDebugContext.DrawCone(WorldOrigin, ConeQuat, ConeHAngle, ConeVelocityRange.Min);
 			if ( InnerConeHAngle.IsSet() )
 			{
-				DrawDebugContext.DrawCone(FVector::ZeroVector, Quat, InnerConeHAngle.GetValue(), ConeVelocityRange.Min);
+				DrawDebugContext.DrawCone(WorldOrigin, ConeQuat, InnerConeHAngle.GetValue(), ConeVelocityRange.Min);
 			}
 
 			if (!FMath::IsNearlyEqual(ConeVelocityRange.Min, ConeVelocityRange.Max))
 			{
-				DrawDebugContext.DrawCone(FVector::ZeroVector, Quat, ConeHAngle, ConeVelocityRange.Max);
+				DrawDebugContext.DrawCone(WorldOrigin, ConeQuat, ConeHAngle, ConeVelocityRange.Max);
 				if (InnerConeHAngle.IsSet())
 				{
-					DrawDebugContext.DrawCone(FVector::ZeroVector, Quat, InnerConeHAngle.GetValue(), ConeVelocityRange.Max);
+					DrawDebugContext.DrawCone(WorldOrigin, ConeQuat, InnerConeHAngle.GetValue(), ConeVelocityRange.Max);
 				}
 			}
 			break;
@@ -95,11 +97,11 @@ void UNiagaraStatelessModule_AddVelocity::DrawDebug(const FNiagaraStatelessDrawD
 			const FNiagaraStatelessRangeFloat PointVelocityRange = PointVelocityDistribution.CalculateRange(0.0f);
 			if (!FMath::IsNearlyEqual(PointVelocityRange.Min, 0.0f))
 			{
-				DrawDebugContext.DrawSphere(FVector(PointOrigin), PointVelocityRange.Min);
+				DrawDebugContext.DrawSphere(DrawDebugContext.TransformPosition(PointOrigin), PointVelocityRange.Min);
 			}
 			if (!FMath::IsNearlyEqual(PointVelocityRange.Min, PointVelocityRange.Max))
 			{
-				DrawDebugContext.DrawSphere(FVector(PointOrigin), PointVelocityRange.Max);
+				DrawDebugContext.DrawSphere(DrawDebugContext.TransformPosition(PointOrigin), PointVelocityRange.Max);
 			}
 			break;
 		}
diff --git a/Engine/Plugins/FX/Niagara/Source/Niagara/Private/Stateless/Modules/NiagaraStatelessModule_ShapeLocation.cpp b/Engine/Plugins/FX/Niagara/Source/Niagara/Private/Stateless/Modules/NiagaraStatelessModule_ShapeLocation.cpp
--- a/Engine/Plugins/FX/Niagara/Source/Niagara/Private/Stateless/Modules/NiagaraStatelessModule_ShapeLocation.cpp
+++ b/Engine/Plugins/FX/Niagara/Source/Niagara/Private/Stateless/Modules/NiagaraStatelessModule_ShapeLocation.cpp
@@ -10,9 +10,11 @@ namespace NSMShapeLocationPrivate
 {
 	struct FModuleBuiltData
 	{
-		ENSM_ShapePrimitive	ShapePrimitive = ENSM_ShapePrimitive::Max;
-		int32				PositionVariableOffset = INDEX_NONE;
-		int32				PreviousPositionVariableOffset = INDEX_NONE;
+		ENSM_ShapePrimitive		ShapePrimitive = ENSM_ShapePrimitive::Max;
+		ENiagaraCoordinateSpace	CoordinateSpace = ENiagaraCoordinateSpace::Simulation;
+		FQuat4f					ShapeRotation = FQuat4f::Identity;
+		int32					PositionVariableOffset = INDEX_NONE;
+		int32					PreviousPositionVariableOffset = INDEX_NONE;
 	};
 
 	struct FModuleBuiltDataBox
@@ -69,6 +71,7 @@ namespace NSMShapeLocationPrivate
 		if ( ShaderParameters->ShapeLocation_Mode.X == 0 )
 		{
 			const FVector3f BoxSize = FVector3f(ShaderParameters->ShapeLocation_Parameters0);
+			const FQuat4f Rotation = FQuat4f(ShaderParameters->ShapeLocation_Parameters2.X, ShaderParameters->ShapeLocation_Parameters2.Y, ShaderParameters->ShapeLocation_Parameters2.Z, ShaderParameters->ShapeLocation_Parameters2.W);
 			const bool bOnSurface = ShaderParameters->ShapeLocation_Mode.Y == 1;
 
 			const FVector3f P0 = ParticleSimulationContext.RandomFloat3(iInstance, 0) - 0.5f;
@@ -88,42 +91,51 @@ namespace NSMShapeLocationPrivate
 				Location += LocalBoxSize * FaceDirB[iAxis] * P0.Y;
 				Location += LocalBoxSize * FaceDirC[iAxis] * 0.5f;
 
-				return Location;
+				return Rotation.RotateVector(Location);
 			}
 			else
 			{
-				return P0 * BoxSize;
+				return Rotation.RotateVector(P0 * BoxSize);
 			}
 		}
 	
 		// ENSM_ShapePrimitive::Cylinder:
 		if ( ShaderParameters->ShapeLocation_Mode.X == 1 )
 		{
+			const float HeightScale	= ShaderParameters->ShapeLocation_Parameters0.W;
+			const float HeightBias	= ShaderParameters->ShapeLocation_Parameters1.W;
+			const float Radius		= ShaderParameters->ShapeLocation_Parameters2.W;
+
+			const FVector3f AxisX	= FVector3f(ShaderParameters->ShapeLocation_Parameters0.X, ShaderParameters->ShapeLocation_Parameters0.Y, ShaderParameters->ShapeLocation_Parameters0.Z);
+			const FVector3f AxisY	= FVector3f(ShaderParameters->ShapeLocation_Parameters1.X, ShaderParameters->ShapeLocation_Parameters1.Y, ShaderParameters->ShapeLocation_Parameters1.Z);
+			const FVector3f AxisZ	= FVector3f(ShaderParameters->ShapeLocation_Parameters2.X, ShaderParameters->ShapeLocation_Parameters2.Y, ShaderParameters->ShapeLocation_Parameters2.Z);
+
 			const FVector4f Random = ParticleSimulationContext.RandomFloat4(iInstance, 0);
-			const float HeightScale = ShaderParameters->ShapeLocation_Parameters0.X;
-			const float HeightBias = ShaderParameters->ShapeLocation_Parameters0.Y;
-			const float Radius = ShaderParameters->ShapeLocation_Parameters0.Z;
-	
-			const FVector2f UnitVec = ParticleSimulationContext.SafeNormalize(FVector2f(Random.X - 0.5f, Random.Y - 0.5f));
-			return FVector3f(
-				UnitVec.X * Radius * Random.Z,
-				UnitVec.Y * Radius * Random.Z,
-				Random.W * HeightScale + HeightBias
-			);
+			const FVector2f XYLocation = ParticleSimulationContext.SafeNormalize(FVector2f(Random.X - 0.5f, Random.Y - 0.5f)) * Radius * Random.Z;
+
+			FVector3f Location;
+			Location  = AxisX * XYLocation.X;
+			Location += AxisY * XYLocation.Y;
+			Location += AxisZ * (Random.W * HeightScale + HeightBias);
+			return Location;
 		}
 	
 		// ENSM_ShapePrimitive::Ring:
 		if ( ShaderParameters->ShapeLocation_Mode.X == 2 )
 		{
-			const float RadiusScale = ShaderParameters->ShapeLocation_Parameters0.X;
-			const float RadiusBias = ShaderParameters->ShapeLocation_Parameters0.Y;
-			const float UDistributionScale = ShaderParameters->ShapeLocation_Parameters0.Z;
-			const float UDistributionBias = ShaderParameters->ShapeLocation_Parameters0.W;
-	
-			const float Radius = ParticleSimulationContext.RandomScaleBiasFloat(iInstance, 0, RadiusScale, RadiusBias);
-			const float U = ParticleSimulationContext.RandomScaleBiasFloat(iInstance, 1, UDistributionScale, UDistributionBias);
-	
-			return FVector3f(cosf(U) * Radius, sinf(U) * Radius, 0.0f);
+			const float RadiusScale			= ShaderParameters->ShapeLocation_Parameters0.X;
+			const float RadiusBias			= ShaderParameters->ShapeLocation_Parameters0.Y;
+			const float UDistributionScale	= ShaderParameters->ShapeLocation_Parameters0.Z;
+			const float UDistributionBias	= ShaderParameters->ShapeLocation_Parameters0.W;
+			const FVector3f AxisX			= FVector3f(ShaderParameters->ShapeLocation_Parameters1.X, ShaderParameters->ShapeLocation_Parameters1.Y, ShaderParameters->ShapeLocation_Parameters1.Z);
+			const FVector3f AxisY			= FVector3f(ShaderParameters->ShapeLocation_Parameters2.X, ShaderParameters->ShapeLocation_Parameters2.Y, ShaderParameters->ShapeLocation_Parameters2.Z);
+			const float Radius				= ParticleSimulationContext.RandomScaleBiasFloat(iInstance, 0, RadiusScale, RadiusBias);
+			const float U					= ParticleSimulationContext.RandomScaleBiasFloat(iInstance, 1, UDistributionScale, UDistributionBias);
+
+			FVector3f Location;
+			Location  = AxisX * FMath::Cos(U) * Radius;
+			Location += AxisY * FMath::Sin(U) * Radius;
+			return Location;
 		}
 	
 		// ENSM_ShapePrimitive::Sphere:
@@ -173,7 +185,9 @@ void UNiagaraStatelessModule_ShapeLocation::BuildEmitterData(const FNiagaraState
 		return;
 	}
 
-	BuiltData->ShapePrimitive = ShapePrimitive;
+	BuiltData->ShapePrimitive	= ShapePrimitive;
+	BuiltData->CoordinateSpace	= CoordinateSpace;
+	BuiltData->ShapeRotation	= ShapeRotation.Quaternion();
 
 	switch (ShapePrimitive)
 	{
@@ -237,9 +251,12 @@ void UNiagaraStatelessModule_ShapeLocation::SetShaderParameters(const FNiagaraSt
 	FParameters* Parameters = SetShaderParameterContext.GetParameterNestedStruct<FParameters>();
 	const FModuleBuiltData* ModuleBuiltData = SetShaderParameterContext.ReadBuiltData<FModuleBuiltData>();
 
+	const FQuat4f Rotation = SetShaderParameterContext.GetSpaceTransforms().TransformRotation(ModuleBuiltData->CoordinateSpace, ModuleBuiltData->ShapeRotation);
+
 	Parameters->ShapeLocation_Mode			= FUintVector4::ZeroValue;
 	Parameters->ShapeLocation_Parameters0	= FVector4f::Zero();
 	Parameters->ShapeLocation_Parameters1	= FVector4f::Zero();
+	Parameters->ShapeLocation_Parameters2	= FVector4f::Zero();
 
 	switch (ModuleBuiltData->ShapePrimitive)
 	{
@@ -255,6 +272,7 @@ void UNiagaraStatelessModule_ShapeLocation::SetShaderParameters(const FNiagaraSt
 			Parameters->ShapeLocation_Mode.Z		= 6;
 			Parameters->ShapeLocation_Parameters0	= FVector4f(Size + ThicknessParameters.Key, 0.0f);
 			Parameters->ShapeLocation_Parameters1	= FVector4f(ThicknessParameters.Value, ThicknessParameters.Value, ThicknessParameters.Value, 0.0f);
+			Parameters->ShapeLocation_Parameters2	= FVector4f(Rotation.X, Rotation.Y, Rotation.Z, Rotation.W);
 
 			break;
 		}
@@ -266,9 +284,9 @@ void UNiagaraStatelessModule_ShapeLocation::SetShaderParameters(const FNiagaraSt
 			const float MidPoint	= SetShaderParameterContext.ConvertRangeToValue(BuiltDataCylinder->MidPoint);
 
 			Parameters->ShapeLocation_Mode.X		= 1;
-			Parameters->ShapeLocation_Parameters0.X	= Height;
-			Parameters->ShapeLocation_Parameters0.Y	= Height * -MidPoint;
-			Parameters->ShapeLocation_Parameters0.Z	= Radius;
+			Parameters->ShapeLocation_Parameters0	= FVector4f(Rotation.GetAxisX(), Height);
+			Parameters->ShapeLocation_Parameters1	= FVector4f(Rotation.GetAxisY(), Height * -MidPoint);
+			Parameters->ShapeLocation_Parameters2	= FVector4f(Rotation.GetAxisZ(), Radius);
 			break;
 		}
 		case ENSM_ShapePrimitive::Plane:
@@ -300,6 +318,8 @@ void UNiagaraStatelessModule_ShapeLocation::SetShaderParameters(const FNiagaraSt
 			Parameters->ShapeLocation_Parameters0.Y = Radius * SDC;
 			Parameters->ShapeLocation_Parameters0.Z = -UE_TWO_PI * (1.0f - Distribution);
 			Parameters->ShapeLocation_Parameters0.W = 0.0f;
+			Parameters->ShapeLocation_Parameters1 = FVector4f(Rotation.GetAxisX());
+			Parameters->ShapeLocation_Parameters2 = FVector4f(Rotation.GetAxisY());
 			break;
 		}
 		case ENSM_ShapePrimitive::Sphere:
@@ -322,13 +342,15 @@ void UNiagaraStatelessModule_ShapeLocation::SetShaderParameters(const FNiagaraSt
 #if WITH_EDITOR
 void UNiagaraStatelessModule_ShapeLocation::DrawDebug(const FNiagaraStatelessDrawDebugContext& DrawDebugContext) const
 {
+	const FVector DrawPosition = DrawDebugContext.TransformPosition(FVector3f::ZeroVector);
+	const FQuat DrawRotation = DrawDebugContext.TransformRotation(CoordinateSpace, ShapeRotation.Quaternion());
 	switch (ShapePrimitive)
 	{
 		case ENSM_ShapePrimitive::Box:
 		{
 			if (BoxSize.IsBinding() == false)
 			{
-				DrawDebugContext.DrawBox(FVector::ZeroVector, FVector(BoxSize.Min * 0.5f));
+				DrawDebugContext.DrawBox(DrawPosition, DrawRotation, FVector(BoxSize.Min * 0.5f));
 			}
 			break;
 		}
@@ -336,7 +358,7 @@ void UNiagaraStatelessModule_ShapeLocation::DrawDebug(const FNiagaraStatelessDra
 		{
 			if (PlaneSize.IsBinding() == false)
 			{
-				DrawDebugContext.DrawBox(FVector::ZeroVector, FVector(PlaneSize.Min.X * 0.5f, PlaneSize.Min.Y * 0.5f, 0.0f));
+				DrawDebugContext.DrawBox(DrawPosition, DrawRotation, FVector(PlaneSize.Min.X * 0.5f, PlaneSize.Min.Y * 0.5f, 0.0f));
 			}
 			break;
 		}
@@ -344,7 +366,7 @@ void UNiagaraStatelessModule_ShapeLocation::DrawDebug(const FNiagaraStatelessDra
 		{
 			if (CylinderHeight.IsBinding() == false && CylinderRadius.IsBinding() == false && CylinderHeightMidpoint.IsBinding() == false)
 			{
-				DrawDebugContext.DrawCylinder(CylinderHeight.Min, CylinderRadius.Min, CylinderHeightMidpoint.Min);
+				DrawDebugContext.DrawCylinder(DrawPosition, DrawRotation, CylinderHeight.Min, CylinderRadius.Min, CylinderHeightMidpoint.Min);
 			}
 			break;
 		}
@@ -354,8 +376,7 @@ void UNiagaraStatelessModule_ShapeLocation::DrawDebug(const FNiagaraStatelessDra
 			{
 				const float DC = FMath::Clamp(1.0f - DiscCoverage.Min, 0.0f, 1.0f);
 				const float SDC = DC > 0.0f ? FMath::Sqrt(DC) : 0.0f;
-				DrawDebugContext.DrawCircle(FVector::ZeroVector, RingRadius.Min);
-				DrawDebugContext.DrawCircle(FVector::ZeroVector, RingRadius.Min * SDC);
+				DrawDebugContext.DrawRing(DrawPosition, DrawRotation, RingRadius.Min * SDC, RingRadius.Min, RingUDistribution.Min);
 			}
 			break;
 		}
@@ -363,10 +384,10 @@ void UNiagaraStatelessModule_ShapeLocation::DrawDebug(const FNiagaraStatelessDra
 		{
 			if (SphereRadius.IsBinding() == false)
 			{
-				DrawDebugContext.DrawSphere(FVector::ZeroVector, SphereRadius.Min);
+				DrawDebugContext.DrawSphere(DrawPosition, SphereRadius.Min);
 				if (SphereRadius.IsRange())
 				{
-					DrawDebugContext.DrawSphere(FVector::ZeroVector, SphereRadius.Max);
+					DrawDebugContext.DrawSphere(DrawPosition, SphereRadius.Max);
 				}
 			}
 			break;
diff --git a/Engine/Plugins/FX/Niagara/Source/Niagara/Private/Stateless/NiagaraStatelessDrawDebugContext.cpp b/Engine/Plugins/FX/Niagara/Source/Niagara/Private/Stateless/NiagaraStatelessDrawDebugContext.cpp
--- a/Engine/Plugins/FX/Niagara/Source/Niagara/Private/Stateless/NiagaraStatelessDrawDebugContext.cpp
+++ b/Engine/Plugins/FX/Niagara/Source/Niagara/Private/Stateless/NiagaraStatelessDrawDebugContext.cpp
@@ -5,8 +5,20 @@
 
 #include "Components/LineBatchComponent.h"
 #include "Engine/World.h"
+#include "Misc/LargeWorldRenderPosition.h"
 
 #if WITH_EDITOR
+FNiagaraStatelessDrawDebugContext::FNiagaraStatelessDrawDebugContext(UWorld* InWorld, const FTransform& InLocalToWorld, bool bIsLocalSpace)
+	: World(InWorld)
+	, LocalToWorld(InLocalToWorld)
+{
+	bApplyLocalToWorld[int(ENiagaraCoordinateSpace::Simulation)]	= bIsLocalSpace;
+	bApplyLocalToWorld[int(ENiagaraCoordinateSpace::World)]			= false;
+	bApplyLocalToWorld[int(ENiagaraCoordinateSpace::Local)]			= true;
+
+	LWCTileOffset = FVector(FLargeWorldRenderScalar::GetTileFor(LocalToWorld.GetTranslation())) * UE_LWC_RENDER_TILE_SIZE;
+}
+
 void FNiagaraStatelessDrawDebugContext::DrawArrow(const FVector& Origin, const FVector& DirectionWithLength, const FColor& Color) const
 {
 	const float Len = static_cast<float>(DirectionWithLength.Length());
@@ -22,13 +34,13 @@ void FNiagaraStatelessDrawDebugContext::DrawArrow(const FVector& Origin, const F
 	}
 }
 
-void FNiagaraStatelessDrawDebugContext::DrawBox(const FVector& Center, const FVector& Extent, const FColor& Color, const FQuat& Rotation) const
+void FNiagaraStatelessDrawDebugContext::DrawBox(const FVector& Center, const FQuat& Rotation, const FVector& Extent, const FColor& Color) const
 {
 	DrawDebugBox(
 		World,
-		LocalToWorldTransform.TransformPosition(Center),
+		Center,
 		Extent,
-		LocalToWorldTransform.TransformRotation(Rotation),
+		Rotation,
 		Color
 	);
 }
@@ -47,10 +59,9 @@ void FNiagaraStatelessDrawDebugContext::DrawCone(const FVector& Origin, const FQ
 	);
 }
 
-void FNiagaraStatelessDrawDebugContext::DrawCylinder(float CylinderHeight, float CylinderRadius, float CylinderHeightMidpoint, const FColor& Color) const
+void FNiagaraStatelessDrawDebugContext::DrawCylinder(const FVector& Center, const FQuat& Rotation, float CylinderHeight, float CylinderRadius, float CylinderHeightMidpoint, const FColor& Color) const
 {
-	const FVector Axis = LocalToWorldTransform.TransformVector(FVector::ZAxisVector);
-	const FVector Center = LocalToWorldTransform.TransformPosition(FVector::ZeroVector);
+	const FVector Axis = Rotation.RotateVector(FVector::ZAxisVector);
 	const float Offset = -CylinderHeightMidpoint * CylinderHeight;
 
 	DrawDebugCylinder(
@@ -63,16 +74,15 @@ void FNiagaraStatelessDrawDebugContext::DrawCylinder(float CylinderHeight, float
 	);
 }
 
-void FNiagaraStatelessDrawDebugContext::DrawCircle(const FVector& InCenter, const float Radius, const FColor& Color) const
+void FNiagaraStatelessDrawDebugContext::DrawCircle(const FVector& Center, const FQuat& Rotation, const float Radius, const FColor& Color) const
 {
 	if (ULineBatchComponent* LineBatcher = World->GetLineBatcher(UWorld::ELineBatcherType::World))
 	{
 		int32 Segments = 16;
 		const float AngleStep = 2.f * UE_PI / float(Segments);
 
-		const FVector Center = LocalToWorldTransform.TransformVector(InCenter);
-		const FVector AxisX = LocalToWorldTransform.TransformVector(FVector::XAxisVector);
-		const FVector AxisY = LocalToWorldTransform.TransformVector(FVector::YAxisVector);
+		const FVector AxisX = Rotation.RotateVector(FVector::XAxisVector);
+		const FVector AxisY = Rotation.RotateVector(FVector::YAxisVector);
 
 		TArray<FBatchedLine> Lines;
 		Lines.Empty(Segments);
@@ -89,14 +99,73 @@ void FNiagaraStatelessDrawDebugContext::DrawCircle(const FVector& InCenter, cons
 	}
 }
 
+void FNiagaraStatelessDrawDebugContext::DrawRing(const FVector& Center, const FQuat& Rotation, const float InnerRadius, const float OuterRadius, float Distribution, const FColor& Color) const
+{
+	if (ULineBatchComponent* LineBatcher = World->GetLineBatcher(UWorld::ELineBatcherType::World))
+	{
+		const int32 NumSegments = 16;
+		const float AngleStep = -2.f * UE_PI * FMath::Clamp(1.0f - Distribution, 0.0f, 1.0f) / float(NumSegments);
+
+		const FVector AxisX = Rotation.RotateVector(FVector::XAxisVector);
+		const FVector AxisY = Rotation.RotateVector(FVector::YAxisVector);
+
+		TArray<FBatchedLine> Lines;
+		Lines.Empty(2 + (NumSegments * 2));
+
+		const bool bDrawCaps = FMath::IsNearlyEqual(Distribution, 0.0f) == false;
+		for (int32 i=0; i < NumSegments; ++i)
+		{
+			const FVector Curr = AxisX * FMath::Cos(AngleStep * float(i + 0)) + AxisY * FMath::Sin(AngleStep * float(i + 0));
+			const FVector Next = AxisX * FMath::Cos(AngleStep * float(i + 1)) + AxisY * FMath::Sin(AngleStep * float(i + 1));
+			const FVector CurrInner = Center + InnerRadius * Curr;
+			const FVector NextInner = Center + InnerRadius * Next;
+			const FVector CurrOuter = Center + OuterRadius * Curr;
+			const FVector NextOuter = Center + OuterRadius * Next;
+
+			if (bDrawCaps && i == 0)
+			{
+				Lines.Emplace(CurrInner, CurrOuter, Color, -1.0f, 0.0f, 0);
+			}
+			if (bDrawCaps && i == NumSegments - 1)
+			{
+				Lines.Emplace(NextInner, NextOuter, Color, -1.0f, 0.0f, 0);
+			}
+
+			Lines.Emplace(CurrInner, NextInner, Color, -1.0f, 0.0f, 0);
+			Lines.Emplace(CurrOuter, NextOuter, Color, -1.0f, 0.0f, 0);
+		}
+		LineBatcher->DrawLines(Lines);
+	}
+}
+
 void FNiagaraStatelessDrawDebugContext::DrawSphere(const FVector& Center, const float Radius, const FColor& Color) const
 {
 	DrawDebugSphere(
 		World,
-		LocalToWorldTransform.TransformPosition(Center),
+		Center,
 		Radius,
 		16,
 		Color
 	);
 }
+
+FQuat FNiagaraStatelessDrawDebugContext::TransformRotation(ENiagaraCoordinateSpace SourceSpace, const FQuat4f& Rotation) const
+{
+	return bApplyLocalToWorld[int(SourceSpace)] ? LocalToWorld.TransformRotation(FQuat(Rotation)) : FQuat(Rotation);
+}
+
+FVector FNiagaraStatelessDrawDebugContext::TransformPosition(ENiagaraCoordinateSpace SourceSpace, const FVector3f& Position) const
+{
+	return bApplyLocalToWorld[int(SourceSpace)] ? LocalToWorld.TransformPosition(FVector(Position)) : FVector(Position) + LWCTileOffset;
+}
+
+FVector FNiagaraStatelessDrawDebugContext::TransformVector(ENiagaraCoordinateSpace SourceSpace, const FVector3f& Vector) const
+{
+	return bApplyLocalToWorld[int(SourceSpace)] ? LocalToWorld.TransformVector(FVector(Vector)) : FVector(Vector);
+}
+
+FVector FNiagaraStatelessDrawDebugContext::TransformVectorNoScale(ENiagaraCoordinateSpace SourceSpace, const FVector3f& Vector) const
+{
+	return bApplyLocalToWorld[int(SourceSpace)] ? LocalToWorld.TransformVectorNoScale(FVector(Vector)) : FVector(Vector);
+}
 #endif
diff --git a/Engine/Plugins/FX/Niagara/Source/Niagara/Private/Stateless/NiagaraStatelessEmitter.cpp b/Engine/Plugins/FX/Niagara/Source/Niagara/Private/Stateless/NiagaraStatelessEmitter.cpp
--- a/Engine/Plugins/FX/Niagara/Source/Niagara/Private/Stateless/NiagaraStatelessEmitter.cpp
+++ b/Engine/Plugins/FX/Niagara/Source/Niagara/Private/Stateless/NiagaraStatelessEmitter.cpp
@@ -857,10 +857,8 @@ UNiagaraStatelessEmitter* UNiagaraStatelessEmitter::CreateAsDuplicate(FName InDu
 
 void UNiagaraStatelessEmitter::DrawModuleDebug(UWorld* World, const FTransform& LocalToWorld) const
 {
-	FNiagaraStatelessDrawDebugContext DrawDebugContext;
-	DrawDebugContext.World					= World;
-	DrawDebugContext.LocalToWorldTransform	= LocalToWorld;
-	DrawDebugContext.WorldTransform			= FTransform::Identity;
+	FNiagaraStatelessDrawDebugContext DrawDebugContext(World, LocalToWorld, true/*bLocalSpace*/);
+
 	for (const UNiagaraStatelessModule* Module : Modules)
 	{
 		if (Module->IsDebugDrawEnabled())
diff --git a/Engine/Plugins/FX/Niagara/Source/NiagaraShader/Internal/Stateless/NiagaraStatelessModuleShaderParameters.h b/Engine/Plugins/FX/Niagara/Source/NiagaraShader/Internal/Stateless/NiagaraStatelessModuleShaderParameters.h
--- a/Engine/Plugins/FX/Niagara/Source/NiagaraShader/Internal/Stateless/NiagaraStatelessModuleShaderParameters.h
+++ b/Engine/Plugins/FX/Niagara/Source/NiagaraShader/Internal/Stateless/NiagaraStatelessModuleShaderParameters.h
@@ -143,6 +143,7 @@ namespace NiagaraStateless
 		SHADER_PARAMETER(FUintVector4,	ShapeLocation_Mode)
 		SHADER_PARAMETER(FVector4f,		ShapeLocation_Parameters0)
 		SHADER_PARAMETER(FVector4f,		ShapeLocation_Parameters1)
+		SHADER_PARAMETER(FVector4f,		ShapeLocation_Parameters2)
 	END_SHADER_PARAMETER_STRUCT()
 
 	BEGIN_SHADER_PARAMETER_STRUCT(FSolveVelocitiesAndForcesModule_ShaderParameters, NIAGARASHADER_API)
