Fixed object categories appearing in 'show only modified' when they don't contain any valid properties

diff --git a/Engine/Plugins/FX/Niagara/Source/NiagaraEditor/Private/ViewModels/Stack/NiagaraStackObject.cpp b/Engine/Plugins/FX/Niagara/Source/NiagaraEditor/Private/ViewModels/Stack/NiagaraStackObject.cpp
--- a/Engine/Plugins/FX/Niagara/Source/NiagaraEditor/Private/ViewModels/Stack/NiagaraStackObject.cpp
+++ b/Engine/Plugins/FX/Niagara/Source/NiagaraEditor/Private/ViewModels/Stack/NiagaraStackObject.cpp
@@ -437,6 +437,15 @@ bool UNiagaraStackObject::FilterOnlyModified(const UNiagaraStackEntry& Child) co
 		return true;
 	}
 
+	// We filter out categories that themselves have no property rows to show.
+	// We don't need to account for sub-categories (that are of type Item) as we only have top level categories here
+	if(PropertyRow->GetDetailTreeNode()->GetNodeType() == EDetailNodeType::Category || PropertyRow->GetDetailTreeNode()->GetNodeType() == EDetailNodeType::Advanced)
+	{
+		TArray<UNiagaraStackPropertyRow*> FilteredChildrenOfChild;
+		PropertyRow->GetFilteredChildrenOfType(FilteredChildrenOfChild);
+		return FilteredChildrenOfChild.Num() > 0;
+	}
+	
 	TSharedPtr<IPropertyHandle> PropertyHandle = PropertyRow->GetDetailTreeNode()->CreatePropertyHandle();
 	if(PropertyHandle == nullptr)
 	{
diff --git a/Engine/Plugins/FX/Niagara/Source/NiagaraEditor/Private/ViewModels/Stack/NiagaraStackPropertyRow.cpp b/Engine/Plugins/FX/Niagara/Source/NiagaraEditor/Private/ViewModels/Stack/NiagaraStackPropertyRow.cpp
--- a/Engine/Plugins/FX/Niagara/Source/NiagaraEditor/Private/ViewModels/Stack/NiagaraStackPropertyRow.cpp
+++ b/Engine/Plugins/FX/Niagara/Source/NiagaraEditor/Private/ViewModels/Stack/NiagaraStackPropertyRow.cpp
@@ -87,7 +87,7 @@ bool UNiagaraStackPropertyRow::GetShouldShowInStack() const
 		return bIsHiddenCategory == false;
 	}
 
-	if(DetailTreeNode->GetNodeType() == EDetailNodeType::Category)
+	if(DetailTreeNode->GetNodeType() == EDetailNodeType::Category || DetailTreeNode->GetNodeType() == EDetailNodeType::Advanced)
 	{
 		if(bIsHiddenCategory)
 		{
@@ -406,6 +406,17 @@ bool UNiagaraStackPropertyRow::FilterOnlyModified(const UNiagaraStackEntry& Chil
 	}
 
 	TSharedPtr<IPropertyHandle> PropertyHandle = PropertyRow->GetDetailTreeNode()->CreatePropertyHandle();
+	EDetailNodeType DetailNodeType = PropertyRow->GetDetailTreeNode()->GetNodeType();
+	
+	// We filter out categories that themselves have no property rows to show
+	// If this is an item but the property handle is nullptr, this is a sub-category (thanks Property Editor..)
+	if((DetailNodeType == EDetailNodeType::Item && PropertyHandle == nullptr) || DetailNodeType == EDetailNodeType::Category || DetailNodeType == EDetailNodeType::Advanced)
+	{
+		TArray<UNiagaraStackPropertyRow*> FilteredChildrenOfChild;
+		PropertyRow->GetFilteredChildrenOfType(FilteredChildrenOfChild);
+		return FilteredChildrenOfChild.Num() > 0;
+	}
+	
 	if(PropertyHandle == nullptr)
 	{
 		return true;
