- Add property to customize quick capture number of frames

diff --git a/Engine/Plugins/FX/Niagara/Source/Niagara/Public/NiagaraComponent.h b/Engine/Plugins/FX/Niagara/Source/Niagara/Public/NiagaraComponent.h
--- a/Engine/Plugins/FX/Niagara/Source/Niagara/Public/NiagaraComponent.h
+++ b/Engine/Plugins/FX/Niagara/Source/Niagara/Public/NiagaraComponent.h
@@ -755,6 +755,9 @@ public:
 #if WITH_EDITORONLY_DATA
 	UPROPERTY(EditAnywhere, Category = Compilation)
 	uint32 bWaitForCompilationOnActivate : 1;
+
+	UPROPERTY(EditAnywhere, Category = "Niagara Utilities", meta = (UIMin="1", ClampMin="1", DisplayName="NumFramesToCapture"))
+	TOptional<int32> SimCacheDebugNumFramesToCapture;
 #endif
 
 	/**
diff --git a/Engine/Plugins/FX/Niagara/Source/NiagaraEditor/Private/Customizations/NiagaraComponentDetails.cpp b/Engine/Plugins/FX/Niagara/Source/NiagaraEditor/Private/Customizations/NiagaraComponentDetails.cpp
--- a/Engine/Plugins/FX/Niagara/Source/NiagaraEditor/Private/Customizations/NiagaraComponentDetails.cpp
+++ b/Engine/Plugins/FX/Niagara/Source/NiagaraEditor/Private/Customizations/NiagaraComponentDetails.cpp
@@ -125,7 +125,7 @@ void FNiagaraComponentDetails::CustomizeDetails(IDetailLayoutBuilder& DetailBuil
 	Builder = &DetailBuilder;
 
 	static const FName ParamCategoryName = TEXT("NiagaraComponent_Parameters");
-	static const FName ParamUtilitiesName = TEXT("NiagaraComponent_Utilities");
+	static const FName ParamUtilitiesName = TEXT("Niagara Utilities");
 	static const FName ScriptCategoryName = TEXT("Parameters");
 
 	static bool bFirstTime = true;
@@ -196,7 +196,7 @@ void FNiagaraComponentDetails::CustomizeDetails(IDetailLayoutBuilder& DetailBuil
 			];
 	}
 	
-	IDetailCategoryBuilder& CustomCategory = DetailBuilder.EditCategory(ParamUtilitiesName, LOCTEXT("ParamUtilsCategoryName", "Niagara Utilities"), ECategoryPriority::Important);
+	IDetailCategoryBuilder& CustomCategory = DetailBuilder.EditCategory(ParamUtilitiesName, FText::GetEmpty(), ECategoryPriority::Important);
 
 	CustomCategory.AddCustomRow(FText::GetEmpty())
 		.WholeRowContent()
@@ -361,12 +361,10 @@ FReply FNiagaraComponentDetails::OnCaptureSelectedSystem()
 		SimCache->SetFlags(RF_Transient);
 		CapturedCaches.Emplace(SimCache);
 
-	
 		FNiagaraSimCacheCaptureParameters CaptureParameters;
-
 		CaptureParameters.CaptureRate = 1;
 		// A value of 0 would capture until completion. Since this can apply broadly, looping systems would not finish capturing, so don't go below 1.
-		CaptureParameters.NumFrames = FMath::Max(1, GetDefault<UNiagaraSettings>()->QuickSimCacheCaptureFrameCount);
+		CaptureParameters.NumFrames = FMath::Max(1, NiagaraComponentToUse->SimCacheDebugNumFramesToCapture.Get(GetDefault<UNiagaraSettings>()->QuickSimCacheCaptureFrameCount));
 		CaptureParameters.bCaptureAllFramesImmediatly = false;
 		
 		if(NiagaraComponentToUse->GetWorld())
