- Nanite renderer lwc / local / world fixes

diff --git a/Engine/Plugins/FX/Niagara/Shaders/Private/NiagaraTransformUtils.ush b/Engine/Plugins/FX/Niagara/Shaders/Private/NiagaraTransformUtils.ush
--- a/Engine/Plugins/FX/Niagara/Shaders/Private/NiagaraTransformUtils.ush
+++ b/Engine/Plugins/FX/Niagara/Shaders/Private/NiagaraTransformUtils.ush
@@ -37,3 +37,15 @@ FNiagaraTransform MultiplyTransform(in FNiagaraTransform A, in FNiagaraTransform
 	R.Scale			= A.Scale * B.Scale;
 	return R;
 }
+
+float4x4 TransformToMatrix44(FNiagaraTransform T)
+{
+	const float3x3 Rotation = NiagaraQuatTo3x3(T.Rotation);
+	
+	return float4x4(  
+		float4(Rotation[0] * T.Scale.x, 0.0f),
+		float4(Rotation[1] * T.Scale.y, 0.0f),
+		float4(Rotation[2] * T.Scale.z, 0.0f),
+		float4(T.Translation, 1.0f)
+	);
+}
