Niagara - Remove pre-edit change calls from the renderer properties when updating MICs so that we don't end up dirtying systems on load, also add additional safety against nulls.

diff --git a/Engine/Plugins/FX/Niagara/Source/Niagara/Private/NiagaraRendererProperties.cpp b/Engine/Plugins/FX/Niagara/Source/Niagara/Private/NiagaraRendererProperties.cpp
--- a/Engine/Plugins/FX/Niagara/Source/Niagara/Private/NiagaraRendererProperties.cpp
+++ b/Engine/Plugins/FX/Niagara/Source/Niagara/Private/NiagaraRendererProperties.cpp
@@ -645,10 +645,8 @@ void UNiagaraRendererProperties::UpdateMaterialParametersMIC(const FNiagaraRende
 					MIC = MICPool[ExistingIndex];
 					MICPool.RemoveAtSwap(ExistingIndex, EAllowShrinking::No);
 				}
-				
-				MIC->PreEditChange(nullptr);
 			}
-			else
+			if (MIC == nullptr)
 			{
 				if (UObject* ExistingObject = StaticFindObject(UMaterialInstanceConstant::StaticClass(), this, *NameBuilder))
 				{
@@ -666,7 +664,6 @@ void UNiagaraRendererProperties::UpdateMaterialParametersMIC(const FNiagaraRende
 				}
 
 				MIC = NewObject<UMaterialInstanceConstant>(this, FName(NameBuilder));
-				MIC->PreEditChange(nullptr);
 				MIC->SetParentEditorOnly(Material);
 			}
 
