Fixed scratch pad scripts not being included in allowed tabs
- Also renamed tab factory to more accurately describe what it's doing

diff --git a/Engine/Plugins/FX/Niagara/Source/NiagaraEditor/Private/ViewModels/NiagaraSystemEditorDocumentsViewModel.cpp b/Engine/Plugins/FX/Niagara/Source/NiagaraEditor/Private/ViewModels/NiagaraSystemEditorDocumentsViewModel.cpp
--- a/Engine/Plugins/FX/Niagara/Source/NiagaraEditor/Private/ViewModels/NiagaraSystemEditorDocumentsViewModel.cpp
+++ b/Engine/Plugins/FX/Niagara/Source/NiagaraEditor/Private/ViewModels/NiagaraSystemEditorDocumentsViewModel.cpp
@@ -254,9 +254,10 @@ void UNiagaraSystemEditorDocumentsViewModel::InitializePreTabManager(TSharedPtr<
 	// Register the document factories
 	{
 
-		TSharedRef<FDocumentTabFactory> GraphEditorFactory = MakeShareable(new FNiagaraGraphEditorSummoner(InToolkit,
-			FNiagaraGraphEditorSummoner::FOnCreateGraphEditorWidget::CreateUObject(this, &UNiagaraSystemEditorDocumentsViewModel::CreateGraphEditorWidget)
-		));
+		TSharedRef<FDocumentTabFactory> GraphEditorFactory = MakeShared<FDocumentTabFactory_NiagaraScratchPad>(
+			InToolkit,
+			FDocumentTabFactory_NiagaraScratchPad::FOnCreateGraphEditorWidget::CreateUObject(this, &UNiagaraSystemEditorDocumentsViewModel::CreateGraphEditorWidget)
+		);
 
 		// Also store off a reference to the grapheditor factory so we can find all the tabs spawned by it later.
 		GraphEditorTabFactoryPtr = GraphEditorFactory;
diff --git a/Engine/Plugins/FX/Niagara/Source/NiagaraEditor/Private/Widgets/NiagaraScratchScriptEditor.cpp b/Engine/Plugins/FX/Niagara/Source/NiagaraEditor/Private/Widgets/NiagaraScratchScriptEditor.cpp
--- a/Engine/Plugins/FX/Niagara/Source/NiagaraEditor/Private/Widgets/NiagaraScratchScriptEditor.cpp
+++ b/Engine/Plugins/FX/Niagara/Source/NiagaraEditor/Private/Widgets/NiagaraScratchScriptEditor.cpp
@@ -198,25 +198,25 @@ void FNiagaraGraphTabHistory::RestoreHistory()
 			Editor->GetGraphEditor()->SetViewLocation(SavedLocation, SavedZoomAmount, SavedBookmarkId);
 	}
 }
-void FNiagaraGraphEditorSummoner::OnTabActivated(TSharedPtr<SDockTab> Tab) const
+void FDocumentTabFactory_NiagaraScratchPad::OnTabActivated(TSharedPtr<SDockTab> Tab) const
 {
 	TSharedRef<SNiagaraScratchPadScriptEditor> GraphEditor = StaticCastSharedRef<SNiagaraScratchPadScriptEditor>(Tab->GetContent());
 	EditorPtr.Pin()->GetSystemViewModel()->GetDocumentViewModel()->SetActiveDocumentTab(Tab);
 }
 
-void FNiagaraGraphEditorSummoner::OnTabBackgrounded(TSharedPtr<SDockTab> Tab) const
+void FDocumentTabFactory_NiagaraScratchPad::OnTabBackgrounded(TSharedPtr<SDockTab> Tab) const
 {
 	TSharedRef<SNiagaraScratchPadScriptEditor> GraphEditor = StaticCastSharedRef<SNiagaraScratchPadScriptEditor>(Tab->GetContent());
 	//EditorPtr.Pin()->OnGraphEditorBackgrounded(GraphEditor);
 }
 
-void FNiagaraGraphEditorSummoner::OnTabRefreshed(TSharedPtr<SDockTab> Tab) const
+void FDocumentTabFactory_NiagaraScratchPad::OnTabRefreshed(TSharedPtr<SDockTab> Tab) const
 {
 	TSharedRef<SNiagaraScratchPadScriptEditor> GraphEditor = StaticCastSharedRef<SNiagaraScratchPadScriptEditor>(Tab->GetContent());
 	GraphEditor->GetGraphEditor()->NotifyGraphChanged();
 }
 
-void FNiagaraGraphEditorSummoner::SaveState(TSharedPtr<SDockTab> Tab, TSharedPtr<FTabPayload> Payload) const
+void FDocumentTabFactory_NiagaraScratchPad::SaveState(TSharedPtr<SDockTab> Tab, TSharedPtr<FTabPayload> Payload) const
 {
 	// Do nothing here for now
 	/*
@@ -230,14 +230,15 @@ void FNiagaraGraphEditorSummoner::SaveState(TSharedPtr<SDockTab> Tab, TSharedPtr
 	*/
 }
 
-FNiagaraGraphEditorSummoner::FNiagaraGraphEditorSummoner(TSharedPtr<class FNiagaraSystemToolkit> InToolkit, FOnCreateGraphEditorWidget CreateGraphEditorWidgetCallback) : FDocumentTabFactoryForObjects<UEdGraph>(TEXT("SystemToolkitGraph"), InToolkit->AsShared())
+FDocumentTabFactory_NiagaraScratchPad::FDocumentTabFactory_NiagaraScratchPad(TSharedPtr<class FNiagaraSystemToolkit> InToolkit, FOnCreateGraphEditorWidget CreateGraphEditorWidgetCallback)
+	: FDocumentTabFactoryForObjects<UEdGraph>(TEXT("NiagaraSystemEditor_ScratchPad"), InToolkit->AsShared())
 , EditorPtr(InToolkit)
 , OnCreateGraphEditorWidget(CreateGraphEditorWidgetCallback)
 {
 	
 }
 
-TSharedRef<SWidget> FNiagaraGraphEditorSummoner::CreateTabBodyForObject(const FWorkflowTabSpawnInfo& Info, UEdGraph* DocumentID) const
+TSharedRef<SWidget> FDocumentTabFactory_NiagaraScratchPad::CreateTabBodyForObject(const FWorkflowTabSpawnInfo& Info, UEdGraph* DocumentID) const
 {
 	check(Info.TabInfo.IsValid());
 
@@ -257,17 +258,17 @@ TSharedRef<SWidget> FNiagaraGraphEditorSummoner::CreateTabBodyForObject(const FW
 	return OnCreateGraphEditorWidget.Execute(Info.TabInfo.ToSharedRef(), DocumentID);
 }
 
-const FSlateBrush* FNiagaraGraphEditorSummoner::GetTabIconForObject(const FWorkflowTabSpawnInfo& Info, UEdGraph* DocumentID) const
+const FSlateBrush* FDocumentTabFactory_NiagaraScratchPad::GetTabIconForObject(const FWorkflowTabSpawnInfo& Info, UEdGraph* DocumentID) const
 {
 	return FSlateIcon(FNiagaraEditorStyle::Get().GetStyleSetName(), "Tab.ScratchPad").GetSmallIcon();
 }
 
-TSharedRef<FGenericTabHistory> FNiagaraGraphEditorSummoner::CreateTabHistoryNode(TSharedPtr<FTabPayload> Payload)
+TSharedRef<FGenericTabHistory> FDocumentTabFactory_NiagaraScratchPad::CreateTabHistoryNode(TSharedPtr<FTabPayload> Payload)
 {
 	return MakeShareable(new FNiagaraGraphTabHistory(SharedThis(this), Payload));
 }
 
-bool FNiagaraGraphEditorSummoner::IsPayloadValid(TSharedRef<FTabPayload> Payload) const
+bool FDocumentTabFactory_NiagaraScratchPad::IsPayloadValid(TSharedRef<FTabPayload> Payload) const
 {
 	bool bBaseResult = FDocumentTabFactoryForObjects<UEdGraph>::IsPayloadValid(Payload);
 
@@ -291,7 +292,7 @@ bool FNiagaraGraphEditorSummoner::IsPayloadValid(TSharedRef<FTabPayload> Payload
 	return false;
 }
 
-TAttribute<FText> FNiagaraGraphEditorSummoner::ConstructTabNameForObject(UEdGraph* DocumentID) const 
+TAttribute<FText> FDocumentTabFactory_NiagaraScratchPad::ConstructTabNameForObject(UEdGraph* DocumentID) const 
 {
 	UNiagaraScript* Script = DocumentID->GetTypedOuter<UNiagaraScript>();
 	TSharedPtr<FNiagaraSystemViewModel> SysViewModel = EditorPtr.Pin()->GetSystemViewModel();
@@ -310,7 +311,7 @@ TAttribute<FText> FNiagaraGraphEditorSummoner::ConstructTabNameForObject(UEdGrap
 	return FText();
 }
 
-TSharedRef<SDockTab> FNiagaraGraphEditorSummoner::OnSpawnTab(const FSpawnTabArgs& SpawnArgs, TWeakPtr<FTabManager> WeakTabManager) const
+TSharedRef<SDockTab> FDocumentTabFactory_NiagaraScratchPad::OnSpawnTab(const FSpawnTabArgs& SpawnArgs, TWeakPtr<FTabManager> WeakTabManager) const
 {
 	TSharedRef<SDockTab> TabRef = FDocumentTabFactoryForObjects<UEdGraph>::OnSpawnTab(SpawnArgs, WeakTabManager);
 	return TabRef;
diff --git a/Engine/Plugins/FX/Niagara/Source/NiagaraEditor/Private/Widgets/NiagaraScratchScriptEditor.h b/Engine/Plugins/FX/Niagara/Source/NiagaraEditor/Private/Widgets/NiagaraScratchScriptEditor.h
--- a/Engine/Plugins/FX/Niagara/Source/NiagaraEditor/Private/Widgets/NiagaraScratchScriptEditor.h
+++ b/Engine/Plugins/FX/Niagara/Source/NiagaraEditor/Private/Widgets/NiagaraScratchScriptEditor.h
@@ -130,12 +130,12 @@ private:
 };
 
 
-struct FNiagaraGraphEditorSummoner : public FDocumentTabFactoryForObjects<UEdGraph>
+struct FDocumentTabFactory_NiagaraScratchPad : public FDocumentTabFactoryForObjects<UEdGraph>
 {
 public:
 	DECLARE_DELEGATE_RetVal_TwoParams(TSharedRef<SNiagaraScratchPadScriptEditor>, FOnCreateGraphEditorWidget, TSharedRef<FTabInfo>, UEdGraph*);
 public:
-	FNiagaraGraphEditorSummoner(TSharedPtr<class FNiagaraSystemToolkit> InToolkit, FOnCreateGraphEditorWidget CreateGraphEditorWidgetCallback);
+	FDocumentTabFactory_NiagaraScratchPad(TSharedPtr<class FNiagaraSystemToolkit> InToolkit, FOnCreateGraphEditorWidget CreateGraphEditorWidgetCallback);
 
 	virtual void OnTabActivated(TSharedPtr<SDockTab> Tab) const override;
 
