- Update Stateless Rotate Around Point module
- Added coordinate space / rotation / center
- Added a CPU path for the module
- Added Debug Draw support

diff --git a/Engine/Plugins/FX/Niagara/Shaders/Private/Stateless/Modules/NiagaraStatelessModule_RotateAroundPoint.ush b/Engine/Plugins/FX/Niagara/Shaders/Private/Stateless/Modules/NiagaraStatelessModule_RotateAroundPoint.ush
--- a/Engine/Plugins/FX/Niagara/Shaders/Private/Stateless/Modules/NiagaraStatelessModule_RotateAroundPoint.ush
+++ b/Engine/Plugins/FX/Niagara/Shaders/Private/Stateless/Modules/NiagaraStatelessModule_RotateAroundPoint.ush
@@ -1,6 +1,7 @@
 // Copyright Epic Games, Inc. All Rights Reserved.
 
 #include "../NiagaraStatelessCommon.ush"
+#include "/Plugin/FX/Niagara/Private/NiagaraQuaternionUtils.ush"
 
 float	RotateAroundPoint_RateScale;
 float	RotateAroundPoint_RateBias;
@@ -8,22 +9,26 @@ float	RotateAroundPoint_RadiusScale;
 float	RotateAroundPoint_RadiusBias;
 float	RotateAroundPoint_InitialPhaseScale;
 float	RotateAroundPoint_InitialPhaseBias;
+float3	RotateAroundPoint_CenterScale;
+float3	RotateAroundPoint_CenterBias;
+float3	RotateAroundPoint_RotationAxisScale;
+float3	RotateAroundPoint_RotationAxisBias;
 
 struct FStatelessModule_RotateAroundPoint
 {
-	float Rate;
-	float Radius;
-	float InitialPhase;
+	float	Rate;
+	float	Radius;
+	float	InitialPhase;
+	float3	Center;
+	float4	RotationQuat;
 };
 
 void RotateAroundPoint_IntegratePosition(in FStatelessModule_RotateAroundPoint ModuleData, float Age, inout float3 Position)
 {
-	const float Tau		= 3.14f * 2.0f;
-	const float3 AxisX	= float3(1.0f, 0.0f, 0.0f);
-	const float3 AxisY	= float3(0.0f, 1.0f, 0.0f);
+	const float  Time	= ModuleData.InitialPhase + Age * ModuleData.Rate;
+	const float3 Offset = float3(cos(Time), sin(Time), 0.0f) * ModuleData.Radius;
 
-	Position += sin(ModuleData.InitialPhase + Age * ModuleData.Rate) * ModuleData.Radius * AxisX;
-	Position += cos(ModuleData.InitialPhase + Age * ModuleData.Rate) * ModuleData.Radius * AxisY;
+	Position += ModuleData.Center + RotateVectorByQuat(Offset, ModuleData.RotationQuat);
 }
 
 void RotateAroundPoint_Simulate(inout FStatelessParticle Particle)
@@ -32,6 +37,8 @@ void RotateAroundPoint_Simulate(inout FStatelessParticle Particle)
 	ModuleData.Rate			= RandomScaleBiasFloat(0, RotateAroundPoint_RateScale, RotateAroundPoint_RateBias);
 	ModuleData.Radius		= RandomScaleBiasFloat(1, RotateAroundPoint_RadiusScale, RotateAroundPoint_RadiusBias);
 	ModuleData.InitialPhase	= RandomScaleBiasFloat(2, RotateAroundPoint_InitialPhaseScale, RotateAroundPoint_InitialPhaseBias);
+	ModuleData.Center		= RandomScaleBiasFloat(3, RotateAroundPoint_CenterScale, RotateAroundPoint_CenterBias);
+	ModuleData.RotationQuat	= RotatorToQuat(RandomScaleBiasFloat(4, RotateAroundPoint_RotationAxisScale, RotateAroundPoint_RotationAxisBias));
 		
 	RotateAroundPoint_IntegratePosition(ModuleData, Particle.NormalizedAge, Particle.Position);
 	RotateAroundPoint_IntegratePosition(ModuleData, Particle.PreviousNormalizedAge, Particle.PreviousPosition);
diff --git a/Engine/Plugins/FX/Niagara/Source/Niagara/Internal/Stateless/Modules/NiagaraStatelessModule_RotateAroundPoint.h b/Engine/Plugins/FX/Niagara/Source/Niagara/Internal/Stateless/Modules/NiagaraStatelessModule_RotateAroundPoint.h
--- a/Engine/Plugins/FX/Niagara/Source/Niagara/Internal/Stateless/Modules/NiagaraStatelessModule_RotateAroundPoint.h
+++ b/Engine/Plugins/FX/Niagara/Source/Niagara/Internal/Stateless/Modules/NiagaraStatelessModule_RotateAroundPoint.h
@@ -3,8 +3,6 @@
 #pragma once
 
 #include "Stateless/NiagaraStatelessModule.h"
-#include "Stateless/NiagaraStatelessEmitterDataBuildContext.h"
-#include "Stateless/NiagaraStatelessModuleShaderParameters.h"
 
 #include "NiagaraStatelessModule_RotateAroundPoint.generated.h"
 
@@ -13,16 +11,7 @@ class UNiagaraStatelessModule_RotateAroundPoint : public UNiagaraStatelessModule
 {
 	GENERATED_BODY()
 
-	struct FModuleBuiltData
-	{
-		FNiagaraStatelessRangeFloat	Rate			= FNiagaraStatelessRangeFloat(0.0f);
-		FNiagaraStatelessRangeFloat	Radius			= FNiagaraStatelessRangeFloat(0.0f);
-		FNiagaraStatelessRangeFloat	InitialPhase	= FNiagaraStatelessRangeFloat(0.0f);
-	};
-
 public:
-	using FParameters = NiagaraStateless::FRotateAroundPointModule_ShaderParameters;
-
 	UPROPERTY(EditAnywhere, Category = "Parameters", meta=(DisableBindingDistribution, Units="deg/s"))
 	FNiagaraDistributionRangeFloat Rate = FNiagaraDistributionRangeFloat(360.0f);
 
@@ -32,49 +21,28 @@ public:
 	UPROPERTY(EditAnywhere, Category = "Parameters", meta = (DisableBindingDistribution))
 	FNiagaraDistributionRangeFloat InitialPhase = FNiagaraDistributionRangeFloat(0.0f);
 
-	//-TODO: Add support for GPU once we settle on a feature set for this module
-	virtual ENiagaraStatelessFeatureMask GetFeatureMask() const override { return ENiagaraStatelessFeatureMask::ExecuteGPU; }
-
-	virtual void BuildEmitterData(const FNiagaraStatelessEmitterDataBuildContext& BuildContext) const override
-	{
-		FModuleBuiltData* BuiltData = BuildContext.AllocateBuiltData<FModuleBuiltData>();
-		if (!IsModuleEnabled())
-		{
-			return;
-		}
+	UPROPERTY(EditAnywhere, Category = "Parameters", meta = (SegmentedDisplay))
+	ENiagaraCoordinateSpace CenterCoordinateSpace = ENiagaraCoordinateSpace::Simulation;
 
-		BuiltData->Rate			= Rate.CalculateRange();
-		BuiltData->Radius		= Radius.CalculateRange();
-		BuiltData->InitialPhase	= InitialPhase.CalculateRange();
-	}
+	UPROPERTY(EditAnywhere, Category = "Parameters", meta=(DisableUniformDistribution))
+	FNiagaraDistributionRangeVector3 Center = FNiagaraDistributionRangeVector3(FVector3f::ZeroVector);
 
-	virtual void BuildShaderParameters(FNiagaraStatelessShaderParametersBuilder& ShaderParametersBuilder) const override
-	{
-		ShaderParametersBuilder.AddParameterNestedStruct<FParameters>();
-	}
+	UPROPERTY(EditAnywhere, Category = "Parameters", meta = (SegmentedDisplay))
+	ENiagaraCoordinateSpace RotationCoordinateSpace = ENiagaraCoordinateSpace::Simulation;
 
-	virtual void SetShaderParameters(const FNiagaraStatelessSetShaderParameterContext& SetShaderParameterContext) const override
-	{
-		FParameters* Parameters = SetShaderParameterContext.GetParameterNestedStruct<FParameters>();
-		const FModuleBuiltData* BuiltData = SetShaderParameterContext.ReadBuiltData<FModuleBuiltData>();
-
-		Parameters->RotateAroundPoint_RateScale			= FMath::DegreesToRadians(BuiltData->Rate.GetScale());
-		Parameters->RotateAroundPoint_RateBias			= FMath::DegreesToRadians(BuiltData->Rate.Min);
-		Parameters->RotateAroundPoint_RadiusScale		= BuiltData->Radius.GetScale();
-		Parameters->RotateAroundPoint_RadiusBias		= BuiltData->Radius.Min;
-		Parameters->RotateAroundPoint_InitialPhaseScale	= BuiltData->InitialPhase.GetScale();
-		Parameters->RotateAroundPoint_InitialPhaseBias	= BuiltData->InitialPhase.Min;
-	}
+	UPROPERTY(EditAnywhere, Category = "Parameters", meta = (DisableUniformDistribution))
+	FNiagaraDistributionRangeVector3 RotationAxis = FNiagaraDistributionRangeVector3(FVector3f::ZeroVector);
 
+	virtual void BuildEmitterData(const FNiagaraStatelessEmitterDataBuildContext& BuildContext) const override;
+	virtual void BuildShaderParameters(FNiagaraStatelessShaderParametersBuilder& ShaderParametersBuilder) const override;
+	virtual void SetShaderParameters(const FNiagaraStatelessSetShaderParameterContext& SetShaderParameterContext) const override;
 #if WITH_EDITOR
 	virtual bool CanDisableModule() const override { return true; }
+
+	virtual bool CanDebugDraw() const override { return true; }
+	virtual void DrawDebug(const FNiagaraStatelessDrawDebugContext& DrawDebugContext) const override;
 #endif
 #if WITH_EDITORONLY_DATA
-	virtual void GetOutputVariables(TArray<FNiagaraVariableBase>& OutVariables) const override
-	{
-		const FNiagaraStatelessGlobals& StatelessGlobals = FNiagaraStatelessGlobals::Get();
-		OutVariables.AddUnique(StatelessGlobals.PositionVariable);
-		OutVariables.AddUnique(StatelessGlobals.PreviousPositionVariable);
-	}
+	virtual void GetOutputVariables(TArray<FNiagaraVariableBase>& OutVariables) const override;
 #endif
 };
diff --git a/Engine/Plugins/FX/Niagara/Source/Niagara/Internal/Stateless/NiagaraStatelessDrawDebugContext.h b/Engine/Plugins/FX/Niagara/Source/Niagara/Internal/Stateless/NiagaraStatelessDrawDebugContext.h
--- a/Engine/Plugins/FX/Niagara/Source/Niagara/Internal/Stateless/NiagaraStatelessDrawDebugContext.h
+++ b/Engine/Plugins/FX/Niagara/Source/Niagara/Internal/Stateless/NiagaraStatelessDrawDebugContext.h
@@ -11,6 +11,7 @@ public:
 	FNiagaraStatelessDrawDebugContext(UWorld* InWorld, const FTransform& LocalToWorldIn, bool bIsLocalSpace);
 
 	// Debug draw functions that all operate in world space
+	NIAGARA_API void DrawAxis(const FVector& Origin, const FQuat& Rotation, float Scale) const;
 	NIAGARA_API void DrawArrow(const FVector& Origin, const FVector& DirectionWithLength, const FColor& Color = FColor::White) const;
 	NIAGARA_API void DrawBox(const FVector& Center, const FQuat& Rotation, const FVector& Extent, const FColor& Color = FColor::White) const;
 	NIAGARA_API void DrawCone(const FVector& Origin, const FQuat& Rotation, float Angle, float Length, const FColor& Color = FColor::White) const;
diff --git a/Engine/Plugins/FX/Niagara/Source/Niagara/Internal/Stateless/NiagaraStatelessSetShaderParameterContext.h b/Engine/Plugins/FX/Niagara/Source/Niagara/Internal/Stateless/NiagaraStatelessSetShaderParameterContext.h
--- a/Engine/Plugins/FX/Niagara/Source/Niagara/Internal/Stateless/NiagaraStatelessSetShaderParameterContext.h
+++ b/Engine/Plugins/FX/Niagara/Source/Niagara/Internal/Stateless/NiagaraStatelessSetShaderParameterContext.h
@@ -63,6 +63,8 @@ public:
 	}
 
 	void TransformVectorRangeToScaleBias(const FNiagaraStatelessRangeVector3& Range, ENiagaraCoordinateSpace SourceSpace, FVector3f& OutScale, FVector3f& OutBias) const;
+	void TransformPositionRangeToScaleBias(const FNiagaraStatelessRangeVector3& Range, ENiagaraCoordinateSpace SourceSpace, FVector3f& OutScale, FVector3f& OutBias) const;
+	void TransformRotationRangeToScaleBias(const FNiagaraStatelessRangeVector3& Range, ENiagaraCoordinateSpace SourceSpace, FVector3f& OutScale, FVector3f& OutBias) const;
 	FVector3f TransformPositionRangeToValue(const FNiagaraStatelessRangeVector3& Range, ENiagaraCoordinateSpace SourceSpace) const;
 
 protected:
diff --git a/Engine/Plugins/FX/Niagara/Source/Niagara/Private/Stateless/Modules/NiagaraStatelessModule_RotateAroundPoint.cpp b/Engine/Plugins/FX/Niagara/Source/Niagara/Private/Stateless/Modules/NiagaraStatelessModule_RotateAroundPoint.cpp
new file mode 100644
--- /dev/null
+++ b/Engine/Plugins/FX/Niagara/Source/Niagara/Private/Stateless/Modules/NiagaraStatelessModule_RotateAroundPoint.cpp
@@ -0,0 +1,164 @@
+// Copyright Epic Games, Inc. All Rights Reserved.
+
+#include "Stateless/Modules/NiagaraStatelessModule_RotateAroundPoint.h"
+
+#include "Stateless/NiagaraStatelessDrawDebugContext.h"
+#include "Stateless/NiagaraStatelessEmitterDataBuildContext.h"
+#include "Stateless/NiagaraStatelessModuleShaderParameters.h"
+#include "Stateless/NiagaraStatelessParticleSimContext.h"
+
+namespace NSMRotateAroundPointPrivate
+{
+	struct FModuleBuiltData
+	{
+		FNiagaraStatelessRangeFloat		Rate = FNiagaraStatelessRangeFloat(0.0f);
+		FNiagaraStatelessRangeFloat		Radius = FNiagaraStatelessRangeFloat(0.0f);
+		FNiagaraStatelessRangeFloat		InitialPhase = FNiagaraStatelessRangeFloat(0.0f);
+		ENiagaraCoordinateSpace			CenterCoordinateSpace = ENiagaraCoordinateSpace::Simulation;
+		FNiagaraStatelessRangeVector3	Center = FNiagaraStatelessRangeVector3(FVector3f::ZeroVector);
+		ENiagaraCoordinateSpace			RotationCoordinateSpace = ENiagaraCoordinateSpace::Simulation;
+		FNiagaraStatelessRangeVector3	RotationAxis = FNiagaraStatelessRangeVector3(FVector3f::ZeroVector);
+
+		int32							PositionVariableOffset = INDEX_NONE;
+		int32							PreviousPositionVariableOffset = INDEX_NONE;
+	};
+
+	using FParameters = NiagaraStateless::FRotateAroundPointModule_ShaderParameters;
+
+	static FVector3f CalculatePosition(float Age, float RotationRate, float RotationRadius, float InitialPhase, const FQuat4f& Rotation)
+	{
+		const float Time = InitialPhase + Age * RotationRate;
+		const FVector3f Offset = FVector3f(FMath::Cos(Time), FMath::Sin(Time), 0.0f) * RotationRadius;
+		return Rotation.RotateVector(Offset);
+	}
+
+	static void ParticleSimulate(const NiagaraStateless::FParticleSimulationContext& ParticleSimulationContext)
+	{
+		using namespace NiagaraStateless;
+
+		const FModuleBuiltData* ModuleBuiltData = ParticleSimulationContext.ReadBuiltData<FModuleBuiltData>();
+		const FParameters* ShaderParameters = ParticleSimulationContext.ReadParameterNestedStruct<FParameters>();
+
+		const float* ParticleAges = ParticleSimulationContext.GetParticleNormalizedAge();
+		const float* ParticlePreviousAges = ParticleSimulationContext.GetParticlePreviousNormalizedAge();
+
+		for (uint32 i = 0; i < ParticleSimulationContext.GetNumInstances(); ++i)
+		{
+			const float RotationRate	= ParticleSimulationContext.RandomScaleBiasFloat(i, 0, ShaderParameters->RotateAroundPoint_RateScale, ShaderParameters->RotateAroundPoint_RateBias);
+			const float RotationRadius	= ParticleSimulationContext.RandomScaleBiasFloat(i, 1, ShaderParameters->RotateAroundPoint_RadiusScale, ShaderParameters->RotateAroundPoint_RadiusBias);
+			const float InitialPhase	= ParticleSimulationContext.RandomScaleBiasFloat(i, 2, ShaderParameters->RotateAroundPoint_InitialPhaseScale, ShaderParameters->RotateAroundPoint_InitialPhaseBias);
+			const FVector3f Center		= ParticleSimulationContext.RandomScaleBiasFloat(i, 3, ShaderParameters->RotateAroundPoint_CenterScale, ShaderParameters->RotateAroundPoint_CenterBias);
+			const FQuat4f Rotation		= ParticleSimulationContext.RotatorToQuat(ParticleSimulationContext.RandomScaleBiasFloat(i, 4, ShaderParameters->RotateAroundPoint_RotationAxisScale, ShaderParameters->RotateAroundPoint_RotationAxisBias));
+
+			FVector3f Position = ParticleSimulationContext.ReadParticleVariable(ModuleBuiltData->PositionVariableOffset, i, FVector3f::ZeroVector);
+			FVector3f PreviousPosition = ParticleSimulationContext.ReadParticleVariable(ModuleBuiltData->PreviousPositionVariableOffset, i, FVector3f::ZeroVector);
+
+			Position += Center + CalculatePosition(ParticleAges[i], RotationRate, RotationRadius, InitialPhase, Rotation);
+			PreviousPosition += Center + CalculatePosition(ParticlePreviousAges[i], RotationRate, RotationRadius, InitialPhase, Rotation);
+
+			ParticleSimulationContext.WriteParticleVariable(ModuleBuiltData->PositionVariableOffset, i, Position);
+			ParticleSimulationContext.WriteParticleVariable(ModuleBuiltData->PreviousPositionVariableOffset, i, PreviousPosition);
+		}
+	}
+}
+
+void UNiagaraStatelessModule_RotateAroundPoint::BuildEmitterData(const FNiagaraStatelessEmitterDataBuildContext& BuildContext) const
+{
+	using namespace NSMRotateAroundPointPrivate;
+
+	FModuleBuiltData* BuiltData = BuildContext.AllocateBuiltData<FModuleBuiltData>();
+	if (!IsModuleEnabled())
+	{
+		return;
+	}
+
+	const FNiagaraStatelessGlobals& StatelessGlobals = FNiagaraStatelessGlobals::Get();
+	BuiltData->PositionVariableOffset			= BuildContext.FindParticleVariableIndex(StatelessGlobals.PositionVariable);
+	BuiltData->PreviousPositionVariableOffset	= BuildContext.FindParticleVariableIndex(StatelessGlobals.PreviousPositionVariable);
+	if (BuiltData->PositionVariableOffset == INDEX_NONE && BuiltData->PreviousPositionVariableOffset == INDEX_NONE)
+	{
+		return;
+	}
+
+	BuiltData->Rate						= Rate.CalculateRange();
+	BuiltData->Radius					= Radius.CalculateRange();
+	BuiltData->InitialPhase				= InitialPhase.CalculateRange();
+	BuiltData->CenterCoordinateSpace	= CenterCoordinateSpace;
+	BuiltData->Center					= Center.CalculateRange();
+	BuiltData->RotationCoordinateSpace	= RotationCoordinateSpace;
+	BuiltData->RotationAxis				= RotationAxis.CalculateRange();
+
+	BuildContext.AddParticleSimulationExecSimulate(&ParticleSimulate);
+}
+
+void UNiagaraStatelessModule_RotateAroundPoint::BuildShaderParameters(FNiagaraStatelessShaderParametersBuilder& ShaderParametersBuilder) const
+{
+	using namespace NSMRotateAroundPointPrivate;
+
+	ShaderParametersBuilder.AddParameterNestedStruct<FParameters>();
+}
+
+void UNiagaraStatelessModule_RotateAroundPoint::SetShaderParameters(const FNiagaraStatelessSetShaderParameterContext& SetShaderParameterContext) const
+{
+	using namespace NSMRotateAroundPointPrivate;
+
+	FParameters* Parameters = SetShaderParameterContext.GetParameterNestedStruct<FParameters>();
+	const FModuleBuiltData* BuiltData = SetShaderParameterContext.ReadBuiltData<FModuleBuiltData>();
+
+	Parameters->RotateAroundPoint_RateScale			= FMath::DegreesToRadians(BuiltData->Rate.GetScale());
+	Parameters->RotateAroundPoint_RateBias			= FMath::DegreesToRadians(BuiltData->Rate.Min);
+	Parameters->RotateAroundPoint_RadiusScale		= BuiltData->Radius.GetScale();
+	Parameters->RotateAroundPoint_RadiusBias		= BuiltData->Radius.Min;
+	Parameters->RotateAroundPoint_InitialPhaseScale	= BuiltData->InitialPhase.GetScale();
+	Parameters->RotateAroundPoint_InitialPhaseBias	= BuiltData->InitialPhase.Min;
+
+	SetShaderParameterContext.TransformPositionRangeToScaleBias(
+		BuiltData->Center,
+		BuiltData->CenterCoordinateSpace,
+		Parameters->RotateAroundPoint_CenterScale,
+		Parameters->RotateAroundPoint_CenterBias
+	);
+
+	SetShaderParameterContext.TransformRotationRangeToScaleBias(
+		BuiltData->RotationAxis,
+		BuiltData->RotationCoordinateSpace,
+		Parameters->RotateAroundPoint_RotationAxisScale,
+		Parameters->RotateAroundPoint_RotationAxisBias
+	);
+	Parameters->RotateAroundPoint_RotationAxisScale /= 360.0f;
+	Parameters->RotateAroundPoint_RotationAxisBias /= 360.0f;
+}
+
+#if WITH_EDITOR
+void UNiagaraStatelessModule_RotateAroundPoint::DrawDebug(const FNiagaraStatelessDrawDebugContext& DrawDebugContext) const
+{
+	if (Center.IsBinding() || Radius.IsBinding() || RotationAxis.IsBinding())
+	{
+		return;
+	}
+
+	// Calculate rotation
+	const FQuat DrawRotation = DrawDebugContext.TransformRotation(RotationCoordinateSpace, FQuat4f::MakeFromEuler(RotationAxis.Min));
+
+	// Calculate center
+	const FNiagaraStatelessRangeVector3 CenterRange = Center.CalculateRange();
+	const FVector3f AverageCenter = (CenterRange.Min + CenterRange.Max) * 0.5f;
+	const FVector DrawPosition = DrawDebugContext.TransformPosition(CenterCoordinateSpace, AverageCenter);
+
+	// Calculate radius
+	const FNiagaraStatelessRangeFloat RadiusRange = Radius.CalculateRange();
+	DrawDebugContext.DrawAxis(DrawPosition, DrawRotation, RadiusRange.Min * 0.5f);
+	DrawDebugContext.DrawRing(DrawPosition, DrawRotation, RadiusRange.Min, RadiusRange.Max, 0.0f);
+}
+#endif
+
+#if WITH_EDITORONLY_DATA
+void UNiagaraStatelessModule_RotateAroundPoint::GetOutputVariables(TArray<FNiagaraVariableBase>& OutVariables) const
+{
+	using namespace NSMRotateAroundPointPrivate;
+
+	const FNiagaraStatelessGlobals& StatelessGlobals = FNiagaraStatelessGlobals::Get();
+	OutVariables.AddUnique(StatelessGlobals.PositionVariable);
+	OutVariables.AddUnique(StatelessGlobals.PreviousPositionVariable);
+}
+#endif
diff --git a/Engine/Plugins/FX/Niagara/Source/Niagara/Private/Stateless/NiagaraStatelessDrawDebugContext.cpp b/Engine/Plugins/FX/Niagara/Source/Niagara/Private/Stateless/NiagaraStatelessDrawDebugContext.cpp
--- a/Engine/Plugins/FX/Niagara/Source/Niagara/Private/Stateless/NiagaraStatelessDrawDebugContext.cpp
+++ b/Engine/Plugins/FX/Niagara/Source/Niagara/Private/Stateless/NiagaraStatelessDrawDebugContext.cpp
@@ -19,6 +19,11 @@ FNiagaraStatelessDrawDebugContext::FNiagaraStatelessDrawDebugContext(UWorld* InW
 	LWCTileOffset = FVector(FLargeWorldRenderScalar::GetTileFor(LocalToWorld.GetTranslation())) * UE_LWC_RENDER_TILE_SIZE;
 }
 
+void FNiagaraStatelessDrawDebugContext::DrawAxis(const FVector& Origin, const FQuat& Rotation, float Scale) const
+{
+	DrawDebugCoordinateSystem(World, Origin, FRotator(Rotation), Scale);
+}
+
 void FNiagaraStatelessDrawDebugContext::DrawArrow(const FVector& Origin, const FVector& DirectionWithLength, const FColor& Color) const
 {
 	const float Len = static_cast<float>(DirectionWithLength.Length());
diff --git a/Engine/Plugins/FX/Niagara/Source/Niagara/Private/Stateless/NiagaraStatelessSetShaderParameterContext.cpp b/Engine/Plugins/FX/Niagara/Source/Niagara/Private/Stateless/NiagaraStatelessSetShaderParameterContext.cpp
--- a/Engine/Plugins/FX/Niagara/Source/Niagara/Private/Stateless/NiagaraStatelessSetShaderParameterContext.cpp
+++ b/Engine/Plugins/FX/Niagara/Source/Niagara/Private/Stateless/NiagaraStatelessSetShaderParameterContext.cpp
@@ -26,6 +26,36 @@ void FNiagaraStatelessSetShaderParameterContext::TransformVectorRangeToScaleBias
 	}
 }
 
+void FNiagaraStatelessSetShaderParameterContext::TransformPositionRangeToScaleBias(const FNiagaraStatelessRangeVector3& Range, ENiagaraCoordinateSpace SourceSpace, FVector3f& OutScale, FVector3f& OutBias) const
+{
+	if (Range.ParameterOffset == INDEX_NONE)
+	{
+		OutScale = SpaceTransforms.TransformVector(SourceSpace, Range.GetScale());
+		OutBias = SpaceTransforms.TransformPosition(SourceSpace, Range.Min);
+	}
+	else
+	{
+		OutScale = FNiagaraStatelessRangeDefaultValue<FVector3f>::Zero();
+		OutBias = SpaceTransforms.TransformPosition(SourceSpace, GetRendererParameterValue(Range.ParameterOffset, Range.Min));
+	}
+}
+
+void FNiagaraStatelessSetShaderParameterContext::TransformRotationRangeToScaleBias(const FNiagaraStatelessRangeVector3& Range, ENiagaraCoordinateSpace SourceSpace, FVector3f& OutScale, FVector3f& OutBias) const
+{
+	if (Range.ParameterOffset == INDEX_NONE)
+	{
+		OutScale = Range.GetScale();
+		OutBias = Range.Min;
+	}
+	else
+	{
+		OutScale = FNiagaraStatelessRangeDefaultValue<FVector3f>::Zero();
+		OutBias = GetRendererParameterValue(Range.ParameterOffset, Range.Min);
+	}
+	const FQuat4f Rotation = SpaceTransforms.TransformRotation(SourceSpace, FQuat4f::MakeFromEuler(OutBias));
+	OutBias = Rotation.Euler();
+}
+
 FVector3f FNiagaraStatelessSetShaderParameterContext::TransformPositionRangeToValue(const FNiagaraStatelessRangeVector3& Range, ENiagaraCoordinateSpace SourceSpace) const
 {
 	return SpaceTransforms.TransformPosition(SourceSpace, GetRendererParameterValue(Range.ParameterOffset, Range.Min));
diff --git a/Engine/Plugins/FX/Niagara/Source/NiagaraShader/Internal/Stateless/NiagaraStatelessModuleShaderParameters.h b/Engine/Plugins/FX/Niagara/Source/NiagaraShader/Internal/Stateless/NiagaraStatelessModuleShaderParameters.h
--- a/Engine/Plugins/FX/Niagara/Source/NiagaraShader/Internal/Stateless/NiagaraStatelessModuleShaderParameters.h
+++ b/Engine/Plugins/FX/Niagara/Source/NiagaraShader/Internal/Stateless/NiagaraStatelessModuleShaderParameters.h
@@ -37,12 +37,16 @@ namespace NiagaraStateless
 		
 
 	BEGIN_SHADER_PARAMETER_STRUCT(FRotateAroundPointModule_ShaderParameters, NIAGARASHADER_API)
-		SHADER_PARAMETER(float,	RotateAroundPoint_RateScale)
-		SHADER_PARAMETER(float,	RotateAroundPoint_RateBias)
-		SHADER_PARAMETER(float,	RotateAroundPoint_RadiusScale)
-		SHADER_PARAMETER(float,	RotateAroundPoint_RadiusBias)
-		SHADER_PARAMETER(float,	RotateAroundPoint_InitialPhaseScale)
-		SHADER_PARAMETER(float,	RotateAroundPoint_InitialPhaseBias)
+		SHADER_PARAMETER(float,		RotateAroundPoint_RateScale)
+		SHADER_PARAMETER(float,		RotateAroundPoint_RateBias)
+		SHADER_PARAMETER(float,		RotateAroundPoint_RadiusScale)
+		SHADER_PARAMETER(float,		RotateAroundPoint_RadiusBias)
+		SHADER_PARAMETER(float,		RotateAroundPoint_InitialPhaseScale)
+		SHADER_PARAMETER(float,		RotateAroundPoint_InitialPhaseBias)
+		SHADER_PARAMETER(FVector3f,	RotateAroundPoint_CenterScale)
+		SHADER_PARAMETER(FVector3f,	RotateAroundPoint_CenterBias)
+		SHADER_PARAMETER(FVector3f, RotateAroundPoint_RotationAxisScale)
+		SHADER_PARAMETER(FVector3f, RotateAroundPoint_RotationAxisBias)
 	END_SHADER_PARAMETER_STRUCT()
 
 	BEGIN_SHADER_PARAMETER_STRUCT(FDynamicMaterialParametersModule_ShaderParameters, NIAGARASHADER_API)
