Fix display of Niagara content browser asset tags

#rnx

diff --git a/Engine/Plugins/FX/Niagara/Source/Niagara/Private/NiagaraSystem.cpp b/Engine/Plugins/FX/Niagara/Source/Niagara/Private/NiagaraSystem.cpp
--- a/Engine/Plugins/FX/Niagara/Source/Niagara/Private/NiagaraSystem.cpp
+++ b/Engine/Plugins/FX/Niagara/Source/Niagara/Private/NiagaraSystem.cpp
@@ -816,7 +816,7 @@ void UNiagaraSystem::UpdateSystemAfterLoad()
 		}
 	}
 
-	if (bUpdateTagsAfterLoad && HasAnyFlags(RF_ClassDefaultObject) == false)
+	if (bUpdateTagsAfterLoad)
 	{
 		if (IAssetRegistry* AssetRegistry = IAssetRegistry::Get())
 		{
@@ -1657,10 +1657,10 @@ void UNiagaraSystem::GetAssetRegistryTags(TArray<FAssetRegistryTag>& OutTags) co
 
 void UNiagaraSystem::GetAssetRegistryTags(FAssetRegistryTagsContext Context) const
 {
-	if (HasAnyFlags(RF_ClassDefaultObject) || !bFullyLoaded)
+	if (!bFullyLoaded && !HasAnyFlags(RF_ClassDefaultObject))
 	{
 		// We're currently not fully loaded and can't do so in the context of asset tags, as it could cause a deadlock with emitter merging.
-		// The system will update the asset registry manually once fully loaded.
+		// The system will update the asset registry manually once fully loaded. The CDO is an exception to this, because otherwise the content browser can't display any of the tags.
 #if WITH_EDITORONLY_DATA
 		bUpdateTagsAfterLoad = true;
 #endif
