- Set the world time information in the baker
- This fixes things like Engine.Time always being 0.0 and material based time effects not working as expected

#tests Create a NS which uses Engine.Time to modify the color, or a material that uses time.  Use the baker and noticed it never updates.

diff --git a/Engine/Plugins/FX/Niagara/Source/NiagaraEditor/Private/NiagaraBakerRenderer.cpp b/Engine/Plugins/FX/Niagara/Source/NiagaraEditor/Private/NiagaraBakerRenderer.cpp
--- a/Engine/Plugins/FX/Niagara/Source/NiagaraEditor/Private/NiagaraBakerRenderer.cpp
+++ b/Engine/Plugins/FX/Niagara/Source/NiagaraEditor/Private/NiagaraBakerRenderer.cpp
@@ -364,13 +364,22 @@ void FNiagaraBakerRenderer::SetAbsoluteTime(float AbsoluteTime, bool bShouldTick
 	PreviewComponent->SetSeekDelta(BakerSettings->GetSeekDelta());
 	PreviewComponent->SeekToDesiredAge(AbsoluteTime);
 
+	UWorld* World = PreviewComponent->GetWorld();
+	if (ensure(World))
+	{
+		World->TimeSeconds			= AbsoluteTime;
+		World->UnpausedTimeSeconds	= AbsoluteTime;
+		World->RealTimeSeconds		= AbsoluteTime;
+		World->DeltaRealTimeSeconds	= BakerSettings->GetSeekDelta();
+		World->DeltaTimeSeconds		= BakerSettings->GetSeekDelta();
+	}
+
 	if (bShouldTickComponent)
 	{
 		PreviewComponent->TickComponent(BakerSettings->GetSeekDelta(), ELevelTick::LEVELTICK_All, nullptr);
 
 		// World should be guaranteed but let's be safe
-		UWorld* World = PreviewComponent->GetWorld();
-		if ( ensure(World) )
+		if ( World )
 		{
 			// Send EOF updates before we flush our pending ticks to ensure everything is ready for Niagara
 			World->SendAllEndOfFrameUpdates();
