- Add CPU path for calculate accurate velocities

#tests Force emitter to CPU only path
#rnx

diff --git a/Engine/Plugins/FX/Niagara/Shaders/Private/Stateless/Modules/NiagaraStatelessModule_CalculateAccurateVelocity.ush b/Engine/Plugins/FX/Niagara/Shaders/Private/Stateless/Modules/NiagaraStatelessModule_CalculateAccurateVelocity.ush
--- a/Engine/Plugins/FX/Niagara/Shaders/Private/Stateless/Modules/NiagaraStatelessModule_CalculateAccurateVelocity.ush
+++ b/Engine/Plugins/FX/Niagara/Shaders/Private/Stateless/Modules/NiagaraStatelessModule_CalculateAccurateVelocity.ush
@@ -4,7 +4,7 @@
 
 void CalculateAccurateVelocity_Simulate(inout FStatelessParticle Particle)
 {  
-	Particle.Velocity			= Particle.Position - Particle.PreviousPosition;
+	Particle.Velocity			= (Particle.Position - Particle.PreviousPosition) * Particle.InvDeltaTime;
 	Particle.PreviousVelocity	= Particle.Velocity;
 }
 	
\ No newline at end of file
diff --git a/Engine/Plugins/FX/Niagara/Source/Niagara/Internal/Stateless/Modules/NiagaraStatelessModule_CalculateAccurateVelocity.h b/Engine/Plugins/FX/Niagara/Source/Niagara/Internal/Stateless/Modules/NiagaraStatelessModule_CalculateAccurateVelocity.h
--- a/Engine/Plugins/FX/Niagara/Source/Niagara/Internal/Stateless/Modules/NiagaraStatelessModule_CalculateAccurateVelocity.h
+++ b/Engine/Plugins/FX/Niagara/Source/Niagara/Internal/Stateless/Modules/NiagaraStatelessModule_CalculateAccurateVelocity.h
@@ -13,12 +13,9 @@ class UNiagaraStatelessModule_CalculateAccurateVelocity : public UNiagaraStatele
 	GENERATED_BODY()
 
 public:
+	virtual void BuildEmitterData(const FNiagaraStatelessEmitterDataBuildContext& BuildContext) const override;
+
 #if WITH_EDITORONLY_DATA
-	virtual void GetOutputVariables(TArray<FNiagaraVariableBase>& OutVariables) const override
-	{
-		const FNiagaraStatelessGlobals& StatelessGlobals = FNiagaraStatelessGlobals::Get();
-		OutVariables.AddUnique(StatelessGlobals.VelocityVariable);
-		OutVariables.AddUnique(StatelessGlobals.PreviousVelocityVariable);
-	}
+	virtual void GetOutputVariables(TArray<FNiagaraVariableBase>& OutVariables) const override;
 #endif
 };
diff --git a/Engine/Plugins/FX/Niagara/Source/Niagara/Private/Stateless/Modules/NiagaraStatelessModule_CalculateAccurateVelocity.cpp b/Engine/Plugins/FX/Niagara/Source/Niagara/Private/Stateless/Modules/NiagaraStatelessModule_CalculateAccurateVelocity.cpp
new file mode 100644
--- /dev/null
+++ b/Engine/Plugins/FX/Niagara/Source/Niagara/Private/Stateless/Modules/NiagaraStatelessModule_CalculateAccurateVelocity.cpp
@@ -0,0 +1,65 @@
+// Copyright Epic Games, Inc. All Rights Reserved.
+
+#include "Stateless/Modules/NiagaraStatelessModule_CalculateAccurateVelocity.h"
+
+#include "Stateless/NiagaraStatelessEmitterDataBuildContext.h"
+#include "Stateless/NiagaraStatelessParticleSimContext.h"
+
+namespace NSMCalculateAccurateVelocityPrivate
+{
+	struct FModuleBuiltData
+	{
+		int32 PositionVariableOffset			= INDEX_NONE;
+		int32 PreviousPositionVariableOffset	= INDEX_NONE;
+		int32 VelocityVariableOffset			= INDEX_NONE;
+		int32 PreviousVelocityVariableOffset	= INDEX_NONE;
+	};
+
+	static void ParticleSimulate(const NiagaraStateless::FParticleSimulationContext& ParticleSimulationContext)
+	{
+		using namespace NiagaraStateless;
+
+		const FModuleBuiltData* ModuleBuiltData = ParticleSimulationContext.ReadBuiltData<FModuleBuiltData>();
+
+		const float InvDeltaTime = ParticleSimulationContext.GetInvDeltaTime();
+
+		for (uint32 i = 0; i < ParticleSimulationContext.GetNumInstances(); ++i)
+		{
+			const FVector3f Position			= ParticleSimulationContext.ReadParticleVariable(ModuleBuiltData->PositionVariableOffset, i, FVector3f::ZeroVector);
+			const FVector3f PreviousPosition	= ParticleSimulationContext.ReadParticleVariable(ModuleBuiltData->PreviousPositionVariableOffset, i, FVector3f::ZeroVector);
+			const FVector3f Velocity			= (Position - PreviousPosition) * InvDeltaTime;
+
+			ParticleSimulationContext.WriteParticleVariable(ModuleBuiltData->VelocityVariableOffset, i, Velocity);
+			ParticleSimulationContext.WriteParticleVariable(ModuleBuiltData->PreviousVelocityVariableOffset, i, Velocity);
+		}
+	}
+}
+
+void UNiagaraStatelessModule_CalculateAccurateVelocity::BuildEmitterData(const FNiagaraStatelessEmitterDataBuildContext& BuildContext) const
+{
+	using namespace NSMCalculateAccurateVelocityPrivate;
+
+	FModuleBuiltData* BuiltData = BuildContext.AllocateBuiltData<FModuleBuiltData>();
+
+	const FNiagaraStatelessGlobals& StatelessGlobals = FNiagaraStatelessGlobals::Get();
+	BuiltData->PositionVariableOffset			= BuildContext.FindParticleVariableIndex(StatelessGlobals.PositionVariable);
+	BuiltData->PreviousPositionVariableOffset	= BuildContext.FindParticleVariableIndex(StatelessGlobals.PreviousPositionVariable);
+	BuiltData->VelocityVariableOffset			= BuildContext.FindParticleVariableIndex(StatelessGlobals.VelocityVariable);
+	BuiltData->PreviousVelocityVariableOffset	= BuildContext.FindParticleVariableIndex(StatelessGlobals.PreviousVelocityVariable);
+	if (BuiltData->VelocityVariableOffset == INDEX_NONE && BuiltData->PreviousVelocityVariableOffset == INDEX_NONE)
+	{
+		return;
+	}
+
+	BuildContext.AddParticleSimulationExecSimulate(&ParticleSimulate);
+}
+
+
+#if WITH_EDITORONLY_DATA
+void UNiagaraStatelessModule_CalculateAccurateVelocity::GetOutputVariables(TArray<FNiagaraVariableBase>& OutVariables) const
+{
+	const FNiagaraStatelessGlobals& StatelessGlobals = FNiagaraStatelessGlobals::Get();
+	OutVariables.AddUnique(StatelessGlobals.VelocityVariable);
+	OutVariables.AddUnique(StatelessGlobals.PreviousVelocityVariable);
+}
+#endif
