- Fix IWYU issues

#rnx

diff --git a/Engine/Plugins/FX/Niagara/Shaders/Private/NiagaraTransformUtils.ush b/Engine/Plugins/FX/Niagara/Shaders/Private/NiagaraTransformUtils.ush
--- a/Engine/Plugins/FX/Niagara/Shaders/Private/NiagaraTransformUtils.ush
+++ b/Engine/Plugins/FX/Niagara/Shaders/Private/NiagaraTransformUtils.ush
@@ -2,6 +2,7 @@
 
 #pragma once
 
+#include "/Plugin/FX/Niagara/Private/NiagaraCommon.ush"
 #include "/Plugin/FX/Niagara/Private/NiagaraQuaternionUtils.ush"
 
 struct FNiagaraTransform
