- Fix Debug HUD toggle not working

#tests Use console command NiagaraDebugHUD and see if it toggles the HUD on / off when it's not enabled
#rnx

diff --git a/Engine/Plugins/FX/Niagara/Source/Niagara/Private/NiagaraDebugHud.cpp b/Engine/Plugins/FX/Niagara/Source/Niagara/Private/NiagaraDebugHud.cpp
--- a/Engine/Plugins/FX/Niagara/Source/Niagara/Private/NiagaraDebugHud.cpp
+++ b/Engine/Plugins/FX/Niagara/Source/Niagara/Private/NiagaraDebugHud.cpp
@@ -178,13 +178,16 @@ namespace NiagaraDebugLocal
 			{
 				if ( Args.Num() == 0 )
 				{
-					Settings.bOverviewEnabled = !Settings.bOverviewEnabled;
+					const bool bIsEnabled = Settings.bHudEnabled && Settings.bOverviewEnabled;
+					Settings.bHudEnabled = true;
+					Settings.bOverviewEnabled = !bIsEnabled;
 					Settings.OverviewMode = ENiagaraDebugHUDOverviewMode::Overview;
 				}
 				else
 				{
 					const int32 Mode = FCString::Atoi(*Args[0]);
 
+					Settings.bHudEnabled = true;
 					Settings.bOverviewEnabled = Mode > 0;
 					Settings.OverviewMode = FMath::Clamp(ENiagaraDebugHUDOverviewMode(Mode - 1), ENiagaraDebugHUDOverviewMode::Overview, ENiagaraDebugHUDOverviewMode::GpuComputePerformance);
 				}
