Added dynamic input support for inline parameters in overview. Only works if script has an overview display name specified

diff --git a/Engine/Plugins/FX/Niagara/Source/Niagara/Classes/NiagaraScript.h b/Engine/Plugins/FX/Niagara/Source/Niagara/Classes/NiagaraScript.h
--- a/Engine/Plugins/FX/Niagara/Source/Niagara/Classes/NiagaraScript.h
+++ b/Engine/Plugins/FX/Niagara/Source/Niagara/Classes/NiagaraScript.h
@@ -676,6 +676,11 @@ public:
 	UPROPERTY(EditAnywhere, Category = Script)
 	FText Category;
 
+	/** If this script is displayed inline in the overview (dynamic input in the overview node for example), use this display name.
+	 *  Should be short and concise. */
+	UPROPERTY(EditAnywhere, Category = Script)
+	FText InlineOverviewDisplayName;
+
 	UPROPERTY(EditAnywhere, Category = "Asset Options")
 	TArray<FNiagaraAssetTagDefinitionReference> AssetTagDefinitionReferences;
 
diff --git a/Engine/Plugins/FX/Niagara/Source/NiagaraEditor/Private/ViewModels/Stack/NiagaraStackFunctionInput.cpp b/Engine/Plugins/FX/Niagara/Source/NiagaraEditor/Private/ViewModels/Stack/NiagaraStackFunctionInput.cpp
--- a/Engine/Plugins/FX/Niagara/Source/NiagaraEditor/Private/ViewModels/Stack/NiagaraStackFunctionInput.cpp
+++ b/Engine/Plugins/FX/Niagara/Source/NiagaraEditor/Private/ViewModels/Stack/NiagaraStackFunctionInput.cpp
@@ -3240,7 +3240,7 @@ bool UNiagaraStackFunctionInput::ShouldDisplayInline() const
 {
 	// we only allow values to show up that are supposed to be visible
 	// we don't allow advanced values to be inlined as they can introduce problems with filtering
-	if(InputMetaData.IsSet() && InputMetaData->bDisplayInOverviewStack && InputValues.Mode == EValueMode::Local && GetShouldPassFilterForVisibleCondition())
+	if(InputMetaData.IsSet() && InputMetaData->bDisplayInOverviewStack && (InputValues.Mode == EValueMode::Local || InputValues.Mode == EValueMode::Dynamic) && GetShouldPassFilterForVisibleCondition())
 	{
 		return true;				
 	}
diff --git a/Engine/Plugins/FX/Niagara/Source/NiagaraEditorWidgets/Private/SNiagaraOverviewInlineParameterBox.cpp b/Engine/Plugins/FX/Niagara/Source/NiagaraEditorWidgets/Private/SNiagaraOverviewInlineParameterBox.cpp
--- a/Engine/Plugins/FX/Niagara/Source/NiagaraEditorWidgets/Private/SNiagaraOverviewInlineParameterBox.cpp
+++ b/Engine/Plugins/FX/Niagara/Source/NiagaraEditorWidgets/Private/SNiagaraOverviewInlineParameterBox.cpp
@@ -149,12 +149,13 @@ TArray<TSharedRef<SWidget>> SNiagaraOverviewInlineParameterBox::GenerateParamete
 			case UNiagaraStackFunctionInput::EValueMode::Local:
 				Widget = GenerateParameterWidgetFromLocalValue(FunctionInput);
 				break;
+		case UNiagaraStackFunctionInput::EValueMode::Dynamic:
+				Widget = GenerateParameterWidgetFromDynamicInput(FunctionInput);
+				break;
 			case UNiagaraStackFunctionInput::EValueMode::Data:
-				// @todo currently data mode inputs won't be retrieved by GetInlineParameterInputs
 				Widget = GenerateParameterWidgetFromDataInterface(FunctionInput);
 				break;
 			case UNiagaraStackFunctionInput::EValueMode::ObjectAsset:
-				// @todo currently object asset inputs won't be retrieved by GetInlineParameterInputs
 				Widget = SNullWidget::NullWidget;
 				break;
 			default:
@@ -193,18 +194,11 @@ TSharedRef<SWidget> SNiagaraOverviewInlineParameterBox::GenerateParameterWidgetF
 	// by default we display type information. In some cases we have overrides, but even then we want to keep the type information as we continue using it in the tooltips
 	const FText ValueText = TypeUtilities->GetStackDisplayText(Variable);
 	FText DisplayedText = ValueText;
-	const FLinearColor OriginalTypeColor = UEdGraphSchema_Niagara::GetTypeColor(Type);
-	FLinearColor ActualTypeColor = OriginalTypeColor;
 	
 	TSharedPtr<IToolTip> TooltipOverride;
-
-	// if we have an enum type or integer acting as an enum, we want to override the type color to be gray instead
-	if(Type.GetEnum() || (Type == FNiagaraTypeDefinition::GetIntDef() && InputMetaData.IsSet() && InputMetaData->WidgetCustomization.WidgetType == ENiagaraInputWidgetType::EnumStyle))
-	{
-		ActualTypeColor = FLinearColor(0.02f, 0.02f, 0.02f, 1.f);
-	}
 	
-	FLinearColor DisplayedColor = InputMetaData->bOverrideColor ? InputMetaData->InlineParameterColorOverride : ActualTypeColor;
+	FLinearColor TypeColor = GetTypeColor(FunctionInput);
+	FLinearColor DisplayedColor = DetermineBackgroundColor(FunctionInput);
 
 	UTexture2D* Icon = nullptr;
 	if (const UEnum* Enum = Type.GetEnum())
@@ -322,7 +316,7 @@ TSharedRef<SWidget> SNiagaraOverviewInlineParameterBox::GenerateParameterWidgetF
 		else
 		{
 			// if we are using a custom color, we set the button style to use a white base in order to get rid of a tint
-			if(DisplayedColor != ActualTypeColor)
+			if(DisplayedColor != TypeColor)
 			{
 				ParameterWidgetButton->SetButtonStyle(&FNiagaraEditorWidgetsStyle::Get().GetWidgetStyle<FButtonStyle>("NiagaraEditor.SystemOverview.InlineParameterButton.NoTint"));
 			}
@@ -355,6 +349,70 @@ TSharedRef<SWidget> SNiagaraOverviewInlineParameterBox::GenerateParameterWidgetF
 	return ParameterWidget.ToSharedRef();
 }
 
+TSharedRef<SWidget> SNiagaraOverviewInlineParameterBox::GenerateParameterWidgetFromDynamicInput(UNiagaraStackFunctionInput* FunctionInput)
+{
+	if(!ensure(FunctionInput->GetValueMode() == UNiagaraStackFunctionInput::EValueMode::Dynamic))
+	{
+		UE_LOGFMT(LogNiagaraEditor, Verbose, "Wrong input configuration for {0}", __FUNCTION__);
+		return SNullWidget::NullWidget;
+	}
+	
+	if(!ensure(FunctionInput->GetInputFunctionCallNode().GetScriptData() != nullptr))
+	{
+		UE_LOGFMT(LogNiagaraEditor, Verbose, "Invalid dynamic input script {0}", __FUNCTION__);
+		return SNullWidget::NullWidget;
+	}
+
+	if(!ensure(FunctionInput->GetDynamicInputNode() != nullptr && FunctionInput->GetDynamicInputNode()->GetScriptData() != nullptr))
+	{
+		UE_LOGFMT(LogNiagaraEditor, Verbose, "Invalid dynamic input script {0}", __FUNCTION__);
+		return SNullWidget::NullWidget;
+	}
+
+	FVersionedNiagaraScriptData* ScriptData = FunctionInput->GetDynamicInputNode()->GetScriptData();
+
+	if(ScriptData->InlineOverviewDisplayName.IsEmptyOrWhitespace() == false)
+	{
+		TWeakObjectPtr<const UNiagaraStackFunctionInput> FunctionInputWeak = FunctionInput;
+
+		TSharedRef<SButton> ParameterWidgetButton = SNew(SButton)
+		.Text(ScriptData->InlineOverviewDisplayName)
+		.TextStyle(&FNiagaraEditorWidgetsStyle::Get().GetWidgetStyle<FTextBlockStyle>("NiagaraEditor.SystemOverview.InlineParameterText"))
+		.ButtonStyle(&FNiagaraEditorWidgetsStyle::Get().GetWidgetStyle<FButtonStyle>("NiagaraEditor.SystemOverview.InlineParameterButton"))
+		.ContentPadding(FMargin(0.f))
+		.OnClicked(this, &SNiagaraOverviewInlineParameterBox::NavigateToStack, FunctionInputWeak);
+
+		ParameterWidgetButton->SetBorderBackgroundColor(DetermineBackgroundColor(FunctionInput));
+
+		FNiagaraTypeDefinition Type = FunctionInput->GetInputType();
+		FNiagaraVariable Variable(Type, FunctionInput->GetInputParameterHandle().GetParameterHandleString());
+
+		TSharedRef<SVerticalBox> DynamicInputInformationBox = SNew(SVerticalBox)
+			+ SVerticalBox::Slot()
+			.AutoHeight()
+			[
+				SNew(SHorizontalBox)
+				+ SHorizontalBox::Slot()
+				.Padding(4.f)
+				.AutoWidth()
+				[
+					SNew(STextBlock).Text(LOCTEXT("DynamicInputLabelKey", "Dynamic Input: "))
+				]
+				+ SHorizontalBox::Slot()
+				.Padding(4.f)
+				.AutoWidth()
+				[
+					SNew(STextBlock).Text(FText::FormatOrdered(LOCTEXT("DynamicInputLabelValue", "{0}"), FText::AsCultureInvariant(GetNameSafe(FunctionInput->GetDynamicInputNode()->FunctionScript))))
+				]
+			];
+		
+		ParameterWidgetButton->SetToolTip(FNiagaraParameterUtilities::GetTooltipWidget(Variable, false, DynamicInputInformationBox));
+		return ParameterWidgetButton;
+	}
+
+	return SNullWidget::NullWidget;	
+}
+
 TSharedRef<SWidget> SNiagaraOverviewInlineParameterBox::GenerateParameterWidgetFromDataInterface(UNiagaraStackFunctionInput* FunctionInput)
 {
 	// if(DataInterfaceInput.DataInterface.IsValid())
@@ -427,4 +485,28 @@ void SNiagaraOverviewInlineParameterBox::OnModuleItemStructureChanged(ENiagaraSt
 	ConstructChildren();
 }
 
+FLinearColor SNiagaraOverviewInlineParameterBox::GetTypeColor(UNiagaraStackFunctionInput* FunctionInput) const
+{
+	FNiagaraTypeDefinition Type = FunctionInput->GetInputType();
+	TOptional<FNiagaraVariableMetaData> InputMetaData = FunctionInput->GetInputMetaData();
+
+	// by default we display type information. In some cases we have overrides, but even then we want to keep the type information as we continue using it in the tooltips
+	const FLinearColor OriginalTypeColor = UEdGraphSchema_Niagara::GetTypeColor(Type);
+	FLinearColor ActualTypeColor = OriginalTypeColor;
+
+	// if we have an enum type or integer acting as an enum, we want to override the type color to be gray instead
+	if(Type.GetEnum() || (Type == FNiagaraTypeDefinition::GetIntDef() && InputMetaData.IsSet() && InputMetaData->WidgetCustomization.WidgetType == ENiagaraInputWidgetType::EnumStyle))
+	{
+		ActualTypeColor = FLinearColor(0.02f, 0.02f, 0.02f, 1.f);
+	}
+
+	return ActualTypeColor;
+}
+
+FLinearColor SNiagaraOverviewInlineParameterBox::DetermineBackgroundColor(UNiagaraStackFunctionInput* FunctionInput) const
+{
+	TOptional<FNiagaraVariableMetaData> InputMetaData = FunctionInput->GetInputMetaData();	
+	return InputMetaData->bOverrideColor ? InputMetaData->InlineParameterColorOverride : GetTypeColor(FunctionInput);
+}
+
 #undef LOCTEXT_NAMESPACE
diff --git a/Engine/Plugins/FX/Niagara/Source/NiagaraEditorWidgets/Private/SNiagaraOverviewInlineParameterBox.h b/Engine/Plugins/FX/Niagara/Source/NiagaraEditorWidgets/Private/SNiagaraOverviewInlineParameterBox.h
--- a/Engine/Plugins/FX/Niagara/Source/NiagaraEditorWidgets/Private/SNiagaraOverviewInlineParameterBox.h
+++ b/Engine/Plugins/FX/Niagara/Source/NiagaraEditorWidgets/Private/SNiagaraOverviewInlineParameterBox.h
@@ -37,12 +37,19 @@ private:
 	void ConstructChildren();
 	TArray<TSharedRef<SWidget>> GenerateParameterWidgets();
 	TSharedRef<SWidget> GenerateParameterWidgetFromLocalValue(UNiagaraStackFunctionInput* FunctionInput);
+	TSharedRef<SWidget> GenerateParameterWidgetFromDynamicInput(UNiagaraStackFunctionInput* FunctionInput);
 	TSharedRef<SWidget> GenerateParameterWidgetFromDataInterface(UNiagaraStackFunctionInput* FunctionInput);
 	
 	/** We build a substitute map of entries that we want to */
 	TWeakObjectPtr<const UNiagaraStackFunctionInput> FindSubstituteEntry(const UNiagaraStackFunctionInput* Input);
 
 	void OnModuleItemStructureChanged(ENiagaraStructureChangedFlags);
+
+	/** Gets the color based on type. */
+	FLinearColor GetTypeColor(UNiagaraStackFunctionInput* FunctionInput) const;
+
+	/** Determines the color for a function input. This is usually based on type, but can be overridden and handles special cases. */
+	FLinearColor DetermineBackgroundColor(UNiagaraStackFunctionInput* FunctionInput) const;
 private:
 	/** The module item whose parameters this box is representing */
 	TWeakObjectPtr<UNiagaraStackModuleItem> ModuleItem;
