Added stack row highlighting based on an explicitly focused stack entry. This allows the inline parameter clicking to highlight a relevant row instead of abusing the search bar by entering a search term (not what the search bar is for, nor was it always accurate)

diff --git a/Engine/Plugins/FX/Niagara/Source/NiagaraEditor/Private/ViewModels/NiagaraSystemSelectionViewModel.cpp b/Engine/Plugins/FX/Niagara/Source/NiagaraEditor/Private/ViewModels/NiagaraSystemSelectionViewModel.cpp
--- a/Engine/Plugins/FX/Niagara/Source/NiagaraEditor/Private/ViewModels/NiagaraSystemSelectionViewModel.cpp
+++ b/Engine/Plugins/FX/Niagara/Source/NiagaraEditor/Private/ViewModels/NiagaraSystemSelectionViewModel.cpp
@@ -517,6 +517,7 @@ void UNiagaraSystemSelectionViewModel::UpdateExternalSelectionState()
 		}
 	}
 
+	SelectionStackViewModel->SetFocusedStackEntryOverride(nullptr);
 	StackSelection->SetSelectedEntries(StackSelectionEntries);
 	OnEntrySelectionChangedDelegate.Broadcast();
 	if (bEmitterHandleIdSelectionChanged)
diff --git a/Engine/Plugins/FX/Niagara/Source/NiagaraEditor/Private/ViewModels/Stack/NiagaraStackViewModel.cpp b/Engine/Plugins/FX/Niagara/Source/NiagaraEditor/Private/ViewModels/Stack/NiagaraStackViewModel.cpp
--- a/Engine/Plugins/FX/Niagara/Source/NiagaraEditor/Private/ViewModels/Stack/NiagaraStackViewModel.cpp
+++ b/Engine/Plugins/FX/Niagara/Source/NiagaraEditor/Private/ViewModels/Stack/NiagaraStackViewModel.cpp
@@ -309,6 +309,9 @@ void UNiagaraStackViewModel::SetSearchTextExternal(const FText& NewSearchText)
 
 void UNiagaraStackViewModel::OnSearchTextChanged(const FText& SearchText)
 {
+	// If the user is searching that means we don't want to have a previous override mess with the search
+	FocusedStackEntryOverride.Reset();
+	
 	if (RootEntry && !CurrentSearchText.EqualTo(SearchText))
 	{
 		CurrentSearchText = SearchText;
@@ -335,6 +338,11 @@ const TArray<UNiagaraStackViewModel::FSearchResult>& UNiagaraStackViewModel::Get
 
 UNiagaraStackEntry* UNiagaraStackViewModel::GetCurrentFocusedEntry() const
 {
+	if(FocusedStackEntryOverride.IsValid() && FocusedStackEntryOverride->IsFinalized() == false)
+	{
+		return FocusedStackEntryOverride.Get();
+	}
+	
 	if (FocusedSearchResultCache.IsSet() && FocusedSearchResultCache->GetEntry() != nullptr && !FocusedSearchResultCache->GetEntry()->IsFinalized())
 	{
 		return FocusedSearchResultCache->GetEntry();
@@ -652,6 +660,7 @@ void UNiagaraStackViewModel::RestartSearch()
 	if(bRestartSearch == false)
 	{
 		bRestartSearch = true;
+		FocusedStackEntryOverride.Reset();
 	}
 }
 
@@ -846,6 +855,18 @@ void UNiagaraStackViewModel::SetLastScrollPosition(double InLastScrollPosition)
 	}
 }
 
+void UNiagaraStackViewModel::SetFocusedStackEntryOverride(UNiagaraStackEntry* InStackEntry)
+{
+	if(InStackEntry)
+	{
+		FocusedStackEntryOverride = InStackEntry;
+	}
+	else
+	{
+		FocusedStackEntryOverride.Reset();
+	}
+}
+
 void UNiagaraStackViewModel::EntryExpansionChanged()
 {
 	ExpansionChangedDelegate.Broadcast();
diff --git a/Engine/Plugins/FX/Niagara/Source/NiagaraEditor/Public/ViewModels/Stack/NiagaraStackViewModel.h b/Engine/Plugins/FX/Niagara/Source/NiagaraEditor/Public/ViewModels/Stack/NiagaraStackViewModel.h
--- a/Engine/Plugins/FX/Niagara/Source/NiagaraEditor/Public/ViewModels/Stack/NiagaraStackViewModel.h
+++ b/Engine/Plugins/FX/Niagara/Source/NiagaraEditor/Public/ViewModels/Stack/NiagaraStackViewModel.h
@@ -118,6 +118,8 @@ public:
 	NIAGARAEDITOR_API double GetLastScrollPosition() const;
 	NIAGARAEDITOR_API void SetLastScrollPosition(double InLastScrollPosition);
 
+	NIAGARAEDITOR_API void SetFocusedStackEntryOverride(UNiagaraStackEntry* InStackEntry);
+	
 	NIAGARAEDITOR_API virtual void Tick();
 	//~ stack search stuff
 	NIAGARAEDITOR_API void ResetSearchText();
@@ -206,12 +208,14 @@ private:
 	FOnStructureChanged StructureChangedDelegate;
 	FOnDataObjectChanged DataObjectChangedDelegate;
 
+
 	// ~Search stuff
 	FText CurrentSearchText;
 	FOnSearchCompleted SearchCompletedDelegate;
 	TArray<FSearchWorkItem> ItemsToSearch;
 	TArray<FSearchResult> CurrentSearchResults;
 	static NIAGARAEDITOR_API const double MaxSearchTime;
+	TWeakObjectPtr<UNiagaraStackEntry> FocusedStackEntryOverride;
 	TOptional<FSearchResult> FocusedSearchResultCache;
 	bool bRestartSearch;
 	bool bRefreshPending;
diff --git a/Engine/Plugins/FX/Niagara/Source/NiagaraEditorWidgets/Private/SNiagaraOverviewInlineParameterBox.cpp b/Engine/Plugins/FX/Niagara/Source/NiagaraEditorWidgets/Private/SNiagaraOverviewInlineParameterBox.cpp
--- a/Engine/Plugins/FX/Niagara/Source/NiagaraEditorWidgets/Private/SNiagaraOverviewInlineParameterBox.cpp
+++ b/Engine/Plugins/FX/Niagara/Source/NiagaraEditorWidgets/Private/SNiagaraOverviewInlineParameterBox.cpp
@@ -57,7 +57,7 @@ SNiagaraOverviewInlineParameterBox::~SNiagaraOverviewInlineParameterBox()
 	BoundFunctionInputs.Empty();
 }
 
-FReply SNiagaraOverviewInlineParameterBox::NavigateToStack(TWeakObjectPtr<const UNiagaraStackFunctionInput> FunctionInput)
+FReply SNiagaraOverviewInlineParameterBox::NavigateToStack(TWeakObjectPtr<UNiagaraStackFunctionInput> FunctionInput)
 {
 	// even if we can't navigate, the button should consume the input
 	if(!ModuleItem.IsValid() || !FunctionInput.IsValid() || FunctionInput->IsFinalized())
@@ -68,19 +68,13 @@ FReply SNiagaraOverviewInlineParameterBox::NavigateToStack(TWeakObjectPtr<const
 	ModuleItem->GetSystemViewModel()->GetSelectionViewModel()->UpdateSelectedEntries({ModuleItem.Get()}, {}, true);
 
 	// we attempt to find a substitute entry, i.e. if our current input is an inline edit toggle, we try to find the input that the toggle is handling instead
-	TWeakObjectPtr<const UNiagaraStackFunctionInput> SubstituteEntry = FindSubstituteEntry(FunctionInput.Get());
+	TWeakObjectPtr<UNiagaraStackFunctionInput> SubstituteEntry = FindSubstituteEntry(FunctionInput.Get());
 	if(SubstituteEntry.IsValid())
 	{
 		FunctionInput = SubstituteEntry;
 	}
 	
-	TArray<UNiagaraStackEntry::FStackSearchItem> SearchItems;
-	FunctionInput->GetSearchItems(SearchItems);
-	if(SearchItems.Num() > 0)
-	{
-		ModuleItem->GetSystemViewModel()->GetSelectionViewModel()->GetSelectionStackViewModel()->SetSearchTextExternal(SearchItems[0].Value);
-		return FReply::Handled();
-	}
+	ModuleItem->GetSystemViewModel()->GetSelectionViewModel()->GetSelectionStackViewModel()->SetFocusedStackEntryOverride(FunctionInput.Get());
 
 	return FReply::Handled();
 }
@@ -295,7 +289,7 @@ TSharedRef<SWidget> SNiagaraOverviewInlineParameterBox::GenerateParameterWidgetF
 	// we might have cases in which we don't want to display a parameter's data, such as if we only want to display a bool parameter when it is set to True
 	if(bGenerateProperWidget)
 	{
-		TWeakObjectPtr<const UNiagaraStackFunctionInput> FunctionInputWeak = FunctionInput;
+		TWeakObjectPtr<UNiagaraStackFunctionInput> FunctionInputWeak = FunctionInput;
 		
 		TSharedRef<SButton> ParameterWidgetButton = SNew(SButton)
 		.Text(DisplayedText)
@@ -373,7 +367,7 @@ TSharedRef<SWidget> SNiagaraOverviewInlineParameterBox::GenerateParameterWidgetF
 
 	if(ScriptData->InlineOverviewDisplayName.IsEmptyOrWhitespace() == false)
 	{
-		TWeakObjectPtr<const UNiagaraStackFunctionInput> FunctionInputWeak = FunctionInput;
+		TWeakObjectPtr<UNiagaraStackFunctionInput> FunctionInputWeak = FunctionInput;
 
 		TSharedRef<SButton> ParameterWidgetButton = SNew(SButton)
 		.Text(ScriptData->InlineOverviewDisplayName)
@@ -449,7 +443,7 @@ TSharedRef<SWidget> SNiagaraOverviewInlineParameterBox::GenerateParameterWidgetF
 	return SNullWidget::NullWidget;
 }
 
-TWeakObjectPtr<const UNiagaraStackFunctionInput> SNiagaraOverviewInlineParameterBox::FindSubstituteEntry(const UNiagaraStackFunctionInput* InInput)
+TWeakObjectPtr<UNiagaraStackFunctionInput> SNiagaraOverviewInlineParameterBox::FindSubstituteEntry(const UNiagaraStackFunctionInput* InInput)
 {
 	// there can only be substitute entries for inline edit condition toggles
 	if(!InInput->GetInputMetaData().IsSet() || !InInput->GetInputMetaData()->bInlineEditConditionToggle)
@@ -461,7 +455,7 @@ TWeakObjectPtr<const UNiagaraStackFunctionInput> SNiagaraOverviewInlineParameter
 	ModuleItem->GetParameterInputs(AllInputs);
 
 	// we gather all edit conditions here so we can later map them back
-	for(const UNiagaraStackFunctionInput* Input : AllInputs)
+	for(UNiagaraStackFunctionInput* Input : AllInputs)
 	{
 		if(Input->GetHasEditCondition())
 		{
diff --git a/Engine/Plugins/FX/Niagara/Source/NiagaraEditorWidgets/Private/SNiagaraOverviewInlineParameterBox.h b/Engine/Plugins/FX/Niagara/Source/NiagaraEditorWidgets/Private/SNiagaraOverviewInlineParameterBox.h
--- a/Engine/Plugins/FX/Niagara/Source/NiagaraEditorWidgets/Private/SNiagaraOverviewInlineParameterBox.h
+++ b/Engine/Plugins/FX/Niagara/Source/NiagaraEditorWidgets/Private/SNiagaraOverviewInlineParameterBox.h
@@ -33,7 +33,7 @@ public:
 
 private:
 	/** Selects the module item in the stack and searches for the function input */
-	FReply NavigateToStack(TWeakObjectPtr<const UNiagaraStackFunctionInput> FunctionInput);
+	FReply NavigateToStack(TWeakObjectPtr<UNiagaraStackFunctionInput> FunctionInput);
 	void ConstructChildren();
 	TArray<TSharedRef<SWidget>> GenerateParameterWidgets();
 	TSharedRef<SWidget> GenerateParameterWidgetFromLocalValue(UNiagaraStackFunctionInput* FunctionInput);
@@ -41,7 +41,7 @@ private:
 	TSharedRef<SWidget> GenerateParameterWidgetFromDataInterface(UNiagaraStackFunctionInput* FunctionInput);
 	
 	/** We build a substitute map of entries that we want to */
-	TWeakObjectPtr<const UNiagaraStackFunctionInput> FindSubstituteEntry(const UNiagaraStackFunctionInput* Input);
+	TWeakObjectPtr<UNiagaraStackFunctionInput> FindSubstituteEntry(const UNiagaraStackFunctionInput* Input);
 
 	void OnModuleItemStructureChanged(ENiagaraStructureChangedFlags);
 
