Moved the PreviewMoveiPath properties to the "Asset Options" category and put them under advanced

diff --git a/Engine/Plugins/FX/Niagara/Source/Niagara/Classes/NiagaraEmitter.h b/Engine/Plugins/FX/Niagara/Source/Niagara/Classes/NiagaraEmitter.h
--- a/Engine/Plugins/FX/Niagara/Source/Niagara/Classes/NiagaraEmitter.h
+++ b/Engine/Plugins/FX/Niagara/Source/Niagara/Classes/NiagaraEmitter.h
@@ -743,7 +743,7 @@ public:
 	UPROPERTY(EditAnywhere, AdvancedDisplay, Category = "Asset Options", AssetRegistrySearchable, DisplayName="Asset Description")
 	FText TemplateAssetDescription;
 
-	UPROPERTY(EditAnywhere, Category="Preview", AssetRegistrySearchable, meta = (SkipSystemResetOnChange = "true", AllowedClasses="/Script/MediaAssets.FileMediaSource"))
+	UPROPERTY(EditAnywhere, AdvancedDisplay, Category="Asset Options", AssetRegistrySearchable, meta = (SkipSystemResetOnChange = "true", AllowedClasses="/Script/MediaAssets.FileMediaSource"))
 	FSoftObjectPath PreviewMoviePath;
 	
 	/** Category to collate this emitter into for "add new emitter" dialogs.*/
diff --git a/Engine/Plugins/FX/Niagara/Source/Niagara/Classes/NiagaraSystem.h b/Engine/Plugins/FX/Niagara/Source/Niagara/Classes/NiagaraSystem.h
--- a/Engine/Plugins/FX/Niagara/Source/Niagara/Classes/NiagaraSystem.h
+++ b/Engine/Plugins/FX/Niagara/Source/Niagara/Classes/NiagaraSystem.h
@@ -503,7 +503,7 @@ public:
 	UPROPERTY(meta = (SkipSystemResetOnChange = "true"))
 	TArray<FNiagaraAssetTagDefinitionReference> AssetTags;
 
-	UPROPERTY(EditAnywhere, Category="Preview", AssetRegistrySearchable, meta = (SkipSystemResetOnChange = "true", AllowedClasses="/Script/MediaAssets.FileMediaSource"))
+	UPROPERTY(EditAnywhere, AdvancedDisplay, Category="Asset Options", AssetRegistrySearchable, meta = (SkipSystemResetOnChange = "true", AllowedClasses="/Script/MediaAssets.FileMediaSource"))
 	FSoftObjectPath PreviewMoviePath;
 	
 	UPROPERTY()
