- Add adapter to parameter binding customization
- Allows the customization to be shared with different types

#tests Tested the other areas that use this widget still work, i.e. sim stage thread binding, binding to a mesh array in the mesh renderer, etc
#rnx

diff --git a/Engine/Plugins/FX/Niagara/Source/Niagara/Private/NiagaraParameterBinding.cpp b/Engine/Plugins/FX/Niagara/Source/Niagara/Private/NiagaraParameterBinding.cpp
--- a/Engine/Plugins/FX/Niagara/Source/Niagara/Private/NiagaraParameterBinding.cpp
+++ b/Engine/Plugins/FX/Niagara/Source/Niagara/Private/NiagaraParameterBinding.cpp
@@ -5,89 +5,6 @@
 #include "NiagaraDataInterface.h"
 
 #if WITH_EDITORONLY_DATA
-bool FNiagaraParameterBinding::CanBindTo(FNiagaraTypeDefinition TypeDefinition) const
-{
-	if (AllowedTypeDefinitions.Contains(TypeDefinition))
-	{
-		return true;
-	}
-
-	if (TypeDefinition.IsStatic() && AllowStaticVariables())
-	{
-		if (AllowedTypeDefinitions.Contains(TypeDefinition.RemoveStaticDef()))
-		{
-			return true;
-		}
-	}
-
-	if (UClass* TypeClass = TypeDefinition.GetClass())
-	{
-		if (TypeClass->IsChildOf(UNiagaraDataInterface::StaticClass()))
-		{
-			for (UClass* AllowedClass : AllowedDataInterfaces)
-			{
-				if (AllowedClass && TypeClass->IsChildOf(AllowedClass))
-				{
-					return true;
-				}
-			}
-
-			for (UClass* AllowedInterface : AllowedInterfaces)
-			{
-				if (AllowedInterface && TypeClass->ImplementsInterface(AllowedInterface))
-				{
-					return true;
-				}
-			}
-		}
-		else
-		{
-			for (UClass* AllowedClass : AllowedObjects)
-			{
-				if (AllowedClass && TypeClass->IsChildOf(AllowedClass))
-				{
-					return true;
-				}
-			}
-		}
-	}
-
-	return false;
-}
-
-bool FNiagaraParameterBinding::CanBindTo(FNiagaraVariableBase InVariable, FNiagaraVariableBase& OutAliasedVariable, FStringView EmitterName) const
-{
-	OutAliasedVariable = InVariable;
-	if ( !AllowStaticVariables() && InVariable.GetType().IsStatic() )
-	{
-		return false;
-	}
-
-	if ( !CanBindTo(InVariable.GetType()) )
-	{
-		return false;
-	}
-
-	if (InVariable.IsInNameSpace(FNiagaraConstants::ParticleAttributeNamespaceString))
-	{
-		return AllowParticleParameters();
-	}
-	else if (InVariable.IsInNameSpace(EmitterName))
-	{
-		if ( AllowEmitterParameters() )
-		{
-			OutAliasedVariable.ReplaceRootNamespace(EmitterName, FNiagaraConstants::EmitterNamespaceString);
-			return true;
-		}
-	}
-	else if (InVariable.IsInNameSpace(FNiagaraConstants::SystemNamespaceString))
-	{
-		return AllowSystemParameters();
-	}
-
-	return false;
-}
-
 void FNiagaraParameterBinding::OnRenameEmitter(FStringView EmitterName)
 {
 	if (AliasedParameter.IsInNameSpace(FNiagaraConstants::EmitterNamespaceString))
@@ -140,7 +57,7 @@ void FNiagaraParameterBinding::SetToDefault()
 	}
 }
 
-bool FNiagaraParameterBinding::IsSetoToDefault() const
+bool FNiagaraParameterBinding::IsSetToDefault() const
 {
 	if (HasDefaultValueEditorOnly() && DefaultAliasedParameter.IsDataAllocated())
 	{
diff --git a/Engine/Plugins/FX/Niagara/Source/Niagara/Public/NiagaraParameterBinding.h b/Engine/Plugins/FX/Niagara/Source/Niagara/Public/NiagaraParameterBinding.h
--- a/Engine/Plugins/FX/Niagara/Source/Niagara/Public/NiagaraParameterBinding.h
+++ b/Engine/Plugins/FX/Niagara/Source/Niagara/Public/NiagaraParameterBinding.h
@@ -31,9 +31,6 @@ struct FNiagaraParameterBinding
 #if WITH_EDITORONLY_DATA
 	virtual ~FNiagaraParameterBinding() = default;
 
-	NIAGARA_API bool CanBindTo(FNiagaraTypeDefinition TypeDefinition) const;
-	NIAGARA_API bool CanBindTo(FNiagaraVariableBase InVariable, FNiagaraVariableBase& OutAliasedVariable, FStringView EmitterName) const;
-
 	NIAGARA_API void OnRenameEmitter(FStringView EmitterName);
 	NIAGARA_API void OnRenameVariable(const FNiagaraVariableBase& OldVariable, const FNiagaraVariableBase& NewVariable, FStringView EmitterName);
 	NIAGARA_API void OnRemoveVariable(const FNiagaraVariableBase& OldVariable, FStringView EmitterName);
@@ -61,7 +58,7 @@ struct FNiagaraParameterBinding
 	TConstArrayView<UClass*> GetAllowedInterfaces() const { return MakeArrayView(AllowedInterfaces); }
 	TConstArrayView<FNiagaraTypeDefinition> GetAllowedTypeDefinitions() const { return MakeArrayView(AllowedTypeDefinitions); }
 
-	NIAGARA_API bool IsSetoToDefault() const;
+	NIAGARA_API bool IsSetToDefault() const;
 	NIAGARA_API void SetToDefault();
 
 	NIAGARA_API FString ToString() const;
diff --git a/Engine/Plugins/FX/Niagara/Source/NiagaraEditor/Internal/Customizations/NiagaraParameterBindingAdapterCustomization.cpp b/Engine/Plugins/FX/Niagara/Source/NiagaraEditor/Internal/Customizations/NiagaraParameterBindingAdapterCustomization.cpp
new file mode 100644
--- /dev/null
+++ b/Engine/Plugins/FX/Niagara/Source/NiagaraEditor/Internal/Customizations/NiagaraParameterBindingAdapterCustomization.cpp
@@ -0,0 +1,670 @@
+// Copyright Epic Games, Inc. All Rights Reserved.
+
+#include "NiagaraParameterBindingAdapterCustomization.h"
+
+#include "Customizations/NiagaraTypeCustomizations.h"
+#include "Widgets/SNiagaraParameterName.h"
+#include "INiagaraEditorTypeUtilities.h"
+#include "NiagaraEditorModule.h"
+#include "NiagaraEmitter.h"
+#include "NiagaraNodeParameterMapBase.h"
+#include "NiagaraParameterMapHistory.h"
+#include "SNiagaraParameterEditor.h"
+#include "NiagaraRendererProperties.h"
+#include "NiagaraScriptSource.h"
+#include "NiagaraSimulationStageBase.h"
+#include "NiagaraSystem.h"
+#include "NiagaraVariableMetaData.h"
+
+#include "Framework/Application/SlateApplication.h"
+#include "Widgets/Input/SComboButton.h"
+#include "DetailLayoutBuilder.h"
+#include "DetailWidgetRow.h"
+#include "IDetailChildrenBuilder.h"
+#include "IDetailGroup.h"
+#include "IPropertyUtilities.h"
+#include "PropertyHandle.h"
+#include "ScopedTransaction.h"
+#include "SGraphActionMenu.h"
+
+#define LOCTEXT_NAMESPACE "NiagaraParameterBindingAdapterCustomization"
+
+namespace NiagaraParameterBindingCustomizationBasePrivate
+{
+	struct FBindingAction : public FEdGraphSchemaAction
+	{
+		FBindingAction() : FEdGraphSchemaAction(FText(), FText(), FText(), 0) {}
+		explicit FBindingAction(FText InName, FText InTooltip) : FEdGraphSchemaAction(FText(), InName, InTooltip, 0) {}
+
+		static TSharedRef<FBindingAction> MakeNone()
+		{
+			TSharedRef<FBindingAction> Action = MakeShared<FBindingAction>(
+				LOCTEXT("None", "None"),
+				LOCTEXT("NoneTooltip", "Not bound to any parameter.")
+			);
+			Action->Grouping = 1;
+			return Action;
+		}
+
+		static TSharedRef<FBindingAction> MakeConstant(const FNiagaraTypeDefinition& TypeDef, TConstArrayView<uint8> ValueData)
+		{
+			TSharedRef<FBindingAction> Action = MakeShared<FBindingAction>(
+				FText::Format(LOCTEXT("ConstantFormat", "Constant Value \"{0}\""), FText::FromString(TypeDef.ToString(ValueData.GetData()))),
+				LOCTEXT("ConstantTooltip", "Use user defined constant value, not bound to any parameter.")
+			);
+			Action->Grouping = 1;
+			return Action;
+		}
+
+		static TSharedRef<FBindingAction> MakeBinding(FName InVariableName, FInstancedStruct InValueStruct)
+		{
+			FText DisplayName = FText::FromName(InVariableName);
+			TSharedRef<FBindingAction> Action = MakeShared<FBindingAction>(
+				DisplayName,
+				FText::Format(LOCTEXT("VariableBindActionTooltip", "Bind to variable \"{0}\""), DisplayName)
+			);
+			Action->Grouping = 0;
+			Action->VariableName = InVariableName;
+			Action->ValueStruct = MoveTemp(InValueStruct);
+			return Action;
+		}
+
+		FName				VariableName;
+		FInstancedStruct	ValueStruct;
+	};
+
+	FText JoinTypeArray(FText Delim, TConstArrayView<UClass*> TypeArray)
+	{
+		TArray<FText> TypeNamesText;
+		for (UClass* Class : TypeArray)
+		{
+			TypeNamesText.Add(Class->GetDisplayNameText());
+		}
+		return FText::Join(Delim, TypeNamesText);
+	}
+
+	FText JoinTypeArray(FText Delim, TConstArrayView<FNiagaraTypeDefinition> TypeArray)
+	{
+		TArray<FText> TypeNamesText;
+		for (const FNiagaraTypeDefinition& TypeDef : TypeArray)
+		{
+			TypeNamesText.Add(TypeDef.GetNameText());
+		}
+		return FText::Join(Delim, TypeNamesText);
+	}
+
+	template<typename TArrayType>
+	void BuildTypesText(FText& InOutString, FText TypesDescription, TArrayType TypeArray)
+	{
+		if (TypeArray.Num() == 0)
+		{
+			return;
+		}
+
+		InOutString = FText::Join(
+			FText::GetEmpty(),
+			TArray<FText>(
+				{
+					InOutString,
+					LOCTEXT("NewLine", "\n"),
+					TypesDescription,
+					JoinTypeArray(LOCTEXT("CommaSpace", ", "), TypeArray)
+				}
+			)
+		);
+	}
+
+	FText GetAllowedBindingTypesText(const FNiagaraParameterBindingAdapter* ParameterBinding)
+	{
+		FText BoundTypes;
+		BuildTypesText(
+			BoundTypes,
+			LOCTEXT("BindingTooltip_DIFormat", "Data Interface Types: "),
+			ParameterBinding->GetAllowedDataInterfaces()
+		);
+		BuildTypesText(
+			BoundTypes,
+			LOCTEXT("BindingTooltip_ObjectFormat", "Object Types: "),
+			ParameterBinding->GetAllowedObjects()
+		);
+		BuildTypesText(
+			BoundTypes,
+			LOCTEXT("BindingTooltip_InterfacesFormat", "Interface Types: "),
+			ParameterBinding->GetAllowedInterfaces()
+		);
+		BuildTypesText(
+			BoundTypes,
+			LOCTEXT("BindingTooltip_AttributesFormat", "Attribute Types: "),
+			ParameterBinding->GetAllowedTypeDefinitions()
+		);
+		return BoundTypes;
+	}
+
+	bool CanBindToType(const FNiagaraParameterBindingAdapter* ParameterBinding, const FNiagaraTypeDefinition& TypeDefinition)
+	{
+		if (ParameterBinding->GetAllowedTypeDefinitions().Contains(TypeDefinition))
+		{
+			return true;
+		}
+
+		if (TypeDefinition.IsStatic() && ParameterBinding->AllowStaticVariables())
+		{
+			if (ParameterBinding->GetAllowedTypeDefinitions().Contains(TypeDefinition.RemoveStaticDef()))
+			{
+				return true;
+			}
+		}
+
+		if (UClass* TypeClass = TypeDefinition.GetClass())
+		{
+			if (TypeClass->IsChildOf(UNiagaraDataInterface::StaticClass()))
+			{
+				for (UClass* AllowedClass : ParameterBinding->GetAllowedDataInterfaces())
+				{
+					if (AllowedClass && TypeClass->IsChildOf(AllowedClass))
+					{
+						return true;
+					}
+				}
+
+				for (UClass* AllowedInterface : ParameterBinding->GetAllowedInterfaces())
+				{
+					if (AllowedInterface && TypeClass->ImplementsInterface(AllowedInterface))
+					{
+						return true;
+					}
+				}
+			}
+			else
+			{
+				for (UClass* AllowedClass : ParameterBinding->GetAllowedObjects())
+				{
+					if (AllowedClass && TypeClass->IsChildOf(AllowedClass))
+					{
+						return true;
+					}
+				}
+			}
+		}
+
+		return false;
+	}
+
+	bool CanBindToNamespace(const FNiagaraParameterBindingAdapter* ParameterBinding, const FNiagaraVariableBase& InAliasedVariable)
+	{
+		if (InAliasedVariable.IsInNameSpace(FNiagaraConstants::ParticleAttributeNamespaceString))
+		{
+			return ParameterBinding->AllowParticleParameters();
+		}
+		else if (InAliasedVariable.IsInNameSpace(FNiagaraConstants::EmitterNamespaceString))
+		{
+			return ParameterBinding->AllowEmitterParameters();
+		}
+		else if (InAliasedVariable.IsInNameSpace(FNiagaraConstants::SystemNamespaceString))
+		{
+			return ParameterBinding->AllowSystemParameters();
+		}
+		else if (InAliasedVariable.IsInNameSpace(FNiagaraConstants::UserNamespaceString))
+		{
+			return ParameterBinding->AllowUserParameters();
+		}
+		return false;
+	}
+}
+
+void FNiagaraParameterBindingAdapterCustomization::CustomizeHeader(TSharedRef<IPropertyHandle> InPropertyHandle, FDetailWidgetRow& HeaderRow, IPropertyTypeCustomizationUtils& StructCustomizationUtils)
+{
+	PropertyHandle = InPropertyHandle;
+	EmitterWeakPtr.Reset();
+	SystemWeakPtr.Reset();
+
+	// We only handle single object
+	UObject* OwnerObject = nullptr;
+	{
+		TArray<UObject*> OuterObjects;
+		PropertyHandle->GetOuterObjects(OuterObjects);
+		OwnerObject = OuterObjects.Num() == 1 ? OuterObjects[0] : nullptr;
+		if (OwnerObject == nullptr)
+		{
+			HeaderRow
+			.NameContent()
+			[
+				PropertyHandle->CreatePropertyNameWidget()
+			]
+			.ValueContent()
+			.MaxDesiredWidth(200.f)
+			[
+				SNew(STextBlock)
+				.Text(FText::FromString(FName::NameToDisplayString(CastField<FStructProperty>(PropertyHandle->GetProperty())->Struct->GetName(), false)))
+				.Font(IDetailLayoutBuilder::GetDetailFont())
+			];
+
+			return;
+		}
+	}
+
+	ParameterBindingAdapter = AdapterCreator(PropertyHandle->GetValueBaseAddress((uint8*)OwnerObject));
+	PropertyHandle->SetOnPropertyResetToDefault(FSimpleDelegate::CreateLambda([]() { }));
+	PropertyHandle->MarkResetToDefaultCustomized(true);
+
+	OwnerWeakPtr	= OwnerObject;
+	EmitterWeakPtr	= OwnerObject->GetTypedOuter<UNiagaraEmitter>();
+	SystemWeakPtr	= OwnerObject->GetTypedOuter<UNiagaraSystem>();
+
+	TSharedPtr<SHorizontalBox> ParameterPanel = SNew(SHorizontalBox);
+
+	// Can this parameter be set to a constant value?
+	FNiagaraParameterBindingAdapter* ParameterBinding = GetParameterBindingAdapter();
+	if (ParameterBinding && ParameterBinding->AllowConstantValue())
+	{
+		// Create scope for constant value
+		const FNiagaraTypeDefinition ConstantValueTypeDef = ParameterBinding->GetConstantTypeDef();
+		TConstArrayView<uint8> ConstantValue = ParameterBinding->GetConstantValue();
+		ConstantValueStructOnScope = MakeShared<FStructOnScope>(ConstantValueTypeDef.GetStruct());
+		FMemory::Memcpy(ConstantValueStructOnScope->GetStructMemory(), ConstantValue.GetData(), ConstantValue.Num());
+
+		// Create editor for constant value
+		FNiagaraEditorModule& NiagaraEditorModule = FModuleManager::GetModuleChecked<FNiagaraEditorModule>("NiagaraEditor");
+		TSharedPtr<INiagaraEditorTypeUtilities, ESPMode::ThreadSafe> TypeEditorUtilities = NiagaraEditorModule.GetTypeUtilities(ConstantValueTypeDef);
+		if (TypeEditorUtilities.IsValid())
+		{
+			FNiagaraInputParameterCustomization CustomizationOptions;
+			CustomizationOptions.bBroadcastValueChangesOnCommitOnly = true; // each broadcast usually forces a recompile, so we only want to do it on commits
+			ConstantValueParameterEditor = TypeEditorUtilities->CreateParameterEditor(ConstantValueTypeDef, EUnit::Unspecified, CustomizationOptions);
+			ConstantValueParameterEditor->UpdateInternalValueFromStruct(ConstantValueStructOnScope.ToSharedRef());
+			ConstantValueParameterEditor->SetOnValueChanged(SNiagaraParameterEditor::FOnValueChange::CreateSP(this, &FNiagaraParameterBindingAdapterCustomization::OnConstantValueChanged));
+
+			ParameterPanel->AddSlot()
+			.AutoWidth()
+			.Padding(0, 0, 0, 3)
+			[
+				SNew(SBox)
+				.HAlign(ConstantValueParameterEditor->GetHorizontalAlignment())
+				.VAlign(ConstantValueParameterEditor->GetVerticalAlignment())
+				.IsEnabled(this, &FNiagaraParameterBindingAdapterCustomization::IsConstantEnabled)
+				.Visibility(this, &FNiagaraParameterBindingAdapterCustomization::IsConstantVisibile)
+				[
+					ConstantValueParameterEditor.ToSharedRef()
+				]
+			];
+		}
+	}
+
+	ParameterPanel->AddSlot()
+	.AutoWidth()
+	[
+		SNew(SNiagaraParameterName)
+		.ParameterName(this, &FNiagaraParameterBindingAdapterCustomization::GetBoundParameterName)
+		.IsReadOnly(true)
+		.IsEnabled(this, &FNiagaraParameterBindingAdapterCustomization::IsBindingEnabled)
+		.Visibility(this, &FNiagaraParameterBindingAdapterCustomization::IsBindingVisibile)
+		.ToolTipText(this, &FNiagaraParameterBindingAdapterCustomization::GetTooltipText)
+	];
+
+	HeaderRow
+	.NameContent()
+	[
+		PropertyHandle->CreatePropertyNameWidget()
+	]
+	.ValueContent()
+	[
+		SNew(SHorizontalBox)
+		+ SHorizontalBox::Slot()
+		[
+			ParameterPanel.ToSharedRef()
+		]
+		+ SHorizontalBox::Slot()
+		.FillWidth(1.0f)
+		.HAlign(HAlign_Right)
+		.Padding(3, 0, 0, 0)
+		[
+			SNew(SComboButton)					
+			.ButtonStyle(FAppStyle::Get(), "HoverHintOnly")
+			.ForegroundColor(FSlateColor::UseForeground())
+			.OnGetMenuContent(this, &FNiagaraParameterBindingAdapterCustomization::OnGetMenuContent)
+			.ContentPadding(2)
+			.MenuPlacement(MenuPlacement_BelowRightAnchor)
+			.ToolTipText(this, &FNiagaraParameterBindingAdapterCustomization::GetTooltipText)
+		]
+		+ SHorizontalBox::Slot()
+		.AutoWidth()
+		.VAlign(VAlign_Center)
+		.Padding(2, 1)
+		[
+			SNew(SButton)
+			.IsFocusable(false)
+			.ToolTipText(LOCTEXT("ResetToDefaultToolTip", "Reset to Default"))
+			.ButtonStyle(FAppStyle::Get(), "NoBorder")
+			.ContentPadding(0)
+			.Visibility(this, &FNiagaraParameterBindingAdapterCustomization::IsResetToDefaultsVisible)
+			.OnClicked(this, &FNiagaraParameterBindingAdapterCustomization::OnResetToDefaultsClicked)
+			.Content()
+			[
+				SNew(SImage)
+				.Image(FAppStyle::GetBrush("PropertyWindow.DiffersFromDefault"))
+				.ColorAndOpacity(FSlateColor::UseForeground())
+			]
+		]
+	];
+}
+
+bool FNiagaraParameterBindingAdapterCustomization::IsValid() const
+{
+	return OwnerWeakPtr.IsValid();
+}
+
+FNiagaraParameterBindingAdapter* FNiagaraParameterBindingAdapterCustomization::GetParameterBindingAdapter() const
+{
+	return IsValid() ? ParameterBindingAdapter.Get() : nullptr;
+}
+
+bool FNiagaraParameterBindingAdapterCustomization::IsBindingEnabled() const
+{
+	const FNiagaraParameterBindingAdapter* ParameterBinding = GetParameterBindingAdapter();
+	return ParameterBinding && ParameterBinding->IsSetToParameter();
+}
+
+bool FNiagaraParameterBindingAdapterCustomization::IsConstantEnabled() const
+{
+	const FNiagaraParameterBindingAdapter* ParameterBinding = GetParameterBindingAdapter();
+	return ParameterBinding && ParameterBinding->IsSetToConstant();
+}
+
+EVisibility FNiagaraParameterBindingAdapterCustomization::IsBindingVisibile() const
+{
+	return IsBindingEnabled() ? EVisibility::Visible : EVisibility::Collapsed;
+}
+
+EVisibility FNiagaraParameterBindingAdapterCustomization::IsConstantVisibile() const
+{
+	return IsConstantEnabled() ? EVisibility::Visible : EVisibility::Collapsed;
+}
+
+FName FNiagaraParameterBindingAdapterCustomization::GetBoundParameterName() const
+{
+	const FNiagaraParameterBindingAdapter* ParameterBinding = GetParameterBindingAdapter();
+	return ParameterBinding ? ParameterBinding->GetBoundParameter().GetName() : FName();
+}
+
+FText FNiagaraParameterBindingAdapterCustomization::GetTooltipText() const
+{
+	if ( const FNiagaraParameterBindingAdapter* ParameterBinding = GetParameterBindingAdapter() )
+	{
+		TArray<FText> TooltipTexts;
+		if (ParameterBinding->IsSetToConstant())
+		{
+			const FNiagaraTypeDefinition ConstantType = ParameterBinding->GetConstantTypeDef();
+			TooltipTexts.Add(
+				FText::Format(
+					LOCTEXT("BindingTooltip_BoundConstant", "Bound to constant value \"{0}\" type \"{1}\""),
+					FText::FromString(ConstantType.ToString(ParameterBinding->GetConstantValue().GetData())),
+					FText::FromString(ConstantType.GetName())
+				)
+			);
+		}
+		else
+		{
+			const FNiagaraVariableBase& BoundParameter = ParameterBinding->GetBoundParameter();
+			if (BoundParameter.GetName().IsNone() == false)
+			{
+				TooltipTexts.Add(
+					FText::Format(LOCTEXT("BindingTooltip_BoundParameter", "Bound to parameter \"{0}\" type \"{1}\""),
+						FText::FromName(BoundParameter.GetName()),
+						FText::FromString(BoundParameter.GetType().GetName())
+					)
+				);
+			}
+			else
+			{
+				TooltipTexts.Add(LOCTEXT("BindingTooltip_UnboundParameter", "Not bound to any parameter."));
+			}
+		}
+
+		TooltipTexts.Add(NiagaraParameterBindingCustomizationBasePrivate::GetAllowedBindingTypesText(ParameterBinding));
+		return FText::Join(LOCTEXT("NewLine", "\n"), TooltipTexts);
+	}
+	return LOCTEXT("BindingTooltip_Invalid", "Binding is invalid");
+}
+
+void FNiagaraParameterBindingAdapterCustomization::OnConstantValueChanged() const
+{
+	const FNiagaraParameterBindingAdapter* ParameterBinding = GetParameterBindingAdapter();
+	if (ParameterBinding == nullptr || !PropertyHandle.IsValid())
+	{
+		return;
+	}
+	
+	FScopedTransaction Transaction(LOCTEXT("ChangeParameterValue", " Change Parameter Value"));
+	
+	TArray<UObject*> Objects;
+	PropertyHandle->GetOuterObjects(Objects);
+	for (UObject* Obj : Objects)
+	{
+		Obj->Modify();
+	}
+	PropertyHandle->NotifyPreChange();
+	
+	ConstantValueParameterEditor->UpdateStructFromInternalValue(ConstantValueStructOnScope.ToSharedRef());
+	ParameterBinding->SetConstantValue(MakeArrayView(ConstantValueStructOnScope->GetStructMemory(), ConstantValueStructOnScope->GetStruct()->GetStructureSize()));
+	
+	PropertyHandle->NotifyPostChange(EPropertyChangeType::ValueSet);
+	PropertyHandle->NotifyFinishedChangingProperties();
+}
+
+TSharedRef<SWidget> FNiagaraParameterBindingAdapterCustomization::OnGetMenuContent() const
+{
+	return SNew(SBorder)
+		.BorderImage(FAppStyle::GetBrush("Menu.Background"))
+		.Padding(5)
+		[
+			SNew(SBox)
+			[
+				SNew(SGraphActionMenu)
+				.OnActionSelected(this, &FNiagaraParameterBindingAdapterCustomization::OnActionSelected)
+				.OnCreateWidgetForAction(this, &FNiagaraParameterBindingAdapterCustomization::OnCreateWidgetForAction)
+				.OnCollectAllActions(this, &FNiagaraParameterBindingAdapterCustomization::CollectAllActions)
+				.AutoExpandActionMenu(false)
+				.ShowFilterTextBox(true)
+			]
+		];
+}
+
+void FNiagaraParameterBindingAdapterCustomization::CollectAllActions(FGraphActionListBuilderBase& OutAllActions) const
+{
+	using namespace NiagaraParameterBindingCustomizationBasePrivate;
+
+	// Add default action(s)
+	const FNiagaraParameterBindingAdapter* ParameterBinding = GetParameterBindingAdapter();
+	if (ParameterBinding && ParameterBinding->AllowConstantValue())
+	{
+		OutAllActions.AddAction(FBindingAction::MakeConstant(ParameterBinding->GetConstantTypeDef(), ParameterBinding->GetConstantValue()));
+	}
+	else
+	{
+		OutAllActions.AddAction(FBindingAction::MakeNone());
+	}
+
+	// Add binding(s)
+	UNiagaraSystem* NiagaraSystem = SystemWeakPtr.Get();
+	UNiagaraEmitter* NiagaraEmitter = EmitterWeakPtr.Get();
+	if (ParameterBinding == nullptr || NiagaraSystem == nullptr)
+	{
+		return;
+	}
+
+	TArray<FNiagaraParameterMapHistory, TInlineAllocator<2>> Histories;
+	if (UNiagaraScriptSource* Source = Cast<UNiagaraScriptSource>(NiagaraSystem->GetSystemUpdateScript()->GetLatestSource()))
+	{
+		Histories.Append(UNiagaraNodeParameterMapBase::GetParameterMaps(Source->NodeGraph));
+	}
+
+	if (NiagaraEmitter)
+	{
+		// This is lame but currently the only way to get the version data
+		if (UObject* Owner = OwnerWeakPtr.Get())
+		{
+			if (UNiagaraSimulationStageBase* SimulationStage = Cast<UNiagaraSimulationStageBase>(Owner))
+			{
+				if (UNiagaraScriptSource* Source = Cast<UNiagaraScriptSource>(SimulationStage->GetEmitterData()->GraphSource))
+				{
+					Histories.Append(UNiagaraNodeParameterMapBase::GetParameterMaps(Source->NodeGraph));
+				}
+			}
+			else if (UNiagaraRendererProperties* RendererProperties = Cast<UNiagaraRendererProperties>(Owner))
+			{
+				if (UNiagaraScriptSource* Source = Cast<UNiagaraScriptSource>(RendererProperties->GetEmitterData()->GraphSource))
+				{
+					Histories.Append(UNiagaraNodeParameterMapBase::GetParameterMaps(Source->NodeGraph));
+				}
+			}
+		}
+	}
+
+	TArray<FNiagaraVariable> VisitedVariables;
+	TArray<TPair<FName, FInstancedStruct>> AllBindings;
+
+	const FStringView OwnerEmitterName = NiagaraEmitter ? NiagaraEmitter->GetUniqueEmitterName() : FStringView();
+	for (const FNiagaraParameterMapHistory& History : Histories)
+	{
+		for (const FNiagaraVariable& ResolvedVariable : History.Variables)
+		{
+			FNiagaraVariable AliasedVariable = ResolvedVariable;
+			AliasedVariable.ReplaceRootNamespace(OwnerEmitterName, FNiagaraConstants::EmitterNamespaceString);
+
+			if (VisitedVariables.Contains(AliasedVariable))
+			{
+				continue;
+			}
+			VisitedVariables.Add(AliasedVariable);
+
+			// Is this type of variable allowed?
+			if (CanBindToType(ParameterBinding, AliasedVariable.GetType()) == false )
+			{
+				continue;
+			}
+
+			if (CanBindToNamespace(ParameterBinding, AliasedVariable) == false)
+			{
+				continue;
+			}
+
+			// Collect bindings for this variable
+			ParameterBinding->CollectBindings(AliasedVariable, ResolvedVariable, AllBindings);
+		}
+	}
+
+	for (const FNiagaraVariableBase& Variable : NiagaraSystem->GetExposedParameters().ReadParameterVariables())
+	{
+		if (VisitedVariables.Contains(Variable))
+		{
+			continue;
+		}
+		VisitedVariables.Add(Variable);
+
+		ParameterBinding->CollectBindings(Variable, Variable, AllBindings);
+	}
+
+
+	for (const TPair<FName, FInstancedStruct>& Binding : AllBindings)
+	{
+		OutAllActions.AddAction(FBindingAction::MakeBinding(Binding.Key, Binding.Value));
+	}
+}
+
+TSharedRef<SWidget> FNiagaraParameterBindingAdapterCustomization::OnCreateWidgetForAction(struct FCreateWidgetForActionData* const InCreateData) const
+{
+	using namespace NiagaraParameterBindingCustomizationBasePrivate;
+
+	const FBindingAction* BindingAction = static_cast<const FBindingAction*>(InCreateData->Action.Get());
+	check(BindingAction);
+
+	if (BindingAction->VariableName.IsNone() == false)
+	{
+		return
+			SNew(SVerticalBox)
+			+ SVerticalBox::Slot()
+			.AutoHeight()
+			[
+				SNew(SNiagaraParameterName)
+				.ParameterName(BindingAction->VariableName)
+				.IsReadOnly(true)
+				.ToolTipText(InCreateData->Action->GetTooltipDescription())
+			];
+	}
+	else
+	{
+		return
+			SNew(SVerticalBox)
+			+ SVerticalBox::Slot()
+			.AutoHeight()
+			[
+				SNew(STextBlock)
+				.Text(InCreateData->Action->GetMenuDescription())
+				.ToolTipText(InCreateData->Action->GetTooltipDescription())
+			];
+	}
+}
+
+void FNiagaraParameterBindingAdapterCustomization::OnActionSelected(const TArray<TSharedPtr<FEdGraphSchemaAction>>& SelectedActions, ESelectInfo::Type InSelectionType) const
+{
+	using namespace NiagaraParameterBindingCustomizationBasePrivate;
+
+	FNiagaraParameterBindingAdapter* ParameterBinding = GetParameterBindingAdapter();
+	if (ParameterBinding == nullptr || !PropertyHandle.IsValid() || (InSelectionType != ESelectInfo::OnMouseClick && InSelectionType != ESelectInfo::OnKeyPress))
+	{
+		return;
+	}
+
+	for (const TSharedPtr<FEdGraphSchemaAction>& EdAction : SelectedActions)
+	{
+		if (EdAction.IsValid() == false)
+		{
+			continue;
+		}
+
+		FBindingAction* BindAction = static_cast<FBindingAction*>(EdAction.Get());
+		FScopedTransaction Transaction(FText::Format(LOCTEXT("ChangeParameterBinding", " Change Parameter Binding to \"{0}\" "), BindAction->GetMenuDescription()));
+		FSlateApplication::Get().DismissAllMenus();
+
+		TArray<UObject*> Objects;
+		PropertyHandle->GetOuterObjects(Objects);
+		for (UObject* Obj : Objects)
+		{
+			Obj->Modify();
+		}
+		PropertyHandle->NotifyPreChange();
+
+		ParameterBinding->SetBoundParameter(BindAction->ValueStruct);
+
+		PropertyHandle->NotifyPostChange(EPropertyChangeType::ValueSet);
+		PropertyHandle->NotifyFinishedChangingProperties();
+	}
+}
+
+EVisibility FNiagaraParameterBindingAdapterCustomization::IsResetToDefaultsVisible() const
+{
+	const FNiagaraParameterBindingAdapter* ParameterBinding = GetParameterBindingAdapter();
+	return ParameterBinding && ParameterBinding->IsSetToDefault() ? EVisibility::Hidden : EVisibility::Visible;
+}
+
+FReply FNiagaraParameterBindingAdapterCustomization::OnResetToDefaultsClicked()
+{
+	if (FNiagaraParameterBindingAdapter* ParameterBinding = GetParameterBindingAdapter())
+	{
+		ParameterBinding->SetToDefault();
+
+		if (ConstantValueStructOnScope.IsValid())
+		{
+			TConstArrayView<uint8> ConstantValue = ParameterBinding->GetConstantValue();
+			FMemory::Memcpy(ConstantValueStructOnScope->GetStructMemory(), ConstantValue.GetData(), ConstantValue.Num());
+			ConstantValueParameterEditor->UpdateInternalValueFromStruct(ConstantValueStructOnScope.ToSharedRef());
+		}
+
+		PropertyHandle->NotifyPostChange(EPropertyChangeType::ValueSet);
+		PropertyHandle->NotifyFinishedChangingProperties();
+	}
+
+	return FReply::Handled();
+}
+
+#undef LOCTEXT_NAMESPACE
diff --git a/Engine/Plugins/FX/Niagara/Source/NiagaraEditor/Internal/Customizations/NiagaraParameterBindingAdapterCustomization.h b/Engine/Plugins/FX/Niagara/Source/NiagaraEditor/Internal/Customizations/NiagaraParameterBindingAdapterCustomization.h
new file mode 100644
--- /dev/null
+++ b/Engine/Plugins/FX/Niagara/Source/NiagaraEditor/Internal/Customizations/NiagaraParameterBindingAdapterCustomization.h
@@ -0,0 +1,106 @@
+// Copyright Epic Games, Inc. All Rights Reserved.
+
+#pragma once
+
+#include "CoreMinimal.h"
+#include "EdGraph/EdGraphSchema.h"
+#include "IPropertyTypeCustomization.h"
+#include "Layout/Visibility.h"
+#include "UObject/StructOnScope.h"
+#include "StructUtils/InstancedStruct.h"
+
+#include "NiagaraCommon.h"
+
+class FDetailWidgetRow;
+class IDetailChildrenBuilder;
+struct FEdGraphSchemaAction;
+struct FGraphActionListBuilderBase;
+class IPropertyTypeCustomizationUtils;
+class SWidget;
+
+class UNiagaraEmitter;
+class UNiagaraSystem;
+class SNiagaraParameterEditor;
+
+class FNiagaraParameterBindingAdapter
+{
+public:
+	virtual ~FNiagaraParameterBindingAdapter() = default;
+
+	virtual bool IsSetToParameter() const = 0;	// !AllowConstantValue || ResolveParameter.GetName().IsNone() == false
+	virtual bool IsSetToConstant() const { return IsSetToParameter() == false; }
+
+	virtual bool AllowConstantValue() const = 0;
+	virtual FNiagaraTypeDefinition GetConstantTypeDef() const = 0;
+	virtual TConstArrayView<uint8> GetConstantValue() const = 0;
+	virtual void SetConstantValue(TConstArrayView<uint8> Memory) const = 0;
+
+	virtual const FNiagaraVariableBase& GetBoundParameter() const = 0;
+	virtual void SetBoundParameter(const FInstancedStruct& Parameter) = 0;
+
+	virtual bool IsSetToDefault() const = 0;
+	virtual void SetToDefault() = 0;
+
+	virtual void CollectBindings(const FNiagaraVariableBase& AliasedVariable, const FNiagaraVariableBase& ResolvedVariable, TArray<TPair<FName, FInstancedStruct>>& OutBindings) const = 0;
+
+	virtual bool AllowUserParameters() const { return true; }
+	virtual bool AllowSystemParameters() const { return true; }
+	virtual bool AllowEmitterParameters() const { return true; }
+	virtual bool AllowParticleParameters() const { return true; }
+	virtual bool AllowStaticVariables() const { return true; }
+
+	virtual TConstArrayView<UClass*> GetAllowedDataInterfaces() const { return {}; }
+	virtual TConstArrayView<UClass*> GetAllowedObjects() const { return {}; }
+	virtual TConstArrayView<UClass*> GetAllowedInterfaces() const { return {}; }
+	virtual TConstArrayView<FNiagaraTypeDefinition> GetAllowedTypeDefinitions() const { return {}; }
+};
+
+class FNiagaraParameterBindingAdapterCustomization : public IPropertyTypeCustomization
+{
+public:
+	using FAdapterCreationFunction = TFunction<TUniquePtr<FNiagaraParameterBindingAdapter>(uint8* ObjectAddress)>;
+
+	FNiagaraParameterBindingAdapterCustomization(FAdapterCreationFunction InAdapterCreator)
+		: AdapterCreator(InAdapterCreator)
+	{
+	}
+
+	/** IPropertyTypeCustomization interface begin */
+	NIAGARAEDITOR_API virtual void CustomizeHeader(TSharedRef<IPropertyHandle> StructPropertyHandle, FDetailWidgetRow& HeaderRow, IPropertyTypeCustomizationUtils& StructCustomizationUtils) override;
+	virtual void CustomizeChildren(TSharedRef<IPropertyHandle> StructPropertyHandle, IDetailChildrenBuilder& ChildBuilder, IPropertyTypeCustomizationUtils& StructCustomizationUtils) override {}
+	/** IPropertyTypeCustomization interface end */
+
+private:
+	bool IsValid() const;
+
+	FNiagaraParameterBindingAdapter* GetParameterBindingAdapter() const;
+
+	bool IsBindingEnabled() const;
+	bool IsConstantEnabled() const;
+	EVisibility IsBindingVisibile() const;
+	EVisibility IsConstantVisibile() const;
+
+	FName GetBoundParameterName() const;
+	FText GetTooltipText() const;
+
+	void OnConstantValueChanged() const;
+
+	TSharedRef<SWidget> OnGetMenuContent() const;
+	void CollectAllActions(FGraphActionListBuilderBase& OutAllActions) const;
+	TSharedRef<SWidget> OnCreateWidgetForAction(struct FCreateWidgetForActionData* const InCreateData) const;
+	void OnActionSelected(const TArray<TSharedPtr<FEdGraphSchemaAction>>& SelectedActions, ESelectInfo::Type InSelectionType) const;
+
+	EVisibility IsResetToDefaultsVisible() const;
+	FReply OnResetToDefaultsClicked();
+
+protected:
+	FAdapterCreationFunction					AdapterCreator;
+	TUniquePtr<FNiagaraParameterBindingAdapter>	ParameterBindingAdapter;
+	TSharedPtr<IPropertyHandle>					PropertyHandle;
+	TWeakObjectPtr<UObject>						OwnerWeakPtr;
+	TWeakObjectPtr<UNiagaraEmitter>				EmitterWeakPtr;
+	TWeakObjectPtr<UNiagaraSystem>				SystemWeakPtr;
+
+	TSharedPtr<SNiagaraParameterEditor>			ConstantValueParameterEditor;
+	TSharedPtr<FStructOnScope>					ConstantValueStructOnScope;
+};
diff --git a/Engine/Plugins/FX/Niagara/Source/NiagaraEditor/Private/Customizations/NiagaraParameterBindingCustomization.cpp b/Engine/Plugins/FX/Niagara/Source/NiagaraEditor/Private/Customizations/NiagaraParameterBindingCustomization.cpp
--- a/Engine/Plugins/FX/Niagara/Source/NiagaraEditor/Private/Customizations/NiagaraParameterBindingCustomization.cpp
+++ b/Engine/Plugins/FX/Niagara/Source/NiagaraEditor/Private/Customizations/NiagaraParameterBindingCustomization.cpp
@@ -2,632 +2,105 @@
 
 #include "NiagaraParameterBindingCustomization.h"
 
-#include "CoreMinimal.h"
-#include "DetailLayoutBuilder.h"
-#include "DetailWidgetRow.h"
-#include "IDetailChildrenBuilder.h"
-#include "IDetailGroup.h"
-#include "IPropertyUtilities.h"
-#include "PropertyHandle.h"
-#include "ScopedTransaction.h"
-#include "SGraphActionMenu.h"
-#include "Framework/Application/SlateApplication.h"
-#include "Widgets/Input/SComboButton.h"
-
-#include "INiagaraEditorTypeUtilities.h"
-#include "NiagaraEditorModule.h"
-#include "NiagaraEmitter.h"
-#include "NiagaraNodeParameterMapBase.h"
-#include "NiagaraParameterMapHistory.h"
-#include "NiagaraRendererProperties.h"
-#include "NiagaraScriptSource.h"
-#include "NiagaraSimulationStageBase.h"
-#include "NiagaraSystem.h"
-#include "NiagaraVariableMetaData.h"
+#include "Customizations/NiagaraParameterBindingAdapterCustomization.h"
 #include "NiagaraParameterBinding.h"
-#include "SNiagaraParameterEditor.h"
-#include "Customizations/NiagaraTypeCustomizations.h"
-#include "Widgets/Input/SButton.h"
-#include "Widgets/SNiagaraParameterName.h"
-
-#include UE_INLINE_GENERATED_CPP_BY_NAME(NiagaraParameterBindingCustomization)
-
-#define LOCTEXT_NAMESPACE "NiagaraParameterBindingCustomization"
 
-namespace NiagaraParameterBindingCustomizationPrivate
+namespace NiagaraParameterBindingAdapterCustomizationPrivate
 {
-	FText JoinTypeArray(FText Delim, TConstArrayView<UClass*> TypeArray)
+	class FAdapter : public FNiagaraParameterBindingAdapter
 	{
-		TArray<FText> TypeNamesText;
-		for (UClass* Class : TypeArray)
+	public:
+		explicit FAdapter(uint8* ObjectAddress)
 		{
-			TypeNamesText.Add(Class->GetDisplayNameText());
+			check(ObjectAddress != nullptr);
+			ParameterBinding = reinterpret_cast<FNiagaraParameterBinding*>(ObjectAddress);
 		}
-		return FText::Join(Delim, TypeNamesText);
-	}
 
-	FText JoinTypeArray(FText Delim, TConstArrayView<FNiagaraTypeDefinition> TypeArray)
-	{
-		TArray<FText> TypeNamesText;
-		for (const FNiagaraTypeDefinition& TypeDef : TypeArray)
+		virtual bool IsSetToParameter() const override
 		{
-			TypeNamesText.Add(TypeDef.GetNameText());
+			return !ParameterBinding->HasDefaultValueEditorOnly() || ParameterBinding->ResolvedParameter.GetName().IsNone() == false;
 		}
-		return FText::Join(Delim, TypeNamesText);
-	}
 
-	template<typename TArrayType>
-	void BuildTypesText(FText& InOutString, FText TypesDescription, TArrayType TypeArray)
-	{
-		if (TypeArray.Num() == 0)
+		virtual bool AllowConstantValue() const override
 		{
-			return;
+			return ParameterBinding->HasDefaultValueEditorOnly();
 		}
 
-		InOutString = FText::Join(
-			FText::GetEmpty(),
-			TArray<FText>(
-				{
-					InOutString,
-					LOCTEXT("NewLine", "\n"),
-					TypesDescription,
-					JoinTypeArray(LOCTEXT("CommaSpace", ", "), TypeArray)
-				}
-			)
-		);
-	}
-
-	FText GetAllowedBindingTypesText(FNiagaraParameterBinding* ParameterBinding)
-	{
-		FText BoundTypes;
-		BuildTypesText(
-			BoundTypes,
-			LOCTEXT("BindingTooltip_DIFormat", "Data Interface Types: "),
-			ParameterBinding->GetAllowedDataInterfaces()
-		);
-		BuildTypesText(
-			BoundTypes,
-			LOCTEXT("BindingTooltip_ObjectFormat", "Object Types: "),
-			ParameterBinding->GetAllowedObjects()
-		);
-		BuildTypesText(
-			BoundTypes,
-			LOCTEXT("BindingTooltip_InterfacesFormat", "Interface Types: "),
-			ParameterBinding->GetAllowedInterfaces()
-		);
-		BuildTypesText(
-			BoundTypes,
-			LOCTEXT("BindingTooltip_AttributesFormat", "Attribute Types: "),
-			ParameterBinding->GetAllowedTypeDefinitions()
-		);
-
-		return BoundTypes;
-	}
-}
-
-FNiagaraParameterBindingAction::FNiagaraParameterBindingAction()
-	: FEdGraphSchemaAction(FText(), FText(), FText(), 0)
-{
-}
-
-FNiagaraParameterBindingAction::FNiagaraParameterBindingAction(FText InName, FText InTooltip)
-	: FEdGraphSchemaAction(FText(), InName, InTooltip, 0)
-{
-}
-
-TSharedRef<FNiagaraParameterBindingAction> FNiagaraParameterBindingAction::MakeNone()
-{
-	TSharedRef<FNiagaraParameterBindingAction> Action = MakeShared<FNiagaraParameterBindingAction>(
-		LOCTEXT("None", "None"),
-		LOCTEXT("NoneTooltip", "Not bound to any parameter.")
-	);
-	Action->Grouping = 1;
-	return Action;
-}
-
-TSharedRef<FNiagaraParameterBindingAction> FNiagaraParameterBindingAction::MakeConstant(const FNiagaraTypeDefinition& TypeDef, const uint8* ValueData)
-{
-	TSharedRef<FNiagaraParameterBindingAction> Action = MakeShared<FNiagaraParameterBindingAction>(
-		FText::Format(LOCTEXT("ConstantFormat", "Constant Value \"{0}\""), FText::FromString(TypeDef.ToString(ValueData))),
-		LOCTEXT("ConstantTooltip", "Use user defined constant value, not bound to any parameter.")
-	);
-	Action->Grouping = 1;
-	return Action;
-}
-
-TSharedRef<FNiagaraParameterBindingAction> FNiagaraParameterBindingAction::MakeParameter(FNiagaraVariableBase InAliasedParameter, FNiagaraVariableBase InResolvedParameter)
-{
-	TSharedRef<FNiagaraParameterBindingAction> Action = MakeShared<FNiagaraParameterBindingAction>(
-		FText::FromName(InAliasedParameter.GetName()),
-		FText::Format(LOCTEXT("VariableBindActionTooltip", "Bind to variable \"{0}\""), FText::FromName(InAliasedParameter.GetName()))
-	);
-	Action->Grouping = 0;
-	Action->bIsValidParameter = true;
-	Action->AliasedParameter = InAliasedParameter;
-	Action->ResolvedParameter = InResolvedParameter;
-	return Action;
-}
-
-TSharedRef<IPropertyTypeCustomization> FNiagaraParameterBindingCustomization::MakeInstance()
-{
-	return MakeShared<FNiagaraParameterBindingCustomization>();
-}
-
-void FNiagaraParameterBindingCustomization::CustomizeHeader(TSharedRef<IPropertyHandle> InPropertyHandle, FDetailWidgetRow& HeaderRow, IPropertyTypeCustomizationUtils& StructCustomizationUtils)
-{
-	PropertyHandle = InPropertyHandle;
-	EmitterWeakPtr.Reset();
-	SystemWeakPtr.Reset();
-
-	// We only handle single object
-	UObject* OwnerObject = nullptr;
-	{
-		TArray<UObject*> OuterObjects;
-		PropertyHandle->GetOuterObjects(OuterObjects);
-		OwnerObject = OuterObjects.Num() == 1 ? OuterObjects[0] : nullptr;
-		if (OwnerObject == nullptr)
+		virtual FNiagaraTypeDefinition GetConstantTypeDef() const override
 		{
-			HeaderRow
-			.NameContent()
-			[
-				PropertyHandle->CreatePropertyNameWidget()
-			]
-			.ValueContent()
-			.MaxDesiredWidth(200.f)
-			[
-				SNew(STextBlock)
-				.Text(FText::FromString(FName::NameToDisplayString(CastField<FStructProperty>(PropertyHandle->GetProperty())->Struct->GetName(), false)))
-				.Font(IDetailLayoutBuilder::GetDetailFont())
-			];
-
-			return;
+			return ParameterBinding->GetDefaultResolvedParameter().GetType();
 		}
-	}
-
-	InPropertyHandle->SetOnPropertyResetToDefault(FSimpleDelegate::CreateLambda([]() { }));
-	PropertyHandle->MarkResetToDefaultCustomized(true);
 
-	OwnerWeakPtr = OwnerObject;
-	EmitterWeakPtr = OwnerObject->GetTypedOuter<UNiagaraEmitter>();
-	SystemWeakPtr = OwnerObject->GetTypedOuter<UNiagaraSystem>();
-
-	TSharedPtr<SHorizontalBox> ParameterPanel = SNew(SHorizontalBox);
-
-	FNiagaraParameterBinding* ParameterBinding = GetParameterBinding();
-	if (ParameterBinding && ParameterBinding->HasDefaultValueEditorOnly())
-	{
-		const FNiagaraVariableBase& DefaultParameter = ParameterBinding->GetDefaultAliasedParameter();
-		TConstArrayView<uint8> DefaultValue = ParameterBinding->GetDefaultValueEditorOnly();
-		DefaultValueStructOnScope = MakeShared<FStructOnScope>(DefaultParameter.GetType().GetStruct());
-		FMemory::Memcpy(DefaultValueStructOnScope->GetStructMemory(), DefaultValue.GetData(), DefaultValue.Num());
-
-		FNiagaraEditorModule& NiagaraEditorModule = FModuleManager::GetModuleChecked<FNiagaraEditorModule>("NiagaraEditor");
-		TSharedPtr<INiagaraEditorTypeUtilities, ESPMode::ThreadSafe> TypeEditorUtilities = NiagaraEditorModule.GetTypeUtilities(DefaultParameter.GetType());
-		if (TypeEditorUtilities.IsValid())
+		virtual TConstArrayView<uint8> GetConstantValue() const override
 		{
-			FNiagaraInputParameterCustomization CustomizationOptions;
-			CustomizationOptions.bBroadcastValueChangesOnCommitOnly = true; // each broadcast usually forces a recompile, so we only want to do it on commits
-			DefaultValueParameterEditor = TypeEditorUtilities->CreateParameterEditor(ParameterBinding->ResolvedParameter.GetType(), EUnit::Unspecified, CustomizationOptions);
-			DefaultValueParameterEditor->UpdateInternalValueFromStruct(DefaultValueStructOnScope.ToSharedRef());
-			DefaultValueParameterEditor->SetOnValueChanged(SNiagaraParameterEditor::FOnValueChange::CreateSP(this, &FNiagaraParameterBindingCustomization::OnValueChanged));
-
-			ParameterPanel->AddSlot()
-			.AutoWidth()
-			.Padding(0, 0, 0, 3)
-			[
-				SNew(SBox)
-				.HAlign(DefaultValueParameterEditor->GetHorizontalAlignment())
-				.VAlign(DefaultValueParameterEditor->GetVerticalAlignment())
-				.IsEnabled(this, &FNiagaraParameterBindingCustomization::IsConstantEnabled)
-				.Visibility(this, &FNiagaraParameterBindingCustomization::IsConstantVisibile)
-				[
-					DefaultValueParameterEditor.ToSharedRef()
-				]
-			];
+			return ParameterBinding->GetDefaultValueEditorOnly();
 		}
-	}
-
-	ParameterPanel->AddSlot()
-	.AutoWidth()
-	[
-		SNew(SNiagaraParameterName)
-		.ParameterName(this, &FNiagaraParameterBindingCustomization::GetVariableName)
-		.IsReadOnly(true)
-		.IsEnabled(this, &FNiagaraParameterBindingCustomization::IsBindingEnabled)
-		.Visibility(this, &FNiagaraParameterBindingCustomization::IsBindingVisibile)
-		.ToolTipText(this, &FNiagaraParameterBindingCustomization::GetTooltipText)
-	];
-
-	HeaderRow
-	.NameContent()
-	[
-		PropertyHandle->CreatePropertyNameWidget()
-	]
-	.ValueContent()
-	[
-		SNew(SHorizontalBox)
-		+ SHorizontalBox::Slot()
-		[
-			ParameterPanel.ToSharedRef()
-		]
-		+ SHorizontalBox::Slot()
-		.FillWidth(1.0f)
-		.HAlign(HAlign_Right)
-		.Padding(3, 0, 0, 0)
-		[
-			SNew(SComboButton)					
-			.ButtonStyle(FAppStyle::Get(), "HoverHintOnly")
-			.ForegroundColor(FSlateColor::UseForeground())
-			.OnGetMenuContent(this, &FNiagaraParameterBindingCustomization::OnGetMenuContent)
-			.ContentPadding(2)
-			.MenuPlacement(MenuPlacement_BelowRightAnchor)
-			.ToolTipText(this, &FNiagaraParameterBindingCustomization::GetTooltipText)
-		]
-		+ SHorizontalBox::Slot()
-		.AutoWidth()
-		.VAlign(VAlign_Center)
-		.Padding(2, 1)
-		[
-			SNew(SButton)
-			.IsFocusable(false)
-			.ToolTipText(LOCTEXT("ResetToDefaultToolTip", "Reset to Default"))
-			.ButtonStyle(FAppStyle::Get(), "NoBorder")
-			.ContentPadding(0)
-			.Visibility(this, &FNiagaraParameterBindingCustomization::IsResetToDefaultsVisible)
-			.OnClicked(this, &FNiagaraParameterBindingCustomization::OnResetToDefaultsClicked)
-			.Content()
-			[
-				SNew(SImage)
-				.Image(FAppStyle::GetBrush("PropertyWindow.DiffersFromDefault"))
-				.ColorAndOpacity(FSlateColor::UseForeground())
-			]
-		]
-	];
-}
-
-bool FNiagaraParameterBindingCustomization::IsValid() const
-{
-	return OwnerWeakPtr.IsValid();
-}
-
-FNiagaraParameterBinding* FNiagaraParameterBindingCustomization::GetParameterBinding() const
-{
-	UObject* OwnerObject = OwnerWeakPtr.Get();
-	return OwnerObject ? (FNiagaraParameterBinding*)PropertyHandle->GetValueBaseAddress((uint8*)OwnerObject) : nullptr;
-}
-
-FNiagaraParameterBinding* FNiagaraParameterBindingCustomization::GetDefaultParameterBinding() const
-{
-	UObject* OwnerObject = OwnerWeakPtr.Get();
-	return OwnerObject ? (FNiagaraParameterBinding*)PropertyHandle->GetValueBaseAddress((uint8*)OwnerObject->GetClass()->GetDefaultObject()) : nullptr;
-}
-
-TSharedRef<SWidget> FNiagaraParameterBindingCustomization::OnGetMenuContent() const
-{
-	FGraphActionMenuBuilder MenuBuilder;
-
-	return SNew(SBorder)
-		.BorderImage(FAppStyle::GetBrush("Menu.Background"))
-		.Padding(5)
-		[
-			SNew(SBox)
-			[
-				SNew(SGraphActionMenu)
-				.OnActionSelected(this, &FNiagaraParameterBindingCustomization::OnActionSelected)
-				.OnCreateWidgetForAction(this, &FNiagaraParameterBindingCustomization::OnCreateWidgetForAction)
-				.OnCollectAllActions(this, &FNiagaraParameterBindingCustomization::CollectAllActions)
-				.AutoExpandActionMenu(false)
-				.ShowFilterTextBox(true)
-			]
-		];
-}
 
-bool FNiagaraParameterBindingCustomization::ForEachBindableVariable(TFunction<bool(FNiagaraVariableBase, FNiagaraVariableBase)> Delegate) const
-{
-	// Add bindings from system
-	FNiagaraParameterBinding* ParameterBinding = GetParameterBinding();
-	UNiagaraSystem* NiagaraSystem = SystemWeakPtr.Get();
-	UNiagaraEmitter* NiagaraEmitter = EmitterWeakPtr.Get();
-	if (ParameterBinding == nullptr || NiagaraSystem == nullptr)
-	{
-		return true;
-	}
-
-	TArray<FNiagaraParameterMapHistory, TInlineAllocator<2>> Histories;
-	if (UNiagaraScriptSource* Source = Cast<UNiagaraScriptSource>(NiagaraSystem->GetSystemUpdateScript()->GetLatestSource()))
-	{
-		Histories.Append(UNiagaraNodeParameterMapBase::GetParameterMaps(Source->NodeGraph));
-	}
-
-	if (NiagaraEmitter)
-	{
-		// This is lame but currently the only way to get the version data
-		if (UObject* Owner = OwnerWeakPtr.Get())
+		virtual void SetConstantValue(TConstArrayView<uint8> Memory) const override
 		{
-			if (UNiagaraSimulationStageBase* SimulationStage = Cast<UNiagaraSimulationStageBase>(Owner))
-			{
-				if (UNiagaraScriptSource* Source = Cast<UNiagaraScriptSource>(SimulationStage->GetEmitterData()->GraphSource))
-				{
-					Histories.Append(UNiagaraNodeParameterMapBase::GetParameterMaps(Source->NodeGraph));
-				}
-			}
-			else if (UNiagaraRendererProperties* RendererProperties = Cast<UNiagaraRendererProperties>(Owner))
-			{
-				if (UNiagaraScriptSource* Source = Cast<UNiagaraScriptSource>(RendererProperties->GetEmitterData()->GraphSource))
-				{
-					Histories.Append(UNiagaraNodeParameterMapBase::GetParameterMaps(Source->NodeGraph));
-				}
-			}
+			ParameterBinding->SetDefaultValueEditorOnly(Memory);
 		}
-	}
-
-	TArray<FNiagaraVariable> VisitedVariables;
 
-	const FStringView NiagaraEmitterName = NiagaraEmitter ? NiagaraEmitter->GetUniqueEmitterName() : FStringView();
-	for (const FNiagaraParameterMapHistory& History : Histories)
-	{
-		for (const FNiagaraVariable& Variable : History.Variables)
+		virtual const FNiagaraVariableBase& GetBoundParameter() const override
 		{
-			FNiagaraVariableBase AliasedParameter = Variable;
-			if (ParameterBinding->CanBindTo(Variable, AliasedParameter, NiagaraEmitterName) && !VisitedVariables.Contains(Variable))
-			{
-				if (Delegate(AliasedParameter, Variable) == false)
-				{
-					return false;
-				}
-				VisitedVariables.Add(Variable);
-			}
+			return ParameterBinding->AliasedParameter;
 		}
-	}
 
-	if (ParameterBinding->AllowUserParameters())
-	{
-		for (const FNiagaraVariableBase& Variable : NiagaraSystem->GetExposedParameters().ReadParameterVariables())
+		virtual void SetBoundParameter(const FInstancedStruct& Parameter) override
 		{
-			if (ParameterBinding->CanBindTo(Variable.GetType()) && !VisitedVariables.Contains(Variable))
+			if ( Parameter.IsValid() )
 			{
-				if (Delegate(Variable, Variable) == false)
-				{
-					return false;
-				}
-				VisitedVariables.Add(Variable);
-			}
-		}
-	}
-	return true;
-}
-
-
-void FNiagaraParameterBindingCustomization::CollectAllActions(FGraphActionListBuilderBase& OutAllActions) const
-{
-	// Add default action
-	bool bHasConstantValue = false;
-	if (FNiagaraParameterBinding* ParameterBinding = GetParameterBinding())
-	{
-		UNiagaraSystem* NiagaraSystem = SystemWeakPtr.Get();
-		UNiagaraEmitter* NiagaraEmitter = EmitterWeakPtr.Get();
-
-		if (ParameterBinding->HasDefaultValueEditorOnly())
-		{
-			OutAllActions.AddAction(FNiagaraParameterBindingAction::MakeConstant(ParameterBinding->GetDefaultAliasedParameter().GetType(), ParameterBinding->GetDefaultValueEditorOnly().GetData()));
-		}
-		else
-		{
-			OutAllActions.AddAction(FNiagaraParameterBindingAction::MakeNone());
-		}
-	}
-	else
-	{
-		OutAllActions.AddAction(FNiagaraParameterBindingAction::MakeNone());
-	}
-
-	// Add each variable
-	ForEachBindableVariable(
-		[&OutAllActions](FNiagaraVariableBase AliasedVariable, FNiagaraVariableBase ResolvedVariable)
-		{
-			OutAllActions.AddAction(FNiagaraParameterBindingAction::MakeParameter(AliasedVariable, ResolvedVariable));
-			return true;
-		}
-	);
-}
-
-TSharedRef<SWidget> FNiagaraParameterBindingCustomization::OnCreateWidgetForAction(struct FCreateWidgetForActionData* const InCreateData) const
-{
-	const FNiagaraParameterBindingAction* BindingAction = static_cast<const FNiagaraParameterBindingAction*>(InCreateData->Action.Get());
-	check(BindingAction);
-
-	if (BindingAction->bIsValidParameter)
-	{
-		return
-			SNew(SVerticalBox)
-			+ SVerticalBox::Slot()
-			.AutoHeight()
-			[
-				SNew(SNiagaraParameterName)
-				.ParameterName(BindingAction->AliasedParameter.GetName())
-				.IsReadOnly(true)
-				.ToolTipText(InCreateData->Action->GetTooltipDescription())
-			];
-	}
-	else
-	{
-		return
-			SNew(SVerticalBox)
-			+ SVerticalBox::Slot()
-			.AutoHeight()
-			[
-				SNew(STextBlock)
-				.Text(InCreateData->Action->GetMenuDescription())
-				.ToolTipText(InCreateData->Action->GetTooltipDescription())
-			];
-	}
-}
-
-void FNiagaraParameterBindingCustomization::OnActionSelected(const TArray<TSharedPtr<FEdGraphSchemaAction>>& SelectedActions, ESelectInfo::Type InSelectionType) const
-{
-	FNiagaraParameterBinding* ParameterBinding = GetParameterBinding();
-	if (ParameterBinding == nullptr || !PropertyHandle.IsValid() || (InSelectionType != ESelectInfo::OnMouseClick && InSelectionType != ESelectInfo::OnKeyPress))
-	{
-		return;
-	}
-
-	for (const TSharedPtr<FEdGraphSchemaAction>& EdAction : SelectedActions)
-	{
-		if (EdAction.IsValid())
-		{
-			FNiagaraParameterBindingAction* BindAction = static_cast<FNiagaraParameterBindingAction*>(EdAction.Get());
-			FScopedTransaction Transaction(FText::Format(LOCTEXT("ChangeParameterBinding", " Change Parameter Binding to \"{0}\" "), FText::FromName(BindAction->ResolvedParameter.GetName())));
-			FSlateApplication::Get().DismissAllMenus();
-
-			TArray<UObject*> Objects;
-			PropertyHandle->GetOuterObjects(Objects);
-			for (UObject* Obj : Objects)
-			{
-				Obj->Modify();
-			}
-			PropertyHandle->NotifyPreChange();
-
-			if (BindAction->bIsValidParameter)
-			{
-				ParameterBinding->AliasedParameter = BindAction->AliasedParameter;
-				ParameterBinding->ResolvedParameter = BindAction->ResolvedParameter;
+				const FNiagaraParameterBinding& BindingData = Parameter.Get<FNiagaraParameterBinding>();
+				ParameterBinding->AliasedParameter = BindingData.AliasedParameter;
+				ParameterBinding->ResolvedParameter = BindingData.ResolvedParameter;
 			}
 			else
 			{
 				ParameterBinding->AliasedParameter = FNiagaraVariableBase();
 				ParameterBinding->ResolvedParameter = FNiagaraVariableBase();
 			}
-
-			PropertyHandle->NotifyPostChange(EPropertyChangeType::ValueSet);
-			PropertyHandle->NotifyFinishedChangingProperties();
 		}
-	}
-}
-
-void FNiagaraParameterBindingCustomization::OnValueChanged() const
-{
-	FNiagaraParameterBinding* ParameterBinding = GetParameterBinding();
-	if (ParameterBinding == nullptr || !PropertyHandle.IsValid())
-	{
-		return;
-	}
-
-	FScopedTransaction Transaction(LOCTEXT("ChangeParameterValue", " Change Parameter Value"));
-
-	TArray<UObject*> Objects;
-	PropertyHandle->GetOuterObjects(Objects);
-	for (UObject* Obj : Objects)
-	{
-		Obj->Modify();
-	}
-	PropertyHandle->NotifyPreChange();
-
-	DefaultValueParameterEditor->UpdateStructFromInternalValue(DefaultValueStructOnScope.ToSharedRef());
-	ParameterBinding->SetDefaultValueEditorOnly(DefaultValueStructOnScope->GetStructMemory());
-
-	PropertyHandle->NotifyPostChange(EPropertyChangeType::ValueSet);
-	PropertyHandle->NotifyFinishedChangingProperties();
-}
-
-EVisibility FNiagaraParameterBindingCustomization::IsBindingVisibile() const
-{
-	FNiagaraParameterBinding* ParameterBinding = GetParameterBinding();
-	if (ParameterBinding && ParameterBinding->HasDefaultValueEditorOnly())
-	{
-		return ParameterBinding->ResolvedParameter.GetName().IsNone() ? EVisibility::Collapsed : EVisibility::Visible;
-	}
-	return EVisibility::Visible;
-}
-
-EVisibility FNiagaraParameterBindingCustomization::IsConstantVisibile() const
-{
-	return IsBindingVisibile() == EVisibility::Visible ? EVisibility::Collapsed : EVisibility::Visible;
-}
-
-bool FNiagaraParameterBindingCustomization::IsBindingEnabled() const
-{
-	FNiagaraParameterBinding* ParameterBinding = GetParameterBinding();
-	if (ParameterBinding && ParameterBinding->HasDefaultValueEditorOnly())
-	{
-		return ParameterBinding->ResolvedParameter.GetName().IsNone() == false;
-	}
-	return true;
-}
 
-bool FNiagaraParameterBindingCustomization::IsConstantEnabled() const
-{
-	return !IsBindingEnabled();
-}
+		virtual bool IsSetToDefault() const override
+		{
+			return ParameterBinding->IsSetToDefault();
+		}
 
-FText FNiagaraParameterBindingCustomization::GetTooltipText() const
-{
-	if ( FNiagaraParameterBinding* ParameterBinding = GetParameterBinding() )
-	{
-		TArray<FText> TooltipTexts;
-		if (ParameterBinding->AliasedParameter.GetName().IsNone())
+		virtual void SetToDefault() override
 		{
-			if (ParameterBinding->HasDefaultValueEditorOnly())
-			{
-				const FNiagaraTypeDefinition ConstantType = ParameterBinding->GetDefaultAliasedParameter().GetType();
-				TooltipTexts.Add(
-					FText::Format(
-						LOCTEXT("BindingTooltip_BoundConstant", "Bound to constant value \"{0}\" type \"{1}\""),
-						FText::FromString(ConstantType.ToString(ParameterBinding->GetDefaultValueEditorOnly().GetData())),
-						FText::FromString(ConstantType.GetName())
-					)
-				);
-			}
-			else
-			{
-				TooltipTexts.Add(LOCTEXT("BindingTooltip_UnboundParameter", "Not bound to any parameter."));
-			}
+			ParameterBinding->SetToDefault();
 		}
-		else
+
+		virtual void CollectBindings(const FNiagaraVariableBase& AliasedVariable, const FNiagaraVariableBase& ResolvedVariable, TArray<TPair<FName, FInstancedStruct>>& OutBindings) const override
 		{
-			TooltipTexts.Add(
-				FText::Format(LOCTEXT("BindingTooltip_BoundParameter", "Bound to parameter \"{0}\" type \"{1}\""),
-					FText::FromName(ParameterBinding->AliasedParameter.GetName()),
-					FText::FromString(ParameterBinding->AliasedParameter.GetType().GetName())
-				)
-			);
+			TPair<FName, FInstancedStruct>& BindingData = OutBindings.AddDefaulted_GetRef();
+			BindingData.Key = AliasedVariable.GetName();
+			FNiagaraParameterBinding& Binding = BindingData.Value.InitializeAs<FNiagaraParameterBinding>();
+			Binding.AliasedParameter = AliasedVariable;
+			Binding.ResolvedParameter = ResolvedVariable;
 		}
 
-		TooltipTexts.Add(NiagaraParameterBindingCustomizationPrivate::GetAllowedBindingTypesText(ParameterBinding));
-		return FText::Join(LOCTEXT("NewLine", "\n"), TooltipTexts);
-	}
-	return LOCTEXT("BindingTooltip_Invalid", "Binding is invalid");
-}
+		virtual bool AllowUserParameters() const override { return ParameterBinding->AllowUserParameters(); }
+		virtual bool AllowSystemParameters() const override { return ParameterBinding->AllowSystemParameters(); }
+		virtual bool AllowEmitterParameters() const override { return ParameterBinding->AllowEmitterParameters(); }
+		virtual bool AllowParticleParameters() const override { return ParameterBinding->AllowParticleParameters(); }
+		virtual bool AllowStaticVariables() const override { return ParameterBinding->AllowStaticVariables(); }
 
-FName FNiagaraParameterBindingCustomization::GetVariableName() const
-{
-	FNiagaraParameterBinding* ParameterBinding = GetParameterBinding();
-	return ParameterBinding ? ParameterBinding->AliasedParameter.GetName() : FName();
-}
+		virtual TConstArrayView<UClass*> GetAllowedDataInterfaces() const override { return ParameterBinding->GetAllowedDataInterfaces(); }
+		virtual TConstArrayView<UClass*> GetAllowedObjects() const override { return ParameterBinding->GetAllowedObjects(); }
+		virtual TConstArrayView<UClass*> GetAllowedInterfaces() const override { return ParameterBinding->GetAllowedInterfaces(); }
+		virtual TConstArrayView<FNiagaraTypeDefinition> GetAllowedTypeDefinitions() const override { return ParameterBinding->GetAllowedTypeDefinitions(); }
 
-EVisibility FNiagaraParameterBindingCustomization::IsResetToDefaultsVisible() const
-{
-	FNiagaraParameterBinding* ParameterBinding = GetParameterBinding();
-	return ParameterBinding && ParameterBinding->IsSetoToDefault() ? EVisibility::Hidden : EVisibility::Visible;
-}
+		FNiagaraParameterBinding* ParameterBinding = nullptr;
+	};
 
-FReply FNiagaraParameterBindingCustomization::OnResetToDefaultsClicked()
-{
-	if (FNiagaraParameterBinding* ParameterBinding = GetParameterBinding())
+	static TUniquePtr<FNiagaraParameterBindingAdapter> CreateAdapter(uint8* ObjectAddress)
 	{
-		ParameterBinding->SetToDefault();
-
-		if (DefaultValueStructOnScope.IsValid())
-		{
-			TConstArrayView<uint8> DefaultValue = ParameterBinding->GetDefaultValueEditorOnly();
-			FMemory::Memcpy(DefaultValueStructOnScope->GetStructMemory(), DefaultValue.GetData(), DefaultValue.Num());
-			DefaultValueParameterEditor->UpdateInternalValueFromStruct(DefaultValueStructOnScope.ToSharedRef());
-		}
-
-		PropertyHandle->NotifyPostChange(EPropertyChangeType::ValueSet);
-		PropertyHandle->NotifyFinishedChangingProperties();
+		return MakeUnique<FAdapter>(ObjectAddress);
 	}
+};
 
-	return FReply::Handled();
+TSharedRef<IPropertyTypeCustomization> FNiagaraParameterBindingCustomization::MakeInstance()
+{
+	return MakeShared<FNiagaraParameterBindingAdapterCustomization>(&NiagaraParameterBindingAdapterCustomizationPrivate::CreateAdapter);
 }
-
-#undef LOCTEXT_NAMESPACE
diff --git a/Engine/Plugins/FX/Niagara/Source/NiagaraEditor/Private/Customizations/NiagaraParameterBindingCustomization.h b/Engine/Plugins/FX/Niagara/Source/NiagaraEditor/Private/Customizations/NiagaraParameterBindingCustomization.h
--- a/Engine/Plugins/FX/Niagara/Source/NiagaraEditor/Private/Customizations/NiagaraParameterBindingCustomization.h
+++ b/Engine/Plugins/FX/Niagara/Source/NiagaraEditor/Private/Customizations/NiagaraParameterBindingCustomization.h
@@ -2,83 +2,9 @@
 
 #pragma once
 
-#include "EdGraph/EdGraphSchema.h"
-#include "CoreMinimal.h"
 #include "IPropertyTypeCustomization.h"
-#include "Layout/Visibility.h"
-#include "UObject/StructOnScope.h"
 
-#include "NiagaraCommon.h"
-#include "NiagaraParameterBindingCustomization.generated.h"
-
-class FDetailWidgetRow;
-class IDetailChildrenBuilder;
-struct FEdGraphSchemaAction;
-struct FGraphActionListBuilderBase;
-class IPropertyTypeCustomizationUtils;
-class SWidget;
-
-class UNiagaraEmitter;
-class UNiagaraSystem;
-struct FNiagaraParameterBinding;
-class SNiagaraParameterEditor;
-
-USTRUCT()
-struct FNiagaraParameterBindingAction : public FEdGraphSchemaAction
-{
-	GENERATED_USTRUCT_BODY();
-
-	FNiagaraParameterBindingAction();
-	explicit FNiagaraParameterBindingAction(FText InName, FText InTooltip);
-
-	static TSharedRef<FNiagaraParameterBindingAction> MakeNone();
-	static TSharedRef<FNiagaraParameterBindingAction> MakeConstant(const FNiagaraTypeDefinition& TypeDef, const uint8* ValueData);
-	static TSharedRef<FNiagaraParameterBindingAction> MakeParameter(FNiagaraVariableBase InAliasedParameter, FNiagaraVariableBase InResolvedParameter);
-
-	bool bIsValidParameter = false;
-	FNiagaraVariableBase AliasedParameter;
-	FNiagaraVariableBase ResolvedParameter;
-};
-
-class FNiagaraParameterBindingCustomization : public IPropertyTypeCustomization
+namespace FNiagaraParameterBindingCustomization
 {
-public:
-	static TSharedRef<IPropertyTypeCustomization> MakeInstance();
-
-	/** IPropertyTypeCustomization interface begin */
-	virtual void CustomizeHeader(TSharedRef<IPropertyHandle> StructPropertyHandle, FDetailWidgetRow& HeaderRow, IPropertyTypeCustomizationUtils& StructCustomizationUtils) override;
-	virtual void CustomizeChildren(TSharedRef<IPropertyHandle> StructPropertyHandle, IDetailChildrenBuilder& ChildBuilder, IPropertyTypeCustomizationUtils& StructCustomizationUtils) override {}
-	/** IPropertyTypeCustomization interface end */
-
-	bool IsValid() const;
-	FNiagaraParameterBinding* GetParameterBinding() const;
-	FNiagaraParameterBinding* GetDefaultParameterBinding() const;
-
-	TSharedRef<SWidget> OnGetMenuContent() const;
-	bool ForEachBindableVariable(TFunction<bool(FNiagaraVariableBase, FNiagaraVariableBase)> Delegate) const;
-	void CollectAllActions(FGraphActionListBuilderBase& OutAllActions) const;
-	TSharedRef<SWidget> OnCreateWidgetForAction(struct FCreateWidgetForActionData* const InCreateData) const;
-	void OnActionSelected(const TArray<TSharedPtr<FEdGraphSchemaAction>>& SelectedActions, ESelectInfo::Type InSelectionType) const;
-	void OnValueChanged() const;
-
-	bool IsBindingValid(FNiagaraParameterBinding* ParameterBinding) const;
-	EVisibility IsBindingVisibile() const;
-	EVisibility IsConstantVisibile() const;
-	bool IsBindingEnabled() const;
-	bool IsConstantEnabled() const;
-
-	FText GetTooltipText() const;
-	FName GetVariableName() const;
-
-	EVisibility IsResetToDefaultsVisible() const;
-	FReply OnResetToDefaultsClicked();
-
-protected:
-	TSharedPtr<IPropertyHandle>		PropertyHandle;
-	TWeakObjectPtr<UObject>			OwnerWeakPtr;
-	TWeakObjectPtr<UNiagaraEmitter>	EmitterWeakPtr;
-	TWeakObjectPtr<UNiagaraSystem>	SystemWeakPtr;
-
-	TSharedPtr<SNiagaraParameterEditor>	DefaultValueParameterEditor;
-	TSharedPtr<FStructOnScope>			DefaultValueStructOnScope;
+	TSharedRef<IPropertyTypeCustomization> MakeInstance();
 };
