Remove MakeViewportToolbar() calls that did nothing.

diff --git a/Engine/Plugins/Animation/IKRig/Source/IKRigEditor/Private/RetargetEditor/SRetargetAnimAssetsWindow.cpp b/Engine/Plugins/Animation/IKRig/Source/IKRigEditor/Private/RetargetEditor/SRetargetAnimAssetsWindow.cpp
--- a/Engine/Plugins/Animation/IKRig/Source/IKRigEditor/Private/RetargetEditor/SRetargetAnimAssetsWindow.cpp
+++ b/Engine/Plugins/Animation/IKRig/Source/IKRigEditor/Private/RetargetEditor/SRetargetAnimAssetsWindow.cpp
@@ -755,11 +755,6 @@ TSharedRef<FEditorViewportClient> SRetargetPoseViewport::MakeEditorViewportClien
 	return EditorViewportClient.ToSharedRef();
 }
 
-TSharedPtr<SWidget> SRetargetPoseViewport::MakeViewportToolbar()
-{
-	return nullptr;
-}
-
 UBatchRetargetSettings* UBatchRetargetSettings::GetInstance()
 {
 	if (!SingletonInstance)
diff --git a/Engine/Plugins/Animation/IKRig/Source/IKRigEditor/Public/RetargetEditor/SRetargetAnimAssetsWindow.h b/Engine/Plugins/Animation/IKRig/Source/IKRigEditor/Public/RetargetEditor/SRetargetAnimAssetsWindow.h
--- a/Engine/Plugins/Animation/IKRig/Source/IKRigEditor/Public/RetargetEditor/SRetargetAnimAssetsWindow.h
+++ b/Engine/Plugins/Animation/IKRig/Source/IKRigEditor/Public/RetargetEditor/SRetargetAnimAssetsWindow.h
@@ -139,7 +139,6 @@ protected:
 	
 	// SEditorViewport interface
 	virtual TSharedRef<FEditorViewportClient> MakeEditorViewportClient() override;
-	virtual TSharedPtr<SWidget> MakeViewportToolbar() override;
 	// END SEditorViewport interface
 
 private:
diff --git a/Engine/Plugins/Experimental/Animation/TrajectoryTools/Source/TrajectoryTools/Private/SExportTrajectoriesWindow.cpp b/Engine/Plugins/Experimental/Animation/TrajectoryTools/Source/TrajectoryTools/Private/SExportTrajectoriesWindow.cpp
--- a/Engine/Plugins/Experimental/Animation/TrajectoryTools/Source/TrajectoryTools/Private/SExportTrajectoriesWindow.cpp
+++ b/Engine/Plugins/Experimental/Animation/TrajectoryTools/Source/TrajectoryTools/Private/SExportTrajectoriesWindow.cpp
@@ -218,12 +218,6 @@ TSharedRef<FEditorViewportClient> SPreviewTrajectoryViewport::MakeEditorViewport
 	return EditorViewportClient.ToSharedRef();
 }
 
-TSharedPtr<SWidget> SPreviewTrajectoryViewport::MakeViewportToolbar()
-{
-	// Use an empty toolbar
-	return SEditorViewport::MakeViewportToolbar();
-}
-
 FPreviewTrajectoryViewportClient::FPreviewTrajectoryViewportClient(FPreviewScene& InPreviewScene, const TSharedRef<SPreviewTrajectoryViewport>& InPreviewTrajectoryViewport)
 	: FEditorViewportClient(nullptr, &InPreviewScene, StaticCastSharedRef<SEditorViewport>(InPreviewTrajectoryViewport))
 {
diff --git a/Engine/Plugins/Experimental/Animation/TrajectoryTools/Source/TrajectoryTools/Public/SExportTrajectoriesWindow.h b/Engine/Plugins/Experimental/Animation/TrajectoryTools/Source/TrajectoryTools/Public/SExportTrajectoriesWindow.h
--- a/Engine/Plugins/Experimental/Animation/TrajectoryTools/Source/TrajectoryTools/Public/SExportTrajectoriesWindow.h
+++ b/Engine/Plugins/Experimental/Animation/TrajectoryTools/Source/TrajectoryTools/Public/SExportTrajectoriesWindow.h
@@ -123,7 +123,6 @@ public:
 protected:
 	/** Begin SEditorViewport interface */
 	virtual TSharedRef<FEditorViewportClient> MakeEditorViewportClient() override;
-	virtual TSharedPtr<SWidget> MakeViewportToolbar() override;
 	/** End SEditorViewport interface */
 
 private:
diff --git a/Engine/Plugins/Experimental/CharacterFXEditor/BaseCharacterFXEditor/Source/BaseCharacterFXEditor/Private/SBaseCharacterFXEditorViewport.cpp b/Engine/Plugins/Experimental/CharacterFXEditor/BaseCharacterFXEditor/Source/BaseCharacterFXEditor/Private/SBaseCharacterFXEditorViewport.cpp
--- a/Engine/Plugins/Experimental/CharacterFXEditor/BaseCharacterFXEditor/Source/BaseCharacterFXEditor/Private/SBaseCharacterFXEditorViewport.cpp
+++ b/Engine/Plugins/Experimental/CharacterFXEditor/BaseCharacterFXEditor/Source/BaseCharacterFXEditor/Private/SBaseCharacterFXEditorViewport.cpp
@@ -18,9 +18,4 @@ void SBaseCharacterFXEditorViewport::RemoveOverlayWidget(TSharedRef<SWidget> Ove
 	ViewportOverlay->RemoveSlot(OverlaidWidget);
 }
 
-TSharedPtr<SWidget> SBaseCharacterFXEditorViewport::MakeViewportToolbar()
-{
-	return SNew(SViewportToolBar);
-}
-
 #undef LOCTEXT_NAMESPACE
diff --git a/Engine/Plugins/Experimental/CharacterFXEditor/BaseCharacterFXEditor/Source/BaseCharacterFXEditor/Public/SBaseCharacterFXEditorViewport.h b/Engine/Plugins/Experimental/CharacterFXEditor/BaseCharacterFXEditor/Source/BaseCharacterFXEditor/Public/SBaseCharacterFXEditorViewport.h
--- a/Engine/Plugins/Experimental/CharacterFXEditor/BaseCharacterFXEditor/Source/BaseCharacterFXEditor/Public/SBaseCharacterFXEditorViewport.h
+++ b/Engine/Plugins/Experimental/CharacterFXEditor/BaseCharacterFXEditor/Source/BaseCharacterFXEditor/Public/SBaseCharacterFXEditorViewport.h
@@ -13,7 +13,4 @@ public:
 	// is not helpful here because that gets called just once.
 	virtual void AddOverlayWidget(TSharedRef<SWidget> OverlaidWidget, int32 ZOrder = INDEX_NONE);
 	virtual void RemoveOverlayWidget(TSharedRef<SWidget> OverlaidWidget);
-
-	// SEditorViewport
-	virtual TSharedPtr<SWidget> MakeViewportToolbar() override;
 };
diff --git a/Engine/Plugins/FX/Niagara/Source/NiagaraEditor/Private/Customizations/SimCache/NiagaraRenderTargetVolumeSimCacheVisualizer.cpp b/Engine/Plugins/FX/Niagara/Source/NiagaraEditor/Private/Customizations/SimCache/NiagaraRenderTargetVolumeSimCacheVisualizer.cpp
--- a/Engine/Plugins/FX/Niagara/Source/NiagaraEditor/Private/Customizations/SimCache/NiagaraRenderTargetVolumeSimCacheVisualizer.cpp
+++ b/Engine/Plugins/FX/Niagara/Source/NiagaraEditor/Private/Customizations/SimCache/NiagaraRenderTargetVolumeSimCacheVisualizer.cpp
@@ -61,7 +61,6 @@ public:
 protected:
 	/** SEditorViewport interface */
 	virtual TSharedRef<FEditorViewportClient> MakeEditorViewportClient() override;
-	virtual TSharedPtr<SWidget> MakeViewportToolbar() override;
 	virtual EVisibility OnGetViewportContentVisibility() const override;
 	virtual void BindCommands() override;
 
@@ -411,11 +410,6 @@ TSharedRef<FEditorViewportClient> SNiagaraVolumeTextureViewport::MakeEditorViewp
 	return ViewportClient.ToSharedRef();
 }
 
-TSharedPtr<SWidget> SNiagaraVolumeTextureViewport::MakeViewportToolbar()
-{
-	return SNew(SBox);
-}
-
 EVisibility SNiagaraVolumeTextureViewport::OnGetViewportContentVisibility() const
 {
 	EVisibility BaseVisibility = SEditorViewport::OnGetViewportContentVisibility();
diff --git a/Engine/Source/Editor/UnrealEd/Private/SSkeletonWidget.cpp b/Engine/Source/Editor/UnrealEd/Private/SSkeletonWidget.cpp
--- a/Engine/Source/Editor/UnrealEd/Private/SSkeletonWidget.cpp
+++ b/Engine/Source/Editor/UnrealEd/Private/SSkeletonWidget.cpp
@@ -674,11 +674,6 @@ TSharedRef<FEditorViewportClient> SBasePoseViewport::MakeEditorViewportClient()
 	return EditorViewportClient.ToSharedRef();
 }
 
-TSharedPtr<SWidget> SBasePoseViewport::MakeViewportToolbar()
-{
-	return nullptr;
-}
-
 
 /////////////////////////////////////////////////
 // select folder dialog
diff --git a/Engine/Source/Editor/UnrealEd/Public/SSkeletonWidget.h b/Engine/Source/Editor/UnrealEd/Public/SSkeletonWidget.h
--- a/Engine/Source/Editor/UnrealEd/Public/SSkeletonWidget.h
+++ b/Engine/Source/Editor/UnrealEd/Public/SSkeletonWidget.h
@@ -296,7 +296,6 @@ public:
 protected:
 	/** SEditorViewport interface */
 	virtual TSharedRef<FEditorViewportClient> MakeEditorViewportClient() override;
-	virtual TSharedPtr<SWidget> MakeViewportToolbar() override;
 
 private:
 	/** Skeleton */
