- Turned off input & output snapping on Niagara curves by default
- Added serialization so it saves last state
- All curve snapping now shares state

diff --git a/Engine/Plugins/FX/Niagara/Source/NiagaraEditor/Private/NiagaraEditorSettings.cpp b/Engine/Plugins/FX/Niagara/Source/NiagaraEditor/Private/NiagaraEditorSettings.cpp
--- a/Engine/Plugins/FX/Niagara/Source/NiagaraEditor/Private/NiagaraEditorSettings.cpp
+++ b/Engine/Plugins/FX/Niagara/Source/NiagaraEditor/Private/NiagaraEditorSettings.cpp
@@ -794,6 +794,26 @@ EGraphZoomLimitHandling UNiagaraEditorSettings::GetScriptEditorZoomLimitHandling
 	return ScriptGraphZoomLimitHandling;
 }
 
+bool UNiagaraEditorSettings::IsCurveInputSnapEnabled() const
+{
+	return CurveEditorSharedSettings.bCurveInputSnapEnabled;
+}
+
+bool UNiagaraEditorSettings::IsCurveOutputSnapEnabled() const
+{
+	return CurveEditorSharedSettings.bCurveOutputSnapEnabled;
+}
+
+void UNiagaraEditorSettings::SetCurveInputSnapEnabled(bool bEnabled)
+{
+	CurveEditorSharedSettings.bCurveInputSnapEnabled = bEnabled;
+}
+
+void UNiagaraEditorSettings::SetCurveOutputSnapEnabled(bool bEnabled)
+{
+	CurveEditorSharedSettings.bCurveOutputSnapEnabled = bEnabled;
+}
+
 bool UNiagaraEditorSettings::IsAllowedObjectByClassUsageInternal(const UObject& InObject, TSet<const UObject*>& CheckedObjects) const
 {
 	CheckedObjects.Add(&InObject);
diff --git a/Engine/Plugins/FX/Niagara/Source/NiagaraEditor/Public/NiagaraEditorSettings.h b/Engine/Plugins/FX/Niagara/Source/NiagaraEditor/Public/NiagaraEditorSettings.h
--- a/Engine/Plugins/FX/Niagara/Source/NiagaraEditor/Public/NiagaraEditorSettings.h
+++ b/Engine/Plugins/FX/Niagara/Source/NiagaraEditor/Public/NiagaraEditorSettings.h
@@ -318,6 +318,18 @@ struct FNiagaraViewportSharedSettings
 	bool bShowStatelessInfo = true;
 };
 
+USTRUCT()
+struct FNiagaraCurveEditorSharedSettings
+{
+	GENERATED_BODY()
+	
+	UPROPERTY(Config)
+	bool bCurveInputSnapEnabled = false;
+
+	UPROPERTY(Config)
+	bool bCurveOutputSnapEnabled = false;
+};
+
 FORCEINLINE uint32 GetTypeHash(const FNiagaraNamespaceMetadata& NamespaceMetaData)
 {
 	return GetTypeHash(NamespaceMetaData.GetGuid());
@@ -405,7 +417,7 @@ public:
 
 	UPROPERTY(config, EditAnywhere, Category = "Niagara Graph")
 	EGraphZoomLimitHandling ScriptGraphZoomLimitHandling = EGraphZoomLimitHandling::Default;
-
+	
 	NIAGARAEDITOR_API TArray<float> GetPlaybackSpeeds() const;
 
 	/** Gets whether or not auto-compile is enabled in the editors. */
@@ -533,6 +545,12 @@ public:
 	NIAGARAEDITOR_API bool GetForceSilentLoadingOfCachedAssets() const { return bForceSilentLoadingOfCachedAssets; }
 
 	NIAGARAEDITOR_API EGraphZoomLimitHandling GetScriptEditorZoomLimitHandling() const;
+
+	NIAGARAEDITOR_API bool IsCurveInputSnapEnabled() const;
+	NIAGARAEDITOR_API bool IsCurveOutputSnapEnabled() const;
+	NIAGARAEDITOR_API void SetCurveInputSnapEnabled(bool bEnabled);
+	NIAGARAEDITOR_API void SetCurveOutputSnapEnabled(bool bEnabled);
+
 private:
 	NIAGARAEDITOR_API bool IsAllowedObjectByClassUsageInternal(const UObject& InObject, TSet<const UObject*>& CheckedObjects) const;
 	NIAGARAEDITOR_API bool IsAllowedAssetObjectByClassUsageInternal(const UObject& InAssetObject, TSet<const UObject*>& CheckedAssetObjects) const;
@@ -609,6 +627,9 @@ private:
 
 	UPROPERTY(config, EditAnywhere, Category = Niagara)
 	TArray<FNiagaraCurveTemplate> CurveTemplates;
+
+	UPROPERTY(config)
+	FNiagaraCurveEditorSharedSettings CurveEditorSharedSettings;
 	
 	UPROPERTY(config)
 	FNiagaraViewportSharedSettings ViewportSettings;
diff --git a/Engine/Plugins/FX/Niagara/Source/NiagaraEditorWidgets/Private/DetailCustomizations/NiagaraDataInterfaceCurveDetails.cpp b/Engine/Plugins/FX/Niagara/Source/NiagaraEditorWidgets/Private/DetailCustomizations/NiagaraDataInterfaceCurveDetails.cpp
--- a/Engine/Plugins/FX/Niagara/Source/NiagaraEditorWidgets/Private/DetailCustomizations/NiagaraDataInterfaceCurveDetails.cpp
+++ b/Engine/Plugins/FX/Niagara/Source/NiagaraEditorWidgets/Private/DetailCustomizations/NiagaraDataInterfaceCurveDetails.cpp
@@ -862,9 +862,16 @@ public:
 			InitializeView();
 		}
 
+		UNiagaraEditorSettings* EditorSettings = GetMutableDefault<UNiagaraEditorSettings>();
 		CurveEditor = MakeShared<FCurveEditor>();
 		FCurveEditorInitParams InitParams;
 		CurveEditor->InitCurveEditor(InitParams);
+		TAttribute<bool>::FGetter InputSnapEnabledGetter = TAttribute<bool>::FGetter::CreateUObject(EditorSettings, &UNiagaraEditorSettings::IsCurveInputSnapEnabled);
+		TAttribute<bool>::FGetter OutputSnapEnabledGetter = TAttribute<bool>::FGetter::CreateUObject(EditorSettings, &UNiagaraEditorSettings::IsCurveOutputSnapEnabled);
+		CurveEditor->InputSnapEnabledAttribute = TAttribute<bool>::Create(InputSnapEnabledGetter);
+		CurveEditor->OutputSnapEnabledAttribute = TAttribute<bool>::Create(OutputSnapEnabledGetter);
+		CurveEditor->OnInputSnapEnabledChanged.BindUObject(EditorSettings, &UNiagaraEditorSettings::SetCurveInputSnapEnabled);
+		CurveEditor->OnOutputSnapEnabledChanged.BindUObject(EditorSettings, &UNiagaraEditorSettings::SetCurveOutputSnapEnabled);
 		CurveEditor->GridLineLabelFormatXAttribute = LOCTEXT("GridXLabelFormat", "{0}");
 
 		// Initialize our bounds at slightly larger than default to avoid clipping the tabs on the color widget.
