- Add CPU support for stateless kill box
- Add an optional compaction step that runs if particles are killed (modules must notify this is required)
- Add missing module data gather from shader parameter set in calculate accurate velocities

#rnx

diff --git a/Engine/Plugins/FX/Niagara/Source/Niagara/Internal/Stateless/Modules/NiagaraStatelessModule_CalculateAccurateVelocity.h b/Engine/Plugins/FX/Niagara/Source/Niagara/Internal/Stateless/Modules/NiagaraStatelessModule_CalculateAccurateVelocity.h
--- a/Engine/Plugins/FX/Niagara/Source/Niagara/Internal/Stateless/Modules/NiagaraStatelessModule_CalculateAccurateVelocity.h
+++ b/Engine/Plugins/FX/Niagara/Source/Niagara/Internal/Stateless/Modules/NiagaraStatelessModule_CalculateAccurateVelocity.h
@@ -14,6 +14,7 @@ class UNiagaraStatelessModule_CalculateAccurateVelocity : public UNiagaraStatele
 
 public:
 	virtual void BuildEmitterData(const FNiagaraStatelessEmitterDataBuildContext& BuildContext) const override;
+	virtual void SetShaderParameters(const FNiagaraStatelessSetShaderParameterContext& SetShaderParameterContext) const override;
 
 #if WITH_EDITORONLY_DATA
 	virtual void GetOutputVariables(TArray<FNiagaraVariableBase>& OutVariables) const override;
diff --git a/Engine/Plugins/FX/Niagara/Source/Niagara/Internal/Stateless/NiagaraStatelessEmitterDataBuildContext.h b/Engine/Plugins/FX/Niagara/Source/Niagara/Internal/Stateless/NiagaraStatelessEmitterDataBuildContext.h
--- a/Engine/Plugins/FX/Niagara/Source/Niagara/Internal/Stateless/NiagaraStatelessEmitterDataBuildContext.h
+++ b/Engine/Plugins/FX/Niagara/Source/Niagara/Internal/Stateless/NiagaraStatelessEmitterDataBuildContext.h
@@ -184,9 +184,9 @@ public:
 	FNiagaraStatelessRangeColor   ConvertDistributionToRange(const FNiagaraDistributionRangeColor& Distribution, const FLinearColor& DefaultValue) const { return ConvertDistributionToRangeHelper<FNiagaraStatelessRangeColor>(Distribution, DefaultValue); }
 	FNiagaraStatelessRangeInt     ConvertDistributionToRange(const FNiagaraDistributionRangeInt& Distribution, int32 DefaultValue) const { return ConvertDistributionToRangeHelper<FNiagaraStatelessRangeInt>(Distribution, DefaultValue); }
 
-	void AddParticleSimulationExecSimulate(TFunction<void(const NiagaraStateless::FParticleSimulationContext&)> Func) const;
+	NIAGARA_API void AddParticleSimulationExecSimulate(TFunction<void(const NiagaraStateless::FParticleSimulationContext&)> Func) const;
 
-	int32 FindParticleVariableIndex(const FNiagaraVariableBase& Variable) const;
+	NIAGARA_API int32 FindParticleVariableIndex(const FNiagaraVariableBase& Variable) const;
 
 private:
 	FNiagaraDataSetCompiledData&					ParticleDataSet;
diff --git a/Engine/Plugins/FX/Niagara/Source/Niagara/Internal/Stateless/NiagaraStatelessParticleSimContext.h b/Engine/Plugins/FX/Niagara/Source/Niagara/Internal/Stateless/NiagaraStatelessParticleSimContext.h
--- a/Engine/Plugins/FX/Niagara/Source/Niagara/Internal/Stateless/NiagaraStatelessParticleSimContext.h
+++ b/Engine/Plugins/FX/Niagara/Source/Niagara/Internal/Stateless/NiagaraStatelessParticleSimContext.h
@@ -42,6 +42,7 @@ public:
 
 private:
 	void SimulateInternal(int32 EmitterRandomSeed, float EmitterAge, float DeltaTime, FSpawnInfoShaderParameters& SpawnParameters, uint32 ActiveParticles);
+	void CompactDeadParticles();
 
 public:
 	uint32 GetNumInstances() const { return NumInstances; }
@@ -52,6 +53,9 @@ public:
 
 	const FQuat4f& GetToSimulationRotation(ENiagaraCoordinateSpace SourceSpace) const;
 
+	// When modifying particle alive you need to notify the simulation that it needs to compact killed instances or not
+	void SetNeedsParticleKillCompaction(bool bNeedsCompaction) const { bNeedsParticleKillCompaction |= bNeedsCompaction; }
+
 	// Get a pointer to required particle components
 	int32* GetParticleAlive() const { return static_cast<int32*>(RequiredComponents[int32(EParticleComponent::Alive)]); }
 	float* GetParticleLifetime() { return static_cast<float*>(RequiredComponents[int32(EParticleComponent::Lifetime)]); }
@@ -234,6 +238,8 @@ private:
 	uint8*									BufferFloatData = nullptr;
 	uint8*									BufferInt32Data = nullptr;
 
+	mutable bool							bNeedsParticleKillCompaction = false;
+
 	void*									RequiredComponents[int32(EParticleComponent::Num)] = {};
 	TArray<uint8*, TInlineAllocator<16>>	VariableComponents;
 
diff --git a/Engine/Plugins/FX/Niagara/Source/Niagara/Internal/Stateless/NiagaraStatelessParticleSimExecData.h b/Engine/Plugins/FX/Niagara/Source/Niagara/Internal/Stateless/NiagaraStatelessParticleSimExecData.h
--- a/Engine/Plugins/FX/Niagara/Source/Niagara/Internal/Stateless/NiagaraStatelessParticleSimExecData.h
+++ b/Engine/Plugins/FX/Niagara/Source/Niagara/Internal/Stateless/NiagaraStatelessParticleSimExecData.h
@@ -35,9 +35,11 @@ public:
 		const bool IsFloat() const { return Type == 0; }
 		const bool IsInt32() const { return Type == 1; }
 		const uint32 GetOffset() const { return Offset; }
+		const uint32 GetNum() const { return Num; }
 
 		uint16	Type	: 1	 = 0;
 		uint16	Offset	: 15 = 0;
+		uint16	Num = 0;
 	};
 
 	struct FRequiredComponentOffset
diff --git a/Engine/Plugins/FX/Niagara/Source/Niagara/Private/Stateless/Modules/NiagaraStatelessModule_CalculateAccurateVelocity.cpp b/Engine/Plugins/FX/Niagara/Source/Niagara/Private/Stateless/Modules/NiagaraStatelessModule_CalculateAccurateVelocity.cpp
--- a/Engine/Plugins/FX/Niagara/Source/Niagara/Private/Stateless/Modules/NiagaraStatelessModule_CalculateAccurateVelocity.cpp
+++ b/Engine/Plugins/FX/Niagara/Source/Niagara/Private/Stateless/Modules/NiagaraStatelessModule_CalculateAccurateVelocity.cpp
@@ -54,6 +54,12 @@ void UNiagaraStatelessModule_CalculateAccurateVelocity::BuildEmitterData(const F
 	BuildContext.AddParticleSimulationExecSimulate(&ParticleSimulate);
 }
 
+void UNiagaraStatelessModule_CalculateAccurateVelocity::SetShaderParameters(const FNiagaraStatelessSetShaderParameterContext& SetShaderParameterContext) const
+{
+	using namespace NSMCalculateAccurateVelocityPrivate;
+
+	const FModuleBuiltData* ModuleBuiltData = SetShaderParameterContext.ReadBuiltData<FModuleBuiltData>();
+}
 
 #if WITH_EDITORONLY_DATA
 void UNiagaraStatelessModule_CalculateAccurateVelocity::GetOutputVariables(TArray<FNiagaraVariableBase>& OutVariables) const
diff --git a/Engine/Plugins/FX/Niagara/Source/Niagara/Private/Stateless/NiagaraStatelessParticleSimContext.cpp b/Engine/Plugins/FX/Niagara/Source/Niagara/Private/Stateless/NiagaraStatelessParticleSimContext.cpp
--- a/Engine/Plugins/FX/Niagara/Source/Niagara/Private/Stateless/NiagaraStatelessParticleSimContext.cpp
+++ b/Engine/Plugins/FX/Niagara/Source/Niagara/Private/Stateless/NiagaraStatelessParticleSimContext.cpp
@@ -233,6 +233,50 @@ namespace NiagaraStateless
 			ModuleRandomSeed = Callback.RandomSeedOffset;
 			Callback.Function(*this);
 		}
+
+		// Perform compaction is needed
+		if (bNeedsParticleKillCompaction)
+		{
+			CompactDeadParticles();
+		}
+	}
+
+	void FParticleSimulationContext::CompactDeadParticles()
+	{
+		const FParticleSimulationExecData* ExecData = EmitterData->ParticleSimExecData;
+
+		int32* ParticleAlive = GetParticleAlive();
+		for (int32 iParticle=int32(NumInstances-1); iParticle >= 0; --iParticle)
+		{
+			if (ParticleAlive[iParticle] != 0)
+			{
+				continue;
+			}
+
+			--NumInstances;
+			if (iParticle >= int32(NumInstances))
+			{
+				continue;
+			}
+
+			for ( int32 iComponent=0; iComponent < int32(EParticleComponent::Num); ++iComponent )
+			{
+				uint32* ComponentData = static_cast<uint32*>(RequiredComponents[iComponent]);
+				ComponentData[iParticle] = ComponentData[NumInstances];
+			}
+			for (int32 iVariable=0; iVariable < VariableComponents.Num(); ++iVariable)
+			{
+				uint8* VariableOutput = VariableComponents[iVariable];
+
+				const uint32 NumComponents = ExecData->VariableComponentOffsets[iVariable].GetNum();
+				for (uint32 iComponent=0; iComponent < NumComponents; ++iComponent)
+				{
+					uint32* ComponentData = reinterpret_cast<uint32*>(VariableOutput);
+					ComponentData[iParticle] = ComponentData[NumInstances];
+					VariableOutput += BufferStride;
+				}
+			}
+		}
 	}
 
 	const FQuat4f& FParticleSimulationContext::GetToSimulationRotation(ENiagaraCoordinateSpace SourceSpace) const
diff --git a/Engine/Plugins/FX/Niagara/Source/Niagara/Private/Stateless/NiagaraStatelessParticleSimExecData.cpp b/Engine/Plugins/FX/Niagara/Source/Niagara/Private/Stateless/NiagaraStatelessParticleSimExecData.cpp
--- a/Engine/Plugins/FX/Niagara/Source/Niagara/Private/Stateless/NiagaraStatelessParticleSimExecData.cpp
+++ b/Engine/Plugins/FX/Niagara/Source/Niagara/Private/Stateless/NiagaraStatelessParticleSimExecData.cpp
@@ -21,12 +21,14 @@ namespace NiagaraStateless
 				check(VariableLayout.GetNumInt32Components() == 0 && VariableLayout.GetNumHalfComponents() == 0);
 				VariableComponentOffsets[i].Type	= 0;
 				VariableComponentOffsets[i].Offset	= VariableLayout.GetFloatComponentStart();
+				VariableComponentOffsets[i].Num		= VariableLayout.GetNumFloatComponents();
 			}
 			else if (VariableLayout.GetNumInt32Components() > 0)
 			{
 				check(VariableLayout.GetNumFloatComponents() == 0 && VariableLayout.GetNumHalfComponents() == 0);
 				VariableComponentOffsets[i].Type	= 1;
 				VariableComponentOffsets[i].Offset	= VariableLayout.GetInt32ComponentStart();
+				VariableComponentOffsets[i].Num		= VariableLayout.GetNumInt32Components();
 			}
 			else
 			{
