- Add half float support for per instance bindings & Nanite GPU scene uploads

diff --git a/Engine/Plugins/FX/Niagara/Source/Niagara/Private/NiagaraRendererProperties.cpp b/Engine/Plugins/FX/Niagara/Source/Niagara/Private/NiagaraRendererProperties.cpp
--- a/Engine/Plugins/FX/Niagara/Source/Niagara/Private/NiagaraRendererProperties.cpp
+++ b/Engine/Plugins/FX/Niagara/Source/Niagara/Private/NiagaraRendererProperties.cpp
@@ -135,6 +135,48 @@ void MarkAndRenameMaterialForGarbage(UMaterialInstanceConstant* MIC)
 
 } // NiagaraRendererPropertiesPrivate
 
+FNiagaraRendererVariableInfo::FNiagaraRendererVariableInfo(const FNiagaraDataSetCompiledData* CompiledData, const FNiagaraVariableBase& Variable)
+{
+	const int32 VariableIndex = CompiledData->Variables.IndexOfByPredicate(
+		[&Variable](const FNiagaraVariable& InVariable)
+		{
+			return InVariable.GetName() == Variable.GetName();
+		}
+	);
+
+	if (VariableIndex != INDEX_NONE)
+	{
+		const FNiagaraVariableLayoutInfo& LayoutInfo = CompiledData->VariableLayouts[VariableIndex];
+		if (LayoutInfo.GetNumFloatComponents() > 0)
+		{
+			check(LayoutInfo.GetNumHalfComponents() + LayoutInfo.GetNumInt32Components() == 0)
+			DatasetOffset	= LayoutInfo.GetFloatComponentStart();
+			GPUBufferOffset	= LayoutInfo.GetFloatComponentStart();
+			NumComponents	= LayoutInfo.GetNumFloatComponents();
+			bUpload			= false;
+			bHalfType		= false;
+		}
+		else if (LayoutInfo.GetNumHalfComponents() > 0)
+		{
+			check(LayoutInfo.GetNumFloatComponents() + LayoutInfo.GetNumInt32Components() == 0)
+			DatasetOffset	= LayoutInfo.GetHalfComponentStart();
+			GPUBufferOffset	= LayoutInfo.GetHalfComponentStart();
+			NumComponents	= LayoutInfo.GetNumHalfComponents();
+			bUpload			= false;
+			bHalfType		= true;
+		}
+		else if (LayoutInfo.GetNumInt32Components() > 0)
+		{
+			check(LayoutInfo.GetNumInt32Components() + LayoutInfo.GetNumHalfComponents() == 0)
+			DatasetOffset	= LayoutInfo.GetInt32ComponentStart();
+			GPUBufferOffset	= LayoutInfo.GetInt32ComponentStart();
+			NumComponents	= LayoutInfo.GetNumInt32Components();
+			bUpload			= false;
+			bHalfType		= false;
+		}
+	}
+}
+
 void FNiagaraRendererLayout::Initialize(int32 NumVariables)
 {
 	VFVariables_GT.Reset(NumVariables);
diff --git a/Engine/Plugins/FX/Niagara/Source/Niagara/Public/NiagaraRendererProperties.h b/Engine/Plugins/FX/Niagara/Source/Niagara/Public/NiagaraRendererProperties.h
--- a/Engine/Plugins/FX/Niagara/Source/Niagara/Public/NiagaraRendererProperties.h
+++ b/Engine/Plugins/FX/Niagara/Source/Niagara/Public/NiagaraRendererProperties.h
@@ -122,6 +122,7 @@ struct FNiagaraRendererVariableInfo
 		check(int32(InGPUBufferOffset) < TNumericLimits<uint16>::Max());
 		check(int32(InNumComponents) <= TNumericLimits<uint16>::Max());
 	}
+	NIAGARA_API FNiagaraRendererVariableInfo(const FNiagaraDataSetCompiledData* CompiledData, const FNiagaraVariableBase& Variable);
 
 	FORCEINLINE int32 GetNumComponents() const { return NumComponents; }
 
