- Fix binding variable rename / remove when not applied to emitter parameters

#tests Bind to emitter / particle data, test rename / delete
#rnx

diff --git a/Engine/Plugins/FX/Niagara/Source/Niagara/Private/NiagaraParameterBinding.cpp b/Engine/Plugins/FX/Niagara/Source/Niagara/Private/NiagaraParameterBinding.cpp
--- a/Engine/Plugins/FX/Niagara/Source/Niagara/Private/NiagaraParameterBinding.cpp
+++ b/Engine/Plugins/FX/Niagara/Source/Niagara/Private/NiagaraParameterBinding.cpp
@@ -21,20 +21,26 @@ void FNiagaraParameterBinding::OnRenameEmitter(FStringView EmitterName)
 
 void FNiagaraParameterBinding::OnRenameVariable(const FNiagaraVariableBase& OldVariable, const FNiagaraVariableBase& NewVariable, FStringView EmitterName)
 {
-	if (AliasedParameter.GetName() == OldVariable.GetName() && AliasedParameter.GetType() == OldVariable.GetType() && ResolvedParameter.IsInNameSpace(EmitterName))
+	if (AliasedParameter.GetName() == OldVariable.GetName() && AliasedParameter.GetType() == OldVariable.GetType())
 	{
-		AliasedParameter = NewVariable;
-		ResolvedParameter = AliasedParameter;
-		ResolvedParameter.ReplaceRootNamespace(FNiagaraConstants::EmitterNamespaceString, EmitterName);
+		if (!AliasedParameter.IsInNameSpace(FNiagaraConstants::EmitterNamespaceString) || ResolvedParameter.IsInNameSpace(EmitterName))
+		{
+			AliasedParameter = NewVariable;
+			ResolvedParameter = AliasedParameter;
+			ResolvedParameter.ReplaceRootNamespace(FNiagaraConstants::EmitterNamespaceString, EmitterName);
+		}
 	}
 }
 
 void FNiagaraParameterBinding::OnRemoveVariable(const FNiagaraVariableBase& OldVariable, FStringView EmitterName)
 {
-	if (AliasedParameter.GetName() == OldVariable.GetName() && AliasedParameter.GetType() == OldVariable.GetType() && ResolvedParameter.IsInNameSpace(EmitterName))
+	if (AliasedParameter.GetName() == OldVariable.GetName() && AliasedParameter.GetType() == OldVariable.GetType())
 	{
-		AliasedParameter = DefaultAliasedParameter;
-		ResolvedParameter = DefaultResolvedParameter;
+		if (!AliasedParameter.IsInNameSpace(FNiagaraConstants::EmitterNamespaceString) || ResolvedParameter.IsInNameSpace(EmitterName))
+		{
+			AliasedParameter = DefaultAliasedParameter;
+			ResolvedParameter = DefaultResolvedParameter;
+		}
 	}
 }
 
