- Add Niagara to Visuals placement menu

diff --git a/Engine/Plugins/FX/Niagara/Source/NiagaraEditor/NiagaraEditor.Build.cs b/Engine/Plugins/FX/Niagara/Source/NiagaraEditor/NiagaraEditor.Build.cs
--- a/Engine/Plugins/FX/Niagara/Source/NiagaraEditor/NiagaraEditor.Build.cs
+++ b/Engine/Plugins/FX/Niagara/Source/NiagaraEditor/NiagaraEditor.Build.cs
@@ -53,6 +53,7 @@ public class NiagaraEditor : ModuleRules
 				"AssetTools",
 				"LevelSequence",
 				"SparseVolumeTexture",
+				"PlacementMode",
 				"EngineSettings",
 				"EditorConfig"
 			}
diff --git a/Engine/Plugins/FX/Niagara/Source/NiagaraEditor/Private/NiagaraActorPlacement.cpp b/Engine/Plugins/FX/Niagara/Source/NiagaraEditor/Private/NiagaraActorPlacement.cpp
new file mode 100644
--- /dev/null
+++ b/Engine/Plugins/FX/Niagara/Source/NiagaraEditor/Private/NiagaraActorPlacement.cpp
@@ -0,0 +1,40 @@
+// Copyright Epic Games, Inc. All Rights Reserved.
+
+#include "NiagaraActorPlacement.h"
+#include "ActorFactoryNiagara.h"
+
+#include "IPlacementModeModule.h"
+
+namespace FNiagaraActorPlacement
+{
+	static TOptional<FPlacementModeID> PlacementModeID;
+
+	void Register()
+	{
+		if (GIsEditor == false || PlacementModeID.IsSet() )
+		{
+			return;
+		}
+
+		if (IPlacementModeModule::IsAvailable())
+		{
+			PlacementModeID = IPlacementModeModule::Get().RegisterPlaceableItem(
+				FBuiltInPlacementCategories::Visual(),
+				MakeShared<FPlaceableItem>(*UActorFactoryNiagara::StaticClass())
+			);
+		}
+	}
+
+	void Unregister()
+	{
+		if (PlacementModeID.IsSet() == false)
+		{
+			return;
+		}
+		if (IPlacementModeModule::IsAvailable())
+		{
+			IPlacementModeModule::Get().UnregisterPlaceableItem(PlacementModeID.GetValue());
+		}
+		PlacementModeID.Reset();
+	}
+}
diff --git a/Engine/Plugins/FX/Niagara/Source/NiagaraEditor/Private/NiagaraActorPlacement.h b/Engine/Plugins/FX/Niagara/Source/NiagaraEditor/Private/NiagaraActorPlacement.h
new file mode 100644
--- /dev/null
+++ b/Engine/Plugins/FX/Niagara/Source/NiagaraEditor/Private/NiagaraActorPlacement.h
@@ -0,0 +1,9 @@
+// Copyright Epic Games, Inc. All Rights Reserved.
+
+#pragma once
+
+namespace FNiagaraActorPlacement
+{
+	void Register();
+	void Unregister();
+}
diff --git a/Engine/Plugins/FX/Niagara/Source/NiagaraEditor/Private/NiagaraEditorModule.cpp b/Engine/Plugins/FX/Niagara/Source/NiagaraEditor/Private/NiagaraEditorModule.cpp
--- a/Engine/Plugins/FX/Niagara/Source/NiagaraEditor/Private/NiagaraEditorModule.cpp
+++ b/Engine/Plugins/FX/Niagara/Source/NiagaraEditor/Private/NiagaraEditorModule.cpp
@@ -42,6 +42,7 @@
 #include "TypeEditorUtilities/NiagaraDataInterfaceCurveTypeEditorUtilities.h"
 #include "TypeEditorUtilities/NiagaraDistributionPropertyEditorUtilities.h"
 
+#include "NiagaraActorPlacement.h"
 #include "NiagaraSystemCompilingManager.h"
 #include "NiagaraEditorStyle.h"
 #include "NiagaraEditorCommands.h"
@@ -1631,6 +1632,8 @@ void FNiagaraEditorModule::StartupModule()
 
 void FNiagaraEditorModule::ShutdownModule()
 {
+	FNiagaraActorPlacement::Unregister();
+
 	MenuExtensibilityManager.Reset();
 	ToolBarExtensibilityManager.Reset();
 
@@ -1765,6 +1768,8 @@ void FNiagaraEditorModule::OnPostEngineInit()
 {
 	if (GIsEditor)
 	{
+		FNiagaraActorPlacement::Register();
+
 		UThumbnailManager::Get().RegisterCustomRenderer(UNiagaraEmitter::StaticClass(), UNiagaraEmitterThumbnailRenderer::StaticClass());
 		UThumbnailManager::Get().RegisterCustomRenderer(UNiagaraSystem::StaticClass(), UNiagaraSystemThumbnailRenderer::StaticClass());
 		bThumbnailRenderersRegistered = true;
