- Change Niagara Actor display name more descriptive & add a tooltip
- This mirrors Niagara Component naming

#rnx

diff --git a/Engine/Plugins/FX/Niagara/Source/Niagara/Public/NiagaraActor.h b/Engine/Plugins/FX/Niagara/Source/Niagara/Public/NiagaraActor.h
--- a/Engine/Plugins/FX/Niagara/Source/Niagara/Public/NiagaraActor.h
+++ b/Engine/Plugins/FX/Niagara/Source/Niagara/Public/NiagaraActor.h
@@ -7,7 +7,13 @@
 #include "GameFramework/Actor.h"
 #include "NiagaraActor.generated.h"
 
-UCLASS(hideCategories = (Activation, "Components|Activation", Input, Collision, "Game|Damage"), ComponentWrapperClass, MinimalAPI)
+/**
+* Niagara Particle System Actor contains a Niagara Particle System Component which can be used to simulate and render Niagara Systems
+*
+* @see https://dev.epicgames.com/documentation/en-us/unreal-engine/BlueprintAPI/NiagaraActor
+* @see UNiagaraComponent
+*/
+UCLASS(hideCategories = (Activation, "Components|Activation", Input, Collision, "Game|Damage"), ComponentWrapperClass, MinimalAPI, meta = (DisplayName = "Niagara Particle System Actor"))
 class ANiagaraActor : public AActor
 {
 	GENERATED_BODY()
