Add tracking of killed systems to the scalabiltiy debug hud.

diff --git a/Engine/Plugins/FX/Niagara/Source/Niagara/Private/NiagaraDebugHud.cpp b/Engine/Plugins/FX/Niagara/Source/Niagara/Private/NiagaraDebugHud.cpp
--- a/Engine/Plugins/FX/Niagara/Source/Niagara/Private/NiagaraDebugHud.cpp
+++ b/Engine/Plugins/FX/Niagara/Source/Niagara/Private/NiagaraDebugHud.cpp
@@ -1298,6 +1298,22 @@ void FNiagaraDebugHud::GatherSystemInfo()
 			SystemDebugInfo.FramesSinceVisible++;
 		}
 	}
+
+	GlobalScalabiltiyKillCounts.Reset();
+
+	if (FNiagaraWorldManager* WorldMan = FNiagaraWorldManager::Get(World))
+	{
+		const TMap<TWeakObjectPtr<UNiagaraSystem>, FNiagaraSystemScalabilityKillCounts>& SystemKillCounts = WorldMan->GetSystemScalabilityKillCounts();
+		for (auto& KillCountEntry : SystemKillCounts)
+		{
+			if (UNiagaraSystem* System = KillCountEntry.Key.Get())
+			{
+				GlobalScalabiltiyKillCounts.Add(KillCountEntry.Value);
+				FSystemDebugInfo& SystemDebugInfo = PerSystemDebugInfo.FindOrAdd(System->GetFName());
+				SystemDebugInfo.ScalabilityKillCounts = KillCountEntry.Value;
+			}
+		}
+	}
 }
 
 const FNiagaraDataSet* FNiagaraDebugHud::GetParticleDataSet(FNiagaraSystemInstance* SystemInstance, FNiagaraEmitterInstance* EmitterInstance, int32 iEmitter)
@@ -1800,6 +1816,7 @@ void FNiagaraDebugHud::DrawOverview(class FNiagaraWorldManager* WorldManager, FC
 		}
 		else if (Settings.OverviewMode == ENiagaraDebugHUDOverviewMode::Scalability)
 		{
+			FString ScalabilitySleepVsKilledFmt(TEXT("%d | %d"));
 			OverviewColumns.Emplace(TEXT("TotalActive:"), FString::FromInt(GlobalTotalActive), TEXT("# Active"), ColumnOffset, Font, TEXT("0000"),
 				[&DetailColor, &DetailHighlightColor, &fAdvanceHeight](FCanvas* Canvas, UFont* Font, float X, float Y, FOverviewColumn& Col, const FSystemDebugInfo& SystemInfo)
 				{
@@ -1820,54 +1837,59 @@ void FNiagaraDebugHud::DrawOverview(class FNiagaraWorldManager* WorldManager, FC
 					Canvas->DrawShadowedString(X, Y, *FString::FromInt(SystemInfo.TotalScalability), Font, RowColor);
 				});
 
-			OverviewColumns.Emplace(TEXT("TotalCulled:"), FString::FromInt(GlobalTotalCulled), TEXT("# Culled"), ColumnOffset, Font, TEXT("000"),
+			OverviewColumns.Emplace(TEXT("TotalCulled:"), FString::Printf(TEXT("%d | %d"), GlobalTotalCulled, GlobalScalabiltiyKillCounts.Total), TEXT("# Culled (Asleep | Killed)"), ColumnOffset, Font, TEXT("000"),
 				[&DetailColor, &DetailHighlightColor, &fAdvanceHeight](FCanvas* Canvas, UFont* Font, float X, float Y, FOverviewColumn& Col, const FSystemDebugInfo& SystemInfo)
 				{
 					FLinearColor RowBGColor = SystemInfo.UniqueColor;
 					RowBGColor.A = Settings.SystemColorTableOpacity;
 					Canvas->DrawTile(X, Y, Col.MaxWidth, fAdvanceHeight, 0, 0, 0, 0, RowBGColor);
 					const FLinearColor RowColor = SystemInfo.bShowInWorld ? DetailHighlightColor : DetailColor;
-					Canvas->DrawShadowedString(X, Y, *FString::FromInt(SystemInfo.TotalCulled), Font, RowColor);
+					FString ValueStr = FString::Printf(TEXT("%d | %d"), SystemInfo.TotalCulled, SystemInfo.ScalabilityKillCounts.Total);
+					Canvas->DrawShadowedString(X, Y, *ValueStr, Font, RowColor);
 				});
 
-			OverviewColumns.Emplace(TEXT("Distance:"), FString::FromInt(GlobalTotalCulledByDistance), TEXT("# Distance"), ColumnOffset, Font, TEXT("000"),
+			OverviewColumns.Emplace(TEXT("Distance:"), FString::Printf(TEXT("%d | %d"), GlobalTotalCulledByDistance, GlobalScalabiltiyKillCounts.Distance), TEXT("# Distance (Asleep | Killed)"), ColumnOffset, Font, TEXT("000"),
 				[&DetailColor, &DetailHighlightColor, &fAdvanceHeight](FCanvas* Canvas, UFont* Font, float X, float Y, FOverviewColumn& Col, const FSystemDebugInfo& SystemInfo)
 				{
 					FLinearColor RowBGColor = SystemInfo.UniqueColor;
 					RowBGColor.A = Settings.SystemColorTableOpacity;
 					Canvas->DrawTile(X, Y, Col.MaxWidth, fAdvanceHeight, 0, 0, 0, 0, RowBGColor);
 					const FLinearColor RowColor = SystemInfo.bShowInWorld ? DetailHighlightColor : DetailColor;
-					Canvas->DrawShadowedString(X, Y, *FString::FromInt(SystemInfo.TotalCulledByDistance), Font, RowColor);
+					FString ValueStr = FString::Printf(TEXT("%d | %d"), SystemInfo.TotalCulledByDistance, SystemInfo.ScalabilityKillCounts.Distance);
+					Canvas->DrawShadowedString(X, Y, *ValueStr, Font, RowColor);
 				});
 
-			OverviewColumns.Emplace(TEXT("Visibility:"), FString::FromInt(GlobalTotalCulledByVisibility), TEXT("# Visibility"), ColumnOffset, Font, TEXT("000"),
+			OverviewColumns.Emplace(TEXT("Visibility:"), FString::Printf(TEXT("%d | %d"), GlobalTotalCulledByVisibility, GlobalScalabiltiyKillCounts.Visibility), TEXT("# Visibility (Asleep | Killed)"), ColumnOffset, Font, TEXT("000"),
 				[&DetailColor, &DetailHighlightColor, &fAdvanceHeight](FCanvas* Canvas, UFont* Font, float X, float Y, FOverviewColumn& Col, const FSystemDebugInfo& SystemInfo)
 				{
 					FLinearColor RowBGColor = SystemInfo.UniqueColor;
 					RowBGColor.A = Settings.SystemColorTableOpacity;
 					Canvas->DrawTile(X, Y, Col.MaxWidth, fAdvanceHeight, 0, 0, 0, 0, RowBGColor);
 					const FLinearColor RowColor = SystemInfo.bShowInWorld ? DetailHighlightColor : DetailColor;
-					Canvas->DrawShadowedString(X, Y, *FString::FromInt(SystemInfo.TotalCulledByVisibility), Font, RowColor);
+					FString ValueStr = FString::Printf(TEXT("%d | %d"), SystemInfo.TotalCulledByVisibility, SystemInfo.ScalabilityKillCounts.Visibility);
+					Canvas->DrawShadowedString(X, Y, *ValueStr, Font, RowColor);
 				});
 
-			OverviewColumns.Emplace(TEXT("InstanceCount:"), FString::FromInt(GlobalTotalCulledByInstanceCount), TEXT("# Inst Count"), ColumnOffset, Font, TEXT("000"),
+			OverviewColumns.Emplace(TEXT("InstanceCount:"), FString::Printf(TEXT("%d | %d"), GlobalTotalCulledByInstanceCount, GlobalScalabiltiyKillCounts.InstanceCounts), TEXT("# Inst Count (Asleep | Killed)"), ColumnOffset, Font, TEXT("000"),
 				[&DetailColor, &DetailHighlightColor, &fAdvanceHeight](FCanvas* Canvas, UFont* Font, float X, float Y, FOverviewColumn& Col, const FSystemDebugInfo& SystemInfo)
 				{
 					FLinearColor RowBGColor = SystemInfo.UniqueColor;
 					RowBGColor.A = Settings.SystemColorTableOpacity;
 					Canvas->DrawTile(X, Y, Col.MaxWidth, fAdvanceHeight, 0, 0, 0, 0, RowBGColor);
 					const FLinearColor RowColor = SystemInfo.bShowInWorld ? DetailHighlightColor : DetailColor;
-					Canvas->DrawShadowedString(X, Y, *FString::FromInt(SystemInfo.TotalCulledByInstanceCount), Font, RowColor);
+					FString ValueStr = FString::Printf(TEXT("%d | %d"), SystemInfo.TotalCulledByInstanceCount, SystemInfo.ScalabilityKillCounts.InstanceCounts);
+					Canvas->DrawShadowedString(X, Y, *ValueStr, Font, RowColor);
 				});
 
-			OverviewColumns.Emplace(TEXT("Budget:"), FString::FromInt(GlobalTotalCulledByBudget), TEXT("# Budget"), ColumnOffset, Font, TEXT("000"),
+			OverviewColumns.Emplace(TEXT("Budget:"), FString::Printf(TEXT("%d | %d"), GlobalTotalCulledByBudget, GlobalScalabiltiyKillCounts.GlobalBudget), TEXT("# Budget (Asleep | Killed)"), ColumnOffset, Font, TEXT("000"),
 				[&DetailColor, &DetailHighlightColor, &fAdvanceHeight](FCanvas* Canvas, UFont* Font, float X, float Y, FOverviewColumn& Col, const FSystemDebugInfo& SystemInfo)
 				{
 					FLinearColor RowBGColor = SystemInfo.UniqueColor;
 					RowBGColor.A = Settings.SystemColorTableOpacity;
 					Canvas->DrawTile(X, Y, Col.MaxWidth, fAdvanceHeight, 0, 0, 0, 0, RowBGColor);
 					const FLinearColor RowColor = SystemInfo.bShowInWorld ? DetailHighlightColor : DetailColor;
-					Canvas->DrawShadowedString(X, Y, *FString::FromInt(SystemInfo.TotalCulledByBudget), Font, RowColor);
+					FString ValueStr = FString::Printf(TEXT("%d | %d"), SystemInfo.TotalCulledByBudget, SystemInfo.ScalabilityKillCounts.GlobalBudget);
+					Canvas->DrawShadowedString(X, Y, *ValueStr, Font, RowColor);
 				});
 
 			OverviewColumns.Emplace(TEXT("Player FX:"), FString::FromInt(GlobalTotalPlayerSystems), TEXT("# Player"), ColumnOffset, Font, TEXT("000"),
diff --git a/Engine/Plugins/FX/Niagara/Source/Niagara/Private/NiagaraDebugHud.h b/Engine/Plugins/FX/Niagara/Source/Niagara/Private/NiagaraDebugHud.h
--- a/Engine/Plugins/FX/Niagara/Source/Niagara/Private/NiagaraDebugHud.h
+++ b/Engine/Plugins/FX/Niagara/Source/Niagara/Private/NiagaraDebugHud.h
@@ -119,6 +119,8 @@ class FNiagaraDebugHud
 		
 		int32		TotalPlayerSystems = 0;
 
+		FNiagaraSystemScalabilityKillCounts ScalabilityKillCounts;
+
 		void Reset()
 		{
 			bShowInWorld = false;
@@ -274,6 +276,8 @@ private:
 
 	int32 GlobalTotalPlayerSystems = 0;
 
+	FNiagaraSystemScalabilityKillCounts GlobalScalabiltiyKillCounts;
+
 	FString LongestSystemPrettyName;
 	TMap<FName, FSystemDebugInfo>	PerSystemDebugInfo;
 
diff --git a/Engine/Plugins/FX/Niagara/Source/Niagara/Private/NiagaraScalabilityManager.cpp b/Engine/Plugins/FX/Niagara/Source/Niagara/Private/NiagaraScalabilityManager.cpp
--- a/Engine/Plugins/FX/Niagara/Source/Niagara/Private/NiagaraScalabilityManager.cpp
+++ b/Engine/Plugins/FX/Niagara/Source/Niagara/Private/NiagaraScalabilityManager.cpp
@@ -76,8 +76,18 @@ static int32 GetMaxUpdatesPerFrame(const UNiagaraEffectType* EffectType, int32 I
 	return ItemsRemaining;
 }
 
+FNiagaraScalabilityManager::FNiagaraScalabilityManager(FNiagaraWorldManager* InOwner)
+	: WorldMan(InOwner)
+	, EffectType(nullptr)
+	, LastUpdateTime(0.0)
+	, bRefreshCachedSystemData(false)
+{
+
+}
+
 FNiagaraScalabilityManager::FNiagaraScalabilityManager()
-	: EffectType(nullptr)
+	: WorldMan(nullptr)
+	, EffectType(nullptr)
 	, LastUpdateTime(0.0)
 	, bRefreshCachedSystemData(false)
 {
@@ -195,7 +205,7 @@ void FNiagaraScalabilityManager::UnregisterAt(int32 IndexToRemove)
 
 // Note that this function may unregister elements in the ManagedComponents array (if the Component or System are no longer valid).
 // Returns false if there was a problem evaluating the specified index and something had to be unregistered
-bool FNiagaraScalabilityManager::EvaluateCullState(FNiagaraWorldManager* WorldMan, FComponentIterationContext& Context, int32 ComponentIndex, int32& UpdateCounter)
+bool FNiagaraScalabilityManager::EvaluateCullState(FComponentIterationContext& Context, int32 ComponentIndex, int32& UpdateCounter)
 {
 	check(ManagedComponents.IsValidIndex(ComponentIndex));
 	UNiagaraComponent* Component = ManagedComponents[ComponentIndex];
@@ -265,7 +275,7 @@ bool FNiagaraScalabilityManager::EvaluateCullState(FNiagaraWorldManager* WorldMa
 	return true;
 }
 
-void FNiagaraScalabilityManager::ProcessSignificance(FNiagaraWorldManager* WorldMan, UNiagaraSignificanceHandler* SignificanceHandler, FComponentIterationContext& Context)
+void FNiagaraScalabilityManager::ProcessSignificance(UNiagaraSignificanceHandler* SignificanceHandler, FComponentIterationContext& Context)
 {
 	// it would be good to get a better estimate for how many indices we're going to need to process
 	Context.SignificanceIndices.Reset(ManagedComponents.Num());
@@ -386,24 +396,33 @@ bool FNiagaraScalabilityManager::ApplyScalabilityState(int32 ComponentIndex, ENi
 		Component->DebugCachedScalabilityState = CompState;
 #endif
 
+		UNiagaraSystem* System = Component->GetAsset();
+		check(System);
+
 		if (CompState.bCulled)
 		{
 			CompState.Apply();
 
+			bool bKilled = false;
 			switch (CullReaction)
 			{
-			case ENiagaraCullReaction::Deactivate:					Component->DeactivateInternal(false); bContinueIteration = false; break;//We don't increment CompIdx here as this call will remove an entry from ManagedObjects;
-			case ENiagaraCullReaction::DeactivateImmediate:			Component->DeactivateImmediateInternal(false); bContinueIteration = false;  break; //We don't increment CompIdx here as this call will remove an entry from ManagedObjects;
+			case ENiagaraCullReaction::Deactivate:					Component->DeactivateInternal(false); bKilled = true; break;//We don't increment CompIdx here as this call will remove an entry from ManagedObjects;
+			case ENiagaraCullReaction::DeactivateImmediate:			Component->DeactivateImmediateInternal(false); bKilled = true;  break; //We don't increment CompIdx here as this call will remove an entry from ManagedObjects;
 			case ENiagaraCullReaction::DeactivateResume:			Component->DeactivateInternal(true); break;
 			case ENiagaraCullReaction::DeactivateImmediateResume:	Component->DeactivateImmediateInternal(true); break;
 			case ENiagaraCullReaction::PauseResume:					Component->SetPausedInternal(true,true); break;
 			};
+
+			if (bKilled)
+			{
+				bContinueIteration = false;
+#if WITH_NIAGARA_DEBUGGER
+				WorldMan->GetSystemScalabilityKillCounts(System).Add(CompState);
+#endif
+			}
 		}
 		else
 		{
-			UNiagaraSystem* System = Component->GetAsset();
-			check(System);
-			
 			int32 SystemInstanceMax = 0;
 			int32 EffectTypeInstanceMax = 0;
 
@@ -443,14 +462,14 @@ bool FNiagaraScalabilityManager::ApplyScalabilityState(int32 ComponentIndex, ENi
 	return bContinueIteration;
 }
 
-void FNiagaraScalabilityManager::UpdateInternal(FNiagaraWorldManager* WorldMan, FComponentIterationContext& Context)
+void FNiagaraScalabilityManager::UpdateInternal(FComponentIterationContext& Context)
 {
 	int32 UpdateCount = 0;
 	if (Context.bProcessAllComponents)
 	{
 		for (int32 ComponentIt = 0; ComponentIt < ManagedComponents.Num();)
 		{
-			if (EvaluateCullState(WorldMan, Context, ComponentIt, UpdateCount))
+			if (EvaluateCullState(Context, ComponentIt, UpdateCount))
 			{
 				++ComponentIt;
 			}
@@ -462,7 +481,7 @@ void FNiagaraScalabilityManager::UpdateInternal(FNiagaraWorldManager* WorldMan,
 		{
 			Context.ComponentRequiresUpdate[ComponentIt.GetIndex()] = false;
 
-			if (EvaluateCullState(WorldMan, Context, ComponentIt.GetIndex(), UpdateCount))
+			if (EvaluateCullState(Context, ComponentIt.GetIndex(), UpdateCount))
 			{
 				++ComponentIt;
 			}
@@ -474,7 +493,7 @@ void FNiagaraScalabilityManager::UpdateInternal(FNiagaraWorldManager* WorldMan,
 		bool bProcessSignificance = Context.bRequiresGlobalSignificancePass && EffectType->SignificanceHandler;
 		if (bProcessSignificance)
 		{
-			ProcessSignificance(WorldMan, EffectType->SignificanceHandler, Context);
+			ProcessSignificance(EffectType->SignificanceHandler, Context);
 		}
 
 		if (Context.bHasDirtyState)
@@ -520,7 +539,7 @@ void FNiagaraScalabilityManager::UpdateInternal(FNiagaraWorldManager* WorldMan,
 	}
 }
 
-void FNiagaraScalabilityManager::Update(FNiagaraWorldManager* WorldMan, float DeltaSeconds, bool bNewOnly)
+void FNiagaraScalabilityManager::Update(float DeltaSeconds, bool bNewOnly)
 {
 	if (bRefreshCachedSystemData)
 	{	
@@ -566,7 +585,7 @@ void FNiagaraScalabilityManager::Update(FNiagaraWorldManager* WorldMan, float De
 		EffectType->bNewSystemsSinceLastScalabilityUpdate = false;
 		NewComponentContext.WorstGlobalBudgetUse = WorstGlobalBudgetUse;
 
-		UpdateInternal(WorldMan, NewComponentContext);
+		UpdateInternal(NewComponentContext);
 		return;
 	}
 	else if (EffectType->UpdateFrequency == ENiagaraScalabilityUpdateFrequency::SpawnOnly)
@@ -627,7 +646,7 @@ void FNiagaraScalabilityManager::Update(FNiagaraWorldManager* WorldMan, float De
 
 	DefaultContext.WorstGlobalBudgetUse = WorstGlobalBudgetUse;
 
-	UpdateInternal(WorldMan, DefaultContext);
+	UpdateInternal(DefaultContext);
 	
 	bRefreshOwnerAllowsScalability = false;
 }
diff --git a/Engine/Plugins/FX/Niagara/Source/Niagara/Private/NiagaraWorldManager.cpp b/Engine/Plugins/FX/Niagara/Source/Niagara/Private/NiagaraWorldManager.cpp
--- a/Engine/Plugins/FX/Niagara/Source/Niagara/Private/NiagaraWorldManager.cpp
+++ b/Engine/Plugins/FX/Niagara/Source/Niagara/Private/NiagaraWorldManager.cpp
@@ -878,6 +878,16 @@ void FNiagaraWorldManager::PostGarbageCollect()
 	//Clear out and scalability managers who's EffectTypes have been GCd.
 	while (ScalabilityManagers.Remove(nullptr)) {}
 
+	#if WITH_NIAGARA_DEBUGGER
+	for (auto It = SystemScalabilityKillCounts.CreateIterator(); It; ++It)
+	{
+		if (It.Key().IsValid() == false)
+		{
+			It.RemoveCurrent();
+		}
+	}
+	#endif
+
 #if WITH_NIAGARA_LEAK_DETECTOR
 	if (ComponentLeakDetector.IsValid() && HasActiveWorld())
 	{
@@ -1750,11 +1760,11 @@ void FNiagaraWorldManager::UpdateScalabilityManagers(float DeltaSeconds, bool bN
 
 		if (bNewSpawnsOnly)
 		{
-			ScalabilityMan.Update(this, DeltaSeconds, true);
+			ScalabilityMan.Update(DeltaSeconds, true);
 		}
 		else
 		{
-			ScalabilityMan.Update(this, DeltaSeconds, false);
+			ScalabilityMan.Update(DeltaSeconds, false);
 		}
 	}
 }
@@ -1776,7 +1786,7 @@ void FNiagaraWorldManager::RegisterWithScalabilityManager(UNiagaraComponent* Com
 		#if DO_CHECK
 			ensureMsgf(bScalabilityManagersLocked == false, TEXT("ScalabilityManagers modified while iterating, by Component (%s) System (%s)"), *GetNameSafe(Component), *GetNameSafe(Component ? Component->GetAsset() : nullptr));
 		#endif
-			ScalabilityManager = &ScalabilityManagers.Add(EffectType);
+			ScalabilityManager = &ScalabilityManagers.Emplace(EffectType, this);
 			ScalabilityManager->EffectType = EffectType;
 		}
 
@@ -1813,6 +1823,9 @@ bool FNiagaraWorldManager::ShouldPreCull(UNiagaraSystem* System, UNiagaraCompone
 				FNiagaraScalabilityState State;
 				const FNiagaraSystemScalabilitySettings& ScalabilitySettings = System->GetScalabilitySettings();
 				CalculateScalabilityState(System, ScalabilitySettings, EffectType, Component, true, FFXBudget::GetWorstAdjustedUsage(), State);
+				#if WITH_NIAGARA_DEBUGGER
+				GetSystemScalabilityKillCounts(System).Add(State);				
+				#endif
 				return State.bCulled;
 			}
 		}
diff --git a/Engine/Plugins/FX/Niagara/Source/Niagara/Public/NiagaraScalabilityManager.h b/Engine/Plugins/FX/Niagara/Source/Niagara/Public/NiagaraScalabilityManager.h
--- a/Engine/Plugins/FX/Niagara/Source/Niagara/Public/NiagaraScalabilityManager.h
+++ b/Engine/Plugins/FX/Niagara/Source/Niagara/Public/NiagaraScalabilityManager.h
@@ -49,6 +49,8 @@ struct FNiagaraScalabilityManager
 {
 	GENERATED_USTRUCT_BODY()
 
+	FNiagaraWorldManager* WorldMan = nullptr;
+
 	UPROPERTY(transient)
 	TObjectPtr<UNiagaraEffectType> EffectType;
 
@@ -64,9 +66,10 @@ struct FNiagaraScalabilityManager
 
 	bool bRefreshOwnerAllowsScalability = false;
 
+	FNiagaraScalabilityManager(FNiagaraWorldManager* InOwner);
 	FNiagaraScalabilityManager();
 	~FNiagaraScalabilityManager();
-	void Update(FNiagaraWorldManager* Owner, float DeltaSeconds, bool bNewOnly);
+	void Update(float DeltaSeconds, bool bNewOnly);
 	void Register(UNiagaraComponent* Component);
 	void Unregister(UNiagaraComponent* Component);
 
@@ -95,9 +98,9 @@ private:
 	void UnregisterAt(int32 IndexToRemove);
 	bool HasPendingUpdates() const { return DefaultContext.ComponentRequiresUpdate.Num() > 0; }
 
-	void UpdateInternal(FNiagaraWorldManager* WorldMan, FComponentIterationContext& Context);
-	bool EvaluateCullState(FNiagaraWorldManager* WorldMan, FComponentIterationContext& Context, int32 ComponentIndex, int32& UpdateCounter);
-	void ProcessSignificance(FNiagaraWorldManager* WorldMan, UNiagaraSignificanceHandler* SignificanceHandler, FComponentIterationContext& Context);
+	void UpdateInternal(FComponentIterationContext& Context);
+	bool EvaluateCullState(FComponentIterationContext& Context, int32 ComponentIndex, int32& UpdateCounter);
+	void ProcessSignificance(UNiagaraSignificanceHandler* SignificanceHandler, FComponentIterationContext& Context);
 	bool ApplyScalabilityState(int32 ComponentIndex, ENiagaraCullReaction CullReaction, bool bNewOnly);
 
 	FNiagaraScalabilitySystemData& GetSystemData(int32 ComponentIndex, bool bForceRefresh=false);
diff --git a/Engine/Plugins/FX/Niagara/Source/Niagara/Public/NiagaraScalabilityState.h b/Engine/Plugins/FX/Niagara/Source/Niagara/Public/NiagaraScalabilityState.h
--- a/Engine/Plugins/FX/Niagara/Source/Niagara/Public/NiagaraScalabilityState.h
+++ b/Engine/Plugins/FX/Niagara/Source/Niagara/Public/NiagaraScalabilityState.h
@@ -74,3 +74,48 @@ struct FNiagaraScalabilityState
 	UPROPERTY(VisibleAnywhere, Category = "Scalability")
 	uint8 bCulledByGlobalBudget : 1;
 };
+
+#if WITH_NIAGARA_DEBUGGER
+/** 
+Debug tracking for scalability culling that kills systems rather than having them sleep.
+Displayed in the debug hud per system to help track down culling issues.
+*/
+struct FNiagaraSystemScalabilityKillCounts
+{
+	int32 Total = 0;
+	int32 Distance = 0;
+	int32 InstanceCounts = 0;
+	int32 Visibility = 0;
+	int32 GlobalBudget = 0;
+
+	inline void Reset()
+	{
+		Total = 0;
+		Distance = 0;
+		InstanceCounts = 0;
+		Visibility = 0;
+		GlobalBudget = 0;
+	}
+
+	inline void Add(const FNiagaraScalabilityState& State)
+	{
+		if (State.bCulled)
+		{
+			++Total;
+			Distance += State.bCulledByDistance ? 1 : 0;
+			InstanceCounts += State.bCulledByInstanceCount ? 1 : 0;
+			Visibility += State.bCulledByVisibility ? 1 : 0;
+			GlobalBudget += State.bCulledByGlobalBudget ? 1 : 0;
+		}
+	}
+
+	inline void Add(const FNiagaraSystemScalabilityKillCounts& InCounts)
+	{
+		Total += InCounts.Total;
+		Distance += InCounts.Distance;
+		InstanceCounts += InCounts.InstanceCounts;
+		Visibility += InCounts.Visibility;
+		GlobalBudget += InCounts.GlobalBudget;
+	}
+};
+#endif
\ No newline at end of file
diff --git a/Engine/Plugins/FX/Niagara/Source/Niagara/Public/NiagaraWorldManager.h b/Engine/Plugins/FX/Niagara/Source/Niagara/Public/NiagaraWorldManager.h
--- a/Engine/Plugins/FX/Niagara/Source/Niagara/Public/NiagaraWorldManager.h
+++ b/Engine/Plugins/FX/Niagara/Source/Niagara/Public/NiagaraWorldManager.h
@@ -251,6 +251,8 @@ public:
 #if WITH_NIAGARA_DEBUGGER
 	class FNiagaraDebugHud* GetNiagaraDebugHud() { return NiagaraDebugHud.Get(); }
 	NIAGARA_API void AddDebugHUDMessage(FName Key, const FNiagaraDebugMessage& Message);
+	FNiagaraSystemScalabilityKillCounts& GetSystemScalabilityKillCounts(UNiagaraSystem* System) { return SystemScalabilityKillCounts.FindOrAdd(System); }
+	const TMap<TWeakObjectPtr<UNiagaraSystem>, FNiagaraSystemScalabilityKillCounts>& GetSystemScalabilityKillCounts() const{ return SystemScalabilityKillCounts; }
 #endif
 
 	class FNiagaraDeferredMethodQueue& GetDeferredMethodQueue() { return DeferredMethods; }
@@ -449,6 +451,7 @@ private:
 
 #if WITH_NIAGARA_DEBUGGER
 	TUniquePtr<class FNiagaraDebugHud> NiagaraDebugHud;
+	TMap<TWeakObjectPtr<UNiagaraSystem>, FNiagaraSystemScalabilityKillCounts> SystemScalabilityKillCounts;
 #endif
 #if WITH_NIAGARA_LEAK_DETECTOR
 	TUniquePtr<class FNiagaraComponentLeakDetector> ComponentLeakDetector;
