- Add missing CanBind to user parameters
- Change style of binding to be more consistent with other bindings

diff --git a/Engine/Plugins/FX/Niagara/Source/NiagaraEditor/Internal/Customizations/NiagaraParameterBindingAdapterCustomization.cpp b/Engine/Plugins/FX/Niagara/Source/NiagaraEditor/Internal/Customizations/NiagaraParameterBindingAdapterCustomization.cpp
--- a/Engine/Plugins/FX/Niagara/Source/NiagaraEditor/Internal/Customizations/NiagaraParameterBindingAdapterCustomization.cpp
+++ b/Engine/Plugins/FX/Niagara/Source/NiagaraEditor/Internal/Customizations/NiagaraParameterBindingAdapterCustomization.cpp
@@ -293,12 +293,19 @@ void FNiagaraParameterBindingAdapterCustomization::CustomizeHeader(TSharedRef<IP
 	ParameterPanel->AddSlot()
 	.AutoWidth()
 	[
-		SNew(SNiagaraParameterName)
-		.ParameterName(this, &FNiagaraParameterBindingAdapterCustomization::GetBoundParameterName)
-		.IsReadOnly(true)
+		SNew(SComboButton)
+		.ContentPadding(2)
 		.IsEnabled(this, &FNiagaraParameterBindingAdapterCustomization::IsBindingEnabled)
 		.Visibility(this, &FNiagaraParameterBindingAdapterCustomization::IsBindingVisibile)
-		.ToolTipText(this, &FNiagaraParameterBindingAdapterCustomization::GetTooltipText)
+		.HasDownArrow(false)
+		.IsFocusable(false)
+		.ButtonContent()
+		[
+			SNew(SNiagaraParameterName)
+			.ParameterName(this, &FNiagaraParameterBindingAdapterCustomization::GetBoundParameterName)
+			.IsReadOnly(true)
+			.ToolTipText(this, &FNiagaraParameterBindingAdapterCustomization::GetTooltipText)
+		]
 	];
 
 	HeaderRow
@@ -318,7 +325,7 @@ void FNiagaraParameterBindingAdapterCustomization::CustomizeHeader(TSharedRef<IP
 		.HAlign(HAlign_Right)
 		.Padding(3, 0, 0, 0)
 		[
-			SNew(SComboButton)					
+			SNew(SComboButton)
 			.ButtonStyle(FAppStyle::Get(), "HoverHintOnly")
 			.ForegroundColor(FSlateColor::UseForeground())
 			.OnGetMenuContent(this, &FNiagaraParameterBindingAdapterCustomization::OnGetMenuContent)
@@ -561,6 +568,11 @@ void FNiagaraParameterBindingAdapterCustomization::CollectAllActions(FGraphActio
 		}
 		VisitedVariables.Add(Variable);
 
+		if (CanBindToType(ParameterBinding, Variable.GetType()) == false)
+		{
+			continue;
+		}
+
 		ParameterBinding->CollectBindings(Variable, Variable, AllBindings);
 	}
 
