- Experimental stateless emitter template building
- Only execute on the CPU path, no compilation for GPU shaders yet
- As this is experimental the data itself may change

diff --git a/Engine/Plugins/FX/Niagara/Source/Niagara/Internal/Stateless/NiagaraStatelessEmitter.h b/Engine/Plugins/FX/Niagara/Source/Niagara/Internal/Stateless/NiagaraStatelessEmitter.h
--- a/Engine/Plugins/FX/Niagara/Source/Niagara/Internal/Stateless/NiagaraStatelessEmitter.h
+++ b/Engine/Plugins/FX/Niagara/Source/Niagara/Internal/Stateless/NiagaraStatelessEmitter.h
@@ -34,6 +34,8 @@ class UNiagaraStatelessEmitter : public UObject//, public INiagaraParameterDefin
 {
 	GENERATED_BODY()
 
+	friend class UNiagaraStatelessEmitterTemplate;
+
 public:
 	DECLARE_MULTICAST_DELEGATE(FOnTemplateChanged);
 
@@ -109,7 +111,7 @@ public:
 	NIAGARA_API bool IsAllowedByScalability() const;
 
 #if WITH_EDITOR
-	NIAGARA_API void SetEmitterTemplateClass(UClass* TemplateClass);
+	NIAGARA_API void SetEmitterTemplate(UNiagaraStatelessEmitterTemplate* Template);
 
 	NIAGARA_API void AddRenderer(UNiagaraRendererProperties* Renderer, FGuid EmitterVersion);
 	NIAGARA_API void RemoveRenderer(UNiagaraRendererProperties* Renderer, FGuid EmitterVersion);
@@ -149,8 +151,13 @@ protected:
 	UPROPERTY()
 	FString UniqueEmitterName;
 
-	UPROPERTY(EditAnywhere, Category = "Emitter Properties", NoClear, meta = (DisplayName = "Template", DisallowCreateNew, AllowedClasses = "/Script/Niagara.NiagaraStatelessEmitterTemplate"))
-	TObjectPtr<UClass> EmitterTemplateClass;
+#if WITH_EDITORONLY_DATA
+	UPROPERTY()
+	TObjectPtr<UClass> EmitterTemplateClass_DEPRECATED;
+#endif
+
+	UPROPERTY(EditAnywhere, Category = "Emitter Properties")
+	TObjectPtr<UNiagaraStatelessEmitterTemplate> EmitterTemplate;
 
 	UPROPERTY(EditAnywhere, Category = "Emitter Properties")
 	uint32 bDeterministic : 1 = false;
diff --git a/Engine/Plugins/FX/Niagara/Source/Niagara/Internal/Stateless/NiagaraStatelessEmitterTemplate.h b/Engine/Plugins/FX/Niagara/Source/Niagara/Internal/Stateless/NiagaraStatelessEmitterTemplate.h
--- a/Engine/Plugins/FX/Niagara/Source/Niagara/Internal/Stateless/NiagaraStatelessEmitterTemplate.h
+++ b/Engine/Plugins/FX/Niagara/Source/Niagara/Internal/Stateless/NiagaraStatelessEmitterTemplate.h
@@ -10,7 +10,7 @@
 
 #include "NiagaraStatelessEmitterTemplate.generated.h"
 
-UCLASS(MinimalAPI, abstract)
+UCLASS(MinimalAPI)
 class UNiagaraStatelessEmitterTemplate : public UObject
 {
 	GENERATED_BODY()
@@ -18,31 +18,35 @@ class UNiagaraStatelessEmitterTemplate : public UObject
 public:
 	//~Begin: UObject Interface
 	NIAGARA_API virtual void PostInitProperties() override;
+	NIAGARA_API virtual void PostLoad() override;
 	//~End: UObject Interface
 
 	static void InitCDOPropertiesAfterModuleStartup();
-	virtual void InitModulesAndAttributes() {}
+	virtual void InitModulesAndAttributes();
+#if WITH_EDITOR
+	NIAGARA_API void ModifyTemplate(UNiagaraStatelessEmitterTemplate* FromTemplate);
+#endif
 
 	TConstArrayView<TObjectPtr<UClass>> GetModules() const { return Modules; }
 #if WITH_EDITORONLY_DATA
 	TConstArrayView<FNiagaraVariableBase> GetOututputComponents() const { return OutputComponents; }
 #endif
 
-	virtual const FShaderParametersMetadata* GetShaderParametersMetadata() const { checkNoEntry(); return nullptr; }
-	virtual TShaderRef<NiagaraStateless::FSimulationShader> GetSimulationShader() const { checkNoEntry(); return TShaderRef<NiagaraStateless::FSimulationShader>(); }
-
-	//-TODO: Add a way to set them directly, we should know that the final struct is a series of ints in the order of the provided variables
-	virtual void SetShaderParameters(uint8* ShaderParametersBase, TConstArrayView<int32> ComponentOffsets) const { checkNoEntry(); }
-	//-TODO: Add a way to set them directly, we should know that the final struct is a series of ints in the order of the provided variables
+	virtual const FShaderParametersMetadata* GetShaderParametersMetadata() const { check(ShaderParametersMetadata.IsValid()); return ShaderParametersMetadata.Get(); }
+	virtual TShaderRef<NiagaraStateless::FSimulationShader> GetSimulationShader() const { return TShaderRef<NiagaraStateless::FSimulationShader>(); }
+	virtual void SetShaderParameters(uint8* ShaderParametersBase, TConstArrayView<int32> ComponentOffsets) const;
 
 protected:
-	TArray<TObjectPtr<UClass>>		Modules;
+	UPROPERTY(EditAnywhere, Category="Modules", meta = (DisallowCreateNew, AllowedClasses = "/Script/Niagara.NiagaraStatelessModule"))
+	TArray<TObjectPtr<UClass>> Modules;
+
 #if WITH_EDITORONLY_DATA
-	TArray<FNiagaraVariableBase>	OutputComponents;
+	TArray<FNiagaraVariableBase> OutputComponents;
 #endif
+	TSharedPtr<FShaderParametersMetadata> ShaderParametersMetadata;
 };
 
-UCLASS(Category = "Statless", DisplayName = "Emitter Example 0")
+UCLASS(Category = "Statless", DisplayName = "Default Template")
 class UNiagaraStatelessEmitterDefault : public UNiagaraStatelessEmitterTemplate
 {
 	GENERATED_BODY()
diff --git a/Engine/Plugins/FX/Niagara/Source/Niagara/Internal/Stateless/NiagaraStatelessShaderParametersBuilder.h b/Engine/Plugins/FX/Niagara/Source/Niagara/Internal/Stateless/NiagaraStatelessShaderParametersBuilder.h
--- a/Engine/Plugins/FX/Niagara/Source/Niagara/Internal/Stateless/NiagaraStatelessShaderParametersBuilder.h
+++ b/Engine/Plugins/FX/Niagara/Source/Niagara/Internal/Stateless/NiagaraStatelessShaderParametersBuilder.h
@@ -3,11 +3,16 @@
 #pragma once
 
 #include "NiagaraStatelessCommon.h"
-#include "ShaderParameterStruct.h"
+#include "ShaderParameterMetadataBuilder.h"
 
 class FNiagaraStatelessShaderParametersBuilder
 {
 public:
+	explicit FNiagaraStatelessShaderParametersBuilder(FShaderParametersMetadataBuilder* OptionalMetadataBuilder = nullptr)
+		: MetadataBuilder(OptionalMetadataBuilder)
+	{
+	}
+
 	// Find the current size of the shader parameter structure
 	int32 GetParametersStructSize() const
 	{
@@ -18,10 +23,17 @@ public:
 	// i.e. if the structured contained "MyFloat" the shader variable would be "UniqueDataInterfaceName_MyFloat"
 	template<typename T> void AddParameterNestedStruct()
 	{
-		const int32 StructOffset = Align(ParameterOffset, TShaderParameterStructTypeInfo<T>::Alignment);
-		ParameterOffset = StructOffset + TShaderParameterStructTypeInfo<T>::GetStructMetadata()->GetSize();
+		using TParamTypeInfo = TShaderParameterStructTypeInfo<T>;
+
+		const int32 StructOffset = Align(ParameterOffset, TParamTypeInfo::Alignment);
+		if (MetadataBuilder)
+		{
+			MetadataBuilder->AddIncludedStruct(TParamTypeInfo::GetStructMetadata());
+		}
+		ParameterOffset = StructOffset + TParamTypeInfo::GetStructMetadata()->GetSize();
 	}
 
 private:
-	int32	ParameterOffset = 0;
+	FShaderParametersMetadataBuilder*	MetadataBuilder = nullptr;
+	int32								ParameterOffset = 0;
 };
diff --git a/Engine/Plugins/FX/Niagara/Source/Niagara/Private/Stateless/NiagaraStatelessComputeManager.cpp b/Engine/Plugins/FX/Niagara/Source/Niagara/Private/Stateless/NiagaraStatelessComputeManager.cpp
--- a/Engine/Plugins/FX/Niagara/Source/Niagara/Private/Stateless/NiagaraStatelessComputeManager.cpp
+++ b/Engine/Plugins/FX/Niagara/Source/Niagara/Private/Stateless/NiagaraStatelessComputeManager.cpp
@@ -61,6 +61,7 @@ namespace NiagaraStatelessComputeManagerPrivate
 		}
 		if (bAllowGPUExec)
 		{
+			check(EmitterData->GetShader().IsValid());
 			return EComputeExecutionPath::GPU;
 		}
 		return EComputeExecutionPath::None;
@@ -232,11 +233,6 @@ FNiagaraDataBuffer* FNiagaraStatelessComputeManager::GetDataBuffer(FRHICommandLi
 	}
 
 	const FNiagaraStatelessEmitterData* EmitterData = EmitterInstance->EmitterData.Get();
-	TShaderRef<NiagaraStateless::FSimulationShader> ComputeShader = EmitterData->GetShader();
-	if (!ComputeShader.IsValid() || !ComputeShader.GetComputeShader())
-	{
-		return nullptr;
-	}
 
 	uint32 ActiveParticles = 0;
 	{
diff --git a/Engine/Plugins/FX/Niagara/Source/Niagara/Private/Stateless/NiagaraStatelessEmitter.cpp b/Engine/Plugins/FX/Niagara/Source/Niagara/Private/Stateless/NiagaraStatelessEmitter.cpp
--- a/Engine/Plugins/FX/Niagara/Source/Niagara/Private/Stateless/NiagaraStatelessEmitter.cpp
+++ b/Engine/Plugins/FX/Niagara/Source/Niagara/Private/Stateless/NiagaraStatelessEmitter.cpp
@@ -68,7 +68,7 @@ namespace NiagaraStatelessInternal
 
 		// Validate the shader is correct
 		const FShaderParametersMetadata* ShaderParametersMetadata = EmitterData.GetShaderParametersMetadata();
-		if (!EmitterData.GetShader().IsValid() || !ShaderParametersMetadata)
+		if (!ShaderParametersMetadata)
 		{
 			return false;
 		}
@@ -117,7 +117,7 @@ void UNiagaraStatelessEmitter::Serialize(FArchive& Ar)
 	// On cooked make sure any null modules are replaced with the CDO
 	if (Ar.IsLoading() && FPlatformProperties::RequiresCookedData())
 	{
-		if (const UNiagaraStatelessEmitterTemplate* EmitterTemplate = GetEmitterTemplate())
+		if (EmitterTemplate)
 		{
 			TConstArrayView<TObjectPtr<UClass>> ModuleClass = EmitterTemplate->GetModules();
 			for (int32 i = 0; i < Modules.Num(); ++i)
@@ -135,6 +135,13 @@ void UNiagaraStatelessEmitter::PostLoad()
 {
 	Super::PostLoad();
 
+#if WITH_EDITORONLY_DATA
+	if (EmitterTemplateClass_DEPRECATED)
+	{
+		EmitterTemplate = Cast<UNiagaraStatelessEmitterTemplate>(EmitterTemplateClass_DEPRECATED->GetDefaultObject());
+	}
+#endif
+
 #if WITH_EDITOR
 	// Ensure our module list is up to date
 	PostTemplateChanged();
@@ -194,10 +201,9 @@ void UNiagaraStatelessEmitter::PostEditChangeProperty(FPropertyChangedEvent& Pro
 		FNiagaraSystemUpdateContext UpdateContext(NiagaraSystem, true);
 
 		// Ensure our template is up to date, right not it's easy to remove things by accident and cause a crash
-		if (PropertyChangedEvent.GetPropertyName() == GET_MEMBER_NAME_CHECKED(UNiagaraStatelessEmitter, EmitterTemplateClass))
+		if (PropertyChangedEvent.GetPropertyName() == GET_MEMBER_NAME_CHECKED(UNiagaraStatelessEmitter, EmitterTemplate))
 		{
 			PostTemplateChanged();
-			OnTemplateChanged.Broadcast();
 		}
 		CacheParameterCollectionReferences();
 	}
@@ -206,10 +212,10 @@ void UNiagaraStatelessEmitter::PostEditChangeProperty(FPropertyChangedEvent& Pro
 void UNiagaraStatelessEmitter::PostTemplateChanged()
 {
 	// No template nothing to do
-	const UNiagaraStatelessEmitterTemplate* EmitterTemplate = GetEmitterTemplate();
 	if (EmitterTemplate == nullptr)
 	{
 		Modules.Empty();
+		OnTemplateChanged.Broadcast();
 		return;
 	}
 
@@ -239,6 +245,7 @@ void UNiagaraStatelessEmitter::PostTemplateChanged()
 		}
 	}
 	Modules = MoveTemp(NewModules);
+	OnTemplateChanged.Broadcast();
 }
 
 void UNiagaraStatelessEmitter::CacheParameterCollectionReferences()
@@ -291,7 +298,7 @@ bool UNiagaraStatelessEmitter::UsesCollection(const UNiagaraParameterCollection*
 
 const UNiagaraStatelessEmitterTemplate* UNiagaraStatelessEmitter::GetEmitterTemplate() const
 {
-	return EmitterTemplateClass ? EmitterTemplateClass->GetDefaultObject<UNiagaraStatelessEmitterTemplate>() : nullptr;
+	return EmitterTemplate;
 }
 
 void UNiagaraStatelessEmitter::CacheFromCompiledData()
@@ -319,6 +326,12 @@ void UNiagaraStatelessEmitter::CacheFromCompiledData()
 		}
 	}
 
+	if (!StatelessEmitterData->GetShader().IsValid())
+	{
+		// No GPU shader we can only execute on the CPU path
+		StatelessEmitterData->FeatureMask &= ~ENiagaraStatelessFeatureMask::ExecuteGPU;
+	}
+
 	if (StatelessEmitterData->FeatureMask == ENiagaraStatelessFeatureMask::None)
 	{
 		StatelessEmitterData->bCanEverExecute = false;
@@ -496,7 +509,7 @@ void UNiagaraStatelessEmitter::BuildCompiledDataSet()
 	ParticleDataSetCompiledData.Empty();
 	ParticleDataSetCompiledData.SimTarget = ENiagaraSimTarget::GPUComputeSim;
 
-	if (const UNiagaraStatelessEmitterTemplate* EmitterTemplate = StatelessEmitterData->EmitterTemplate)
+	if (EmitterTemplate)
 	{
 		// Gather a list of all the output variables from the modules, this can change based on what is enabled / disabled
 		TArray<FNiagaraVariableBase> AvailableVariables;
@@ -698,14 +711,16 @@ bool UNiagaraStatelessEmitter::IsAllowedByScalability() const
 }
 
 #if WITH_EDITOR
-void UNiagaraStatelessEmitter::SetEmitterTemplateClass(UClass* TemplateClass)
+void UNiagaraStatelessEmitter::SetEmitterTemplate(UNiagaraStatelessEmitterTemplate* Template)
 {
-	if (TemplateClass && !TemplateClass->IsChildOf(UNiagaraStatelessEmitterTemplate::StaticClass()))
+	FNiagaraSystemUpdateContext UpdateContext;
+	UpdateContext.SetDestroyOnAdd(true);
+	if (UNiagaraSystem* Owner = GetTypedOuter<UNiagaraSystem>())
 	{
-		TemplateClass = nullptr;
+		UpdateContext.Add(Owner, true);
 	}
 
-	EmitterTemplateClass = TemplateClass;
+	EmitterTemplate = Template;
 	PostTemplateChanged();
 }
 
diff --git a/Engine/Plugins/FX/Niagara/Source/Niagara/Private/Stateless/NiagaraStatelessEmitterData.cpp b/Engine/Plugins/FX/Niagara/Source/Niagara/Private/Stateless/NiagaraStatelessEmitterData.cpp
--- a/Engine/Plugins/FX/Niagara/Source/Niagara/Private/Stateless/NiagaraStatelessEmitterData.cpp
+++ b/Engine/Plugins/FX/Niagara/Source/Niagara/Private/Stateless/NiagaraStatelessEmitterData.cpp
@@ -35,9 +35,8 @@ void FNiagaraStatelessEmitterData::InitRenderResources()
 
 TShaderRef<NiagaraStateless::FSimulationShader>	FNiagaraStatelessEmitterData::GetShader() const
 {
-	if (FApp::CanEverRender())
+	if (FApp::CanEverRender() && EmitterTemplate)
 	{
-		check(EmitterTemplate);
 		return EmitterTemplate->GetSimulationShader();
 	}
 	return TShaderRef<NiagaraStateless::FSimulationShader>();
diff --git a/Engine/Plugins/FX/Niagara/Source/Niagara/Private/Stateless/NiagaraStatelessEmitterTemplate.cpp b/Engine/Plugins/FX/Niagara/Source/Niagara/Private/Stateless/NiagaraStatelessEmitterTemplate.cpp
--- a/Engine/Plugins/FX/Niagara/Source/Niagara/Private/Stateless/NiagaraStatelessEmitterTemplate.cpp
+++ b/Engine/Plugins/FX/Niagara/Source/Niagara/Private/Stateless/NiagaraStatelessEmitterTemplate.cpp
@@ -2,9 +2,12 @@
 
 #include "Stateless/NiagaraStatelessEmitterTemplate.h"
 #include "Stateless/NiagaraStatelessCommon.h"
+#include "Stateless/NiagaraStatelessEmitter.h"
 #include "Stateless/NiagaraStatelessEmitterShaders.h"
+#include "Stateless/NiagaraStatelessShaderParametersBuilder.h"
 #include "NiagaraConstants.h"
 #include "NiagaraModule.h"
+#include "NiagaraSystem.h"
 
 #include "Modules/ModuleManager.h"
 
@@ -35,6 +38,11 @@
 #include "Stateless/Modules/NiagaraStatelessModule_SubUVAnimation.h"
 #include "Stateless/Modules/NiagaraStatelessModule_DynamicMaterialParameters.h"
 
+#if WITH_EDITOR
+#include "UObject/UObjectIterator.h"
+#include "Engine/Engine.h"
+#endif
+
 namespace NiagaraStatelessEmitterTemplateImpl
 {
 	TArray<TWeakObjectPtr<UNiagaraStatelessEmitterTemplate>> ObjectsToDeferredInit;
@@ -47,7 +55,7 @@ void UNiagaraStatelessEmitterTemplate::PostInitProperties()
 	Super::PostInitProperties();
 
 	// We can end up hitting PostInitProperties before the Niagara Module has initialized bindings this needs, mark this object for deferred init and early out.
-	if (GetClass()->HasAnyClassFlags(CLASS_Abstract) == false)
+	if (GetClass() != UNiagaraStatelessEmitterTemplate::StaticClass())
 	{
 		if (FModuleManager::Get().IsModuleLoaded("Niagara") == false)
 		{
@@ -61,6 +69,12 @@ void UNiagaraStatelessEmitterTemplate::PostInitProperties()
 	}
 }
 
+void UNiagaraStatelessEmitterTemplate::PostLoad()
+{
+	Super::PostLoad();
+	InitModulesAndAttributes();
+}
+
 void UNiagaraStatelessEmitterTemplate::InitCDOPropertiesAfterModuleStartup()
 {
 	for (TWeakObjectPtr<UNiagaraStatelessEmitterTemplate>& WeakObject : NiagaraStatelessEmitterTemplateImpl::ObjectsToDeferredInit)
@@ -74,6 +88,77 @@ void UNiagaraStatelessEmitterTemplate::InitCDOPropertiesAfterModuleStartup()
 	NiagaraStatelessEmitterTemplateImpl::ObjectsToDeferredInit.Empty();
 }
 
+void UNiagaraStatelessEmitterTemplate::InitModulesAndAttributes()
+{
+	// Remove any invalid modules
+	Modules.RemoveAll([](UClass* ModuleClass) { return ModuleClass == nullptr || !ModuleClass->IsChildOf(UNiagaraStatelessModule::StaticClass()); });
+
+#if WITH_EDITORONLY_DATA
+	// Find All Output Components
+	OutputComponents.Empty();
+	for (UClass* ModuleClass : Modules)
+	{
+		UNiagaraStatelessModule* Module = CastChecked<UNiagaraStatelessModule>(ModuleClass->GetDefaultObject());
+		Module->GetOutputVariables(OutputComponents);
+	}
+#endif
+
+	// Build Shader Parameters
+	FShaderParametersMetadataBuilder ShaderMetadataBuilder;
+	FNiagaraStatelessShaderParametersBuilder ShaderParametersBuilder(&ShaderMetadataBuilder);
+
+	ShaderParametersBuilder.AddParameterNestedStruct<NiagaraStateless::FCommonShaderParameters>();
+	for (UClass* ModuleClass : Modules)
+	{
+		UNiagaraStatelessModule* Module = CastChecked<UNiagaraStatelessModule>(ModuleClass->GetDefaultObject());
+		Module->BuildShaderParameters(ShaderParametersBuilder);
+	}
+	//-TODO: Reserve space for component outputs when compiling GPU ShaderParametersBuilder.AddParameterNestedStruct<NiagaraStateless::FCommonShaderParameters>();
+
+	ShaderParametersMetadata = MakeShareable<FShaderParametersMetadata>(ShaderMetadataBuilder.Build(FShaderParametersMetadata::EUseCase::ShaderParameterStruct, TEXT("NiagaraStatelessEmitterTemplate")));
+}
+
+#if WITH_EDITOR
+void UNiagaraStatelessEmitterTemplate::ModifyTemplate(UNiagaraStatelessEmitterTemplate* FromTemplate)
+{
+	// Modify for capturing in transaction
+	Modify();
+
+	// Gather all the systems to update
+	FNiagaraSystemUpdateContext UpdateContext;
+
+	TArray<UNiagaraStatelessEmitter*> EmittersToUpdate;
+	for (TObjectIterator<UNiagaraStatelessEmitter> It; It; ++It)
+	{
+		UNiagaraStatelessEmitter* Emitter = *It;
+		if (Emitter && Emitter->GetEmitterTemplate() == this)
+		{
+			EmittersToUpdate.Add(Emitter);
+
+			if (UNiagaraSystem* System = Emitter->GetTypedOuter<UNiagaraSystem>())
+			{
+				UpdateContext.Add(System, true);
+			}
+		}
+	}
+
+	// Copy and initialize the template
+	UEngine::CopyPropertiesForUnrelatedObjects(FromTemplate, this);
+	InitModulesAndAttributes();
+
+	for (UNiagaraStatelessEmitter* Emitter : EmittersToUpdate)
+	{
+		Emitter->PostTemplateChanged();
+		Emitter->CacheFromCompiledData();
+	}
+}
+#endif
+
+void UNiagaraStatelessEmitterTemplate::SetShaderParameters(uint8* ShaderParametersBase, TConstArrayView<int32> ComponentOffsets) const
+{
+	//-TODO: We don't store component offsets right now
+}
+
 //////////////////////////////////////////////////////////////////////////////////////////////////////////////////////////////////////////////////////////
 
 void UNiagaraStatelessEmitterDefault::InitModulesAndAttributes()
diff --git a/Engine/Plugins/FX/Niagara/Source/NiagaraEditor/Internal/ViewModels/NiagaraStatelessEmitterTemplateViewModel.h b/Engine/Plugins/FX/Niagara/Source/NiagaraEditor/Internal/ViewModels/NiagaraStatelessEmitterTemplateViewModel.h
new file mode 100644
--- /dev/null
+++ b/Engine/Plugins/FX/Niagara/Source/NiagaraEditor/Internal/ViewModels/NiagaraStatelessEmitterTemplateViewModel.h
@@ -0,0 +1,37 @@
+// Copyright Epic Games, Inc. All Rights Reserved.
+
+#pragma once
+
+#include "Stateless/NiagaraStatelessCommon.h"
+
+#include "Misc/NotifyHook.h"
+#include "UObject/StrongObjectPtr.h"
+
+class UNiagaraStatelessEmitterTemplate;
+
+class FNiagaraStatelessEmitterTemplateViewModel : public TSharedFromThis<FNiagaraStatelessEmitterTemplateViewModel>, public FNotifyHook
+{
+public:
+	DECLARE_MULTICAST_DELEGATE(FOnTemplateChanged)
+
+	explicit FNiagaraStatelessEmitterTemplateViewModel(UNiagaraStatelessEmitterTemplate* EmitterTemplate);
+	virtual ~FNiagaraStatelessEmitterTemplateViewModel();
+
+	void PostUndoRedo();
+
+	//- Begin: FNotifyHook
+	virtual void NotifyPreChange(FProperty* PropertyAboutToChange) override;
+	virtual void NotifyPostChange(const FPropertyChangedEvent& PropertyChangedEvent, FProperty* PropertyThatChanged) override;
+	//- End: FNotifyHook
+
+	UNiagaraStatelessEmitterTemplate* GetTemplateForEdit() { return EmitterTemplateForEdit.Get(); }
+	TConstArrayView<FNiagaraVariableBase> GetOututputComponents() const;
+	TPair<ENiagaraStatelessFeatureMask, ENiagaraStatelessFeatureMask> GetFeatureMaskRange() const;
+
+	FOnTemplateChanged& OnTemplateChanged() { return OnTemplateChangedDelegate; }
+
+private:
+	TStrongObjectPtr<UNiagaraStatelessEmitterTemplate>	EmitterTemplate;
+	TStrongObjectPtr<UNiagaraStatelessEmitterTemplate>	EmitterTemplateForEdit;
+	FOnTemplateChanged									OnTemplateChangedDelegate;
+};
diff --git a/Engine/Plugins/FX/Niagara/Source/NiagaraEditor/Private/AssetDefinitions/AssetDefinition_NiagaraStatelessEmitterTemplate.cpp b/Engine/Plugins/FX/Niagara/Source/NiagaraEditor/Private/AssetDefinitions/AssetDefinition_NiagaraStatelessEmitterTemplate.cpp
new file mode 100644
--- /dev/null
+++ b/Engine/Plugins/FX/Niagara/Source/NiagaraEditor/Private/AssetDefinitions/AssetDefinition_NiagaraStatelessEmitterTemplate.cpp
@@ -0,0 +1,50 @@
+// Copyright Epic Games, Inc. All Rights Reserved.
+
+#include "AssetDefinition_NiagaraStatelessEmitterTemplate.h"
+#include "Stateless/NiagaraStatelessEmitterTemplate.h"
+#include "Toolkits/NiagaraStatelessEmitterTemplateToolkit.h"
+
+#include "ContentBrowserMenuContexts.h"
+#include "DesktopPlatformModule.h"
+#include "EditorDirectories.h"
+#include "NiagaraEditorStyle.h"
+#include "Framework/Application/SlateApplication.h"
+#include "Framework/Notifications/NotificationManager.h"
+#include "ToolMenus.h"
+#include "Widgets/Notifications/SNotificationList.h"
+
+#define LOCTEXT_NAMESPACE "AssetDefinition_NiagaraStatelessEmitterTemplate"
+
+FText UAssetDefinition_NiagaraStatelessEmitterTemplate::GetAssetDisplayName() const
+{
+	return NSLOCTEXT("AssetTypeActions", "AssetTypeActions_NiagaraStatelessEmitterTemplate", "Niagara Lightweight Emitter Template");
+}
+
+FLinearColor UAssetDefinition_NiagaraStatelessEmitterTemplate::GetAssetColor() const
+{
+	return FNiagaraEditorStyle::Get().GetColor("NiagaraEditor.AssetColors.StatelessEmitterTemplate").ToFColor(true);
+}
+
+TSoftClassPtr<UObject> UAssetDefinition_NiagaraStatelessEmitterTemplate::GetAssetClass() const
+{
+	return UNiagaraStatelessEmitterTemplate::StaticClass();
+}
+
+TConstArrayView<FAssetCategoryPath> UAssetDefinition_NiagaraStatelessEmitterTemplate::GetAssetCategories() const
+{
+	static const auto Categories = { EAssetCategoryPaths::FX / NSLOCTEXT("Niagara", "NiagaraAssetSubMenu_Advanced", "Advanced") };
+	return Categories;
+}
+
+EAssetCommandResult UAssetDefinition_NiagaraStatelessEmitterTemplate::OpenAssets(const FAssetOpenArgs& OpenArgs) const
+{
+	for (UNiagaraStatelessEmitterTemplate* EmitterTemplate : OpenArgs.LoadObjects<UNiagaraStatelessEmitterTemplate>())
+	{
+		const TSharedRef<FNiagaraStatelessEmitterTemplateToolkit> NewToolkit(new FNiagaraStatelessEmitterTemplateToolkit());
+		NewToolkit->Initialize(OpenArgs.GetToolkitMode(), OpenArgs.ToolkitHost, EmitterTemplate);
+	}
+
+	return EAssetCommandResult::Handled;
+}
+
+#undef LOCTEXT_NAMESPACE
diff --git a/Engine/Plugins/FX/Niagara/Source/NiagaraEditor/Private/AssetDefinitions/AssetDefinition_NiagaraStatelessEmitterTemplate.h b/Engine/Plugins/FX/Niagara/Source/NiagaraEditor/Private/AssetDefinitions/AssetDefinition_NiagaraStatelessEmitterTemplate.h
new file mode 100644
--- /dev/null
+++ b/Engine/Plugins/FX/Niagara/Source/NiagaraEditor/Private/AssetDefinitions/AssetDefinition_NiagaraStatelessEmitterTemplate.h
@@ -0,0 +1,23 @@
+// Copyright Epic Games, Inc. All Rights Reserved.
+
+#pragma once
+
+#include "CoreMinimal.h"
+#include "AssetDefinitionDefault.h"
+
+#include "AssetDefinition_NiagaraStatelessEmitterTemplate.generated.h"
+
+UCLASS()
+class UAssetDefinition_NiagaraStatelessEmitterTemplate : public UAssetDefinitionDefault
+{
+	GENERATED_BODY()
+
+public:
+	// UAssetDefinition Begin
+	virtual FText GetAssetDisplayName() const override;
+	virtual FLinearColor GetAssetColor() const override;
+	virtual TSoftClassPtr<UObject> GetAssetClass() const override;
+	virtual TConstArrayView<FAssetCategoryPath> GetAssetCategories() const override;
+	virtual EAssetCommandResult OpenAssets(const FAssetOpenArgs& OpenArgs) const override;
+	// UAssetDefinition End
+};
diff --git a/Engine/Plugins/FX/Niagara/Source/NiagaraEditor/Private/NiagaraEditorStyle.cpp b/Engine/Plugins/FX/Niagara/Source/NiagaraEditor/Private/NiagaraEditorStyle.cpp
--- a/Engine/Plugins/FX/Niagara/Source/NiagaraEditor/Private/NiagaraEditorStyle.cpp
+++ b/Engine/Plugins/FX/Niagara/Source/NiagaraEditor/Private/NiagaraEditorStyle.cpp
@@ -498,6 +498,7 @@ void FNiagaraEditorStyle::InitAssetColors()
 	Set("NiagaraEditor.AssetColors.ParameterDefinitions", FLinearColor(0.57f, 0.82f, 0.06f));
 	Set("NiagaraEditor.AssetColors.EffectType", FLinearColor(1.f, 1.f, 1.f));
 	Set("NiagaraEditor.AssetColors.SimCache", FLinearColor(0.9f, 0.3f, 0.25f));
+	Set("NiagaraEditor.AssetColors.StatelessEmitterTemplate", FLinearColor(0.9f, 0.3f, 0.25f));
 	Set("NiagaraEditor.AssetColors.ValidationRuleSet", FLinearColor(0.25f, 0.25f, 0.25f));
 	Set("NiagaraEditor.AssetColors.DataChannelDefinitions", FLinearColor(0.5f, 0.0f, 1.0f));
 	Set("NiagaraEditor.AssetColors.StatelessConversionData", FLinearColor(1.0f, 1.0f, 1.0f));
diff --git a/Engine/Plugins/FX/Niagara/Source/NiagaraEditor/Private/NiagaraStatelessEmitterTemplateFactoryNew.cpp b/Engine/Plugins/FX/Niagara/Source/NiagaraEditor/Private/NiagaraStatelessEmitterTemplateFactoryNew.cpp
new file mode 100644
--- /dev/null
+++ b/Engine/Plugins/FX/Niagara/Source/NiagaraEditor/Private/NiagaraStatelessEmitterTemplateFactoryNew.cpp
@@ -0,0 +1,36 @@
+// Copyright Epic Games, Inc. All Rights Reserved.
+
+#include "NiagaraStatelessEmitterTemplateFactoryNew.h"
+#include "Stateless/NiagaraStatelessEmitterTemplate.h"
+
+#include UE_INLINE_GENERATED_CPP_BY_NAME(NiagaraStatelessEmitterTemplateFactoryNew)
+
+namespace NiagaraStatelessEmitterTemplateFactoryNewPrivate
+{
+	static bool GbAllowTemplateCreation = false;
+	static FAutoConsoleVariableRef CVarAllowTemplateCreation(
+		TEXT("fx.NiagaraStateless.AllowTemplateCreation"),
+		GbAllowTemplateCreation,
+		TEXT("Allow experimental template creation.  Data may be lost."),
+		ECVF_Default
+	);
+}
+
+UNiagaraStatelessEmitterTemplateFactoryNew::UNiagaraStatelessEmitterTemplateFactoryNew(const FObjectInitializer& ObjectInitializer)
+{
+	SupportedClass = UNiagaraStatelessEmitterTemplate::StaticClass();
+	bCreateNew = true;
+	bEditAfterNew = true;
+}
+
+bool UNiagaraStatelessEmitterTemplateFactoryNew::ShouldShowInNewMenu() const
+{
+	using namespace NiagaraStatelessEmitterTemplateFactoryNewPrivate;
+	return GbAllowTemplateCreation;
+}
+
+UObject* UNiagaraStatelessEmitterTemplateFactoryNew::FactoryCreateNew(UClass* InClass, UObject* InParent, FName InName, EObjectFlags Flags, UObject* Context, FFeedbackContext* Warn)
+{
+	return NewObject<UNiagaraStatelessEmitterTemplate>(InParent, InClass, InName, Flags);
+}
+
diff --git a/Engine/Plugins/FX/Niagara/Source/NiagaraEditor/Private/NiagaraStatelessEmitterTemplateFactoryNew.h b/Engine/Plugins/FX/Niagara/Source/NiagaraEditor/Private/NiagaraStatelessEmitterTemplateFactoryNew.h
new file mode 100644
--- /dev/null
+++ b/Engine/Plugins/FX/Niagara/Source/NiagaraEditor/Private/NiagaraStatelessEmitterTemplateFactoryNew.h
@@ -0,0 +1,15 @@
+// Copyright Epic Games, Inc. All Rights Reserved.
+
+#pragma once
+
+#include "Factories/Factory.h"
+#include "NiagaraStatelessEmitterTemplateFactoryNew.generated.h"
+
+UCLASS(hidecategories=Object, MinimalAPI)
+class UNiagaraStatelessEmitterTemplateFactoryNew : public UFactory
+{
+	GENERATED_UCLASS_BODY()
+
+	virtual bool ShouldShowInNewMenu() const override;
+	virtual UObject* FactoryCreateNew(UClass* InClass, UObject* InParent, FName InName, EObjectFlags Flags, UObject* Context, FFeedbackContext* Warn) override;
+};
diff --git a/Engine/Plugins/FX/Niagara/Source/NiagaraEditor/Private/Stateless/NiagaraStatelessEmitterDetails.cpp b/Engine/Plugins/FX/Niagara/Source/NiagaraEditor/Private/Stateless/NiagaraStatelessEmitterDetails.cpp
--- a/Engine/Plugins/FX/Niagara/Source/NiagaraEditor/Private/Stateless/NiagaraStatelessEmitterDetails.cpp
+++ b/Engine/Plugins/FX/Niagara/Source/NiagaraEditor/Private/Stateless/NiagaraStatelessEmitterDetails.cpp
@@ -1,14 +1,20 @@
 // Copyright Epic Games, Inc. All Rights Reserved.
 
 #include "Stateless/NiagaraStatelessEmitterDetails.h"
+#include "Stateless/NiagaraStatelessEmitterTemplate.h"
 #include "Stateless/NiagaraStatelessEmitter.h"
 #include "NiagaraSystem.h"
 
+#include "Algo/Sort.h"
+#include "AssetRegistry/AssetData.h"
+#include "AssetRegistry/AssetRegistryModule.h"
+#include "Framework/MultiBox/MultiBoxBuilder.h"
+#include "UObject/UObjectIterator.h"
+#include "Widgets/Input/SComboButton.h"
 #include "DetailLayoutBuilder.h"
 #include "DetailCategoryBuilder.h"
-//#include "IDetailChildrenBuilder.h"
-//#include "IDetailGroup.h"
-//#include "DetailWidgetRow.h"
+#include "DetailWidgetRow.h"
+#include "ScopedTransaction.h"
 
 #define LOCTEXT_NAMESPACE "NiagaraStatelessEmitter"
 
@@ -33,6 +39,7 @@ void FNiagaraStatelessEmitterDetails::CustomizeDetails(IDetailLayoutBuilder& Det
 
 	static const FName NAME_EmitterProperties("Emitter Properties");
 	static const FName NAME_FixedBounds("FixedBounds");
+	static const FName NAME_EmitterTemplate("EmitterTemplate");
 
 	IDetailCategoryBuilder& DetailCategory = DetailBuilder.EditCategory(NAME_EmitterProperties);
 
@@ -47,6 +54,34 @@ void FNiagaraStatelessEmitterDetails::CustomizeDetails(IDetailLayoutBuilder& Det
 			IDetailPropertyRow& PropertyRow = DetailCategory.AddProperty(PropertyHandle);
 			PropertyRow.IsEnabled(TAttribute<bool>::CreateSP(this, &FNiagaraStatelessEmitterDetails::GetFixedBoundsEnabled));
 		}
+		else if (PropertyName == NAME_EmitterTemplate)
+		{
+			IDetailPropertyRow& PropertyRow = DetailCategory.AddProperty(PropertyHandle);
+			PropertyRow.CustomWidget()
+			.NameContent()
+			[
+				PropertyHandle->CreatePropertyNameWidget()
+			]
+			.ValueContent()
+			[
+				SNew(SHorizontalBox)
+				+ SHorizontalBox::Slot()
+				.AutoWidth()
+				.Padding(0.0f, 0.0f, 4.0f, 0.0f)
+				[
+					SNew(SComboButton)
+					.OnGetMenuContent(this, &FNiagaraStatelessEmitterDetails::OnGetTemplateActions)
+					.ContentPadding(1)
+					.ToolTipText(LOCTEXT("SelectEmitterTemplate", "Select the template to use for this emitter."))
+					.ButtonContent()
+					[
+						SNew(STextBlock)
+						.Font(IDetailLayoutBuilder::GetDetailFont())
+						.Text(this, &FNiagaraStatelessEmitterDetails::GetSelectedTemplateName)
+					]
+				]
+			];
+		}
 		else
 		{
 			DetailCategory.AddProperty(PropertyHandle);
@@ -66,4 +101,107 @@ bool FNiagaraStatelessEmitterDetails::GetFixedBoundsEnabled() const
 	return false;
 }
 
+FText FNiagaraStatelessEmitterDetails::GetSelectedTemplateName() const
+{
+	if (const UNiagaraStatelessEmitter* Emitter = WeakEmitter.Get())
+	{
+		if (const UNiagaraStatelessEmitterTemplate* EmitterTemplate = Emitter->GetEmitterTemplate())
+		{
+			return EmitterTemplate->IsAsset() ? FText::FromName(EmitterTemplate->GetFName()) : EmitterTemplate->GetClass()->GetDisplayNameText();
+		}
+	}
+	return FText();
+}
+
+void FNiagaraStatelessEmitterDetails::SetTemplate(UNiagaraStatelessEmitterTemplate* Template) const
+{
+	UNiagaraStatelessEmitter* Emitter = WeakEmitter.Get();
+	if (Template == nullptr || Emitter == nullptr)
+	{
+		return;
+	}
+
+	FScopedTransaction ScopedTransaction(LOCTEXT("SetEmitterTemplate", "Set Emitter Template"));
+	Emitter->Modify();
+	Emitter->SetEmitterTemplate(Template);
+}
+
+void FNiagaraStatelessEmitterDetails::SetTemplate(FSoftObjectPath TemplatePath) const
+{
+	if (UNiagaraStatelessEmitterTemplate* Template = Cast<UNiagaraStatelessEmitterTemplate>(TemplatePath.TryLoad()))
+	{
+		SetTemplate(Template);
+	}
+}
+
+TSharedRef<SWidget> FNiagaraStatelessEmitterDetails::OnGetTemplateActions() const
+{
+	struct FMenuAction
+	{
+		FText		Label;
+		FText		ToolTip;
+		FUIAction	Action;
+	};
+	TArray<FMenuAction> Actions;
+
+	// Find all Build In Types
+	{
+		TSet<UClass*> VisitedClasses;
+		for (TObjectIterator<UClass> It; It; ++It)
+		{
+			UClass* TemplateClass = *It;
+			if (TemplateClass && TemplateClass != UNiagaraStatelessEmitterTemplate::StaticClass() && TemplateClass->IsChildOf(UNiagaraStatelessEmitterTemplate::StaticClass()))
+			{
+				const FText DisplayName = TemplateClass->GetDisplayNameText();
+				const FText AssetName = FText::FromName(TemplateClass->GetFName());
+				Actions.Emplace(
+					DisplayName,
+					FText::Format(LOCTEXT("SetBuiltInTemplate", "Set to the built in template '{0}' class '{1}'"), DisplayName, AssetName),
+					FUIAction(FExecuteAction::CreateSP(this, &FNiagaraStatelessEmitterDetails::SetTemplate, Cast<UNiagaraStatelessEmitterTemplate>(TemplateClass->GetDefaultObject())))
+				);
+			}
+		}
+	}
+
+	// Find all Template Assets
+	{
+		FAssetRegistryModule& AssetRegistryModule = FModuleManager::LoadModuleChecked<FAssetRegistryModule>("AssetRegistry");
+
+		FARFilter Filter;
+		Filter.bRecursiveClasses = true;
+		Filter.ClassPaths.Add(FTopLevelAssetPath(UNiagaraStatelessEmitterTemplate::StaticClass()));
+
+		TArray<FAssetData> TemplateAssetDatas;
+		AssetRegistryModule.Get().GetAssets(Filter, TemplateAssetDatas);
+
+		for (const FAssetData& AssetData : TemplateAssetDatas)
+		{
+			const FSoftObjectPath ObjectPath = AssetData.GetSoftObjectPath();
+			const FText DisplayName = FText::FromString(ObjectPath.GetAssetName());
+			const FText AssetName = FText::FromString(ObjectPath.ToString());
+			Actions.Emplace(
+				DisplayName,
+				FText::Format(LOCTEXT("SetAssetTemplate", "Set to template '{0}' asset path '{1}'"), DisplayName, AssetName),
+				FUIAction(FExecuteAction::CreateSP(this, &FNiagaraStatelessEmitterDetails::SetTemplate, ObjectPath))
+			);
+		}
+	}
+
+	// Sort Options
+	Algo::SortBy(Actions, [](const FMenuAction& Item) { return Item.Label; }, FText::FSortPredicate());
+
+	// Create Menu
+	FMenuBuilder MenuBuilder(true, NULL);
+	for (const FMenuAction& Action : Actions)
+	{
+		MenuBuilder.AddMenuEntry(
+			Action.Label,
+			Action.ToolTip,
+			FSlateIcon(),
+			Action.Action
+		);
+	}
+	return MenuBuilder.MakeWidget();
+}
+
 #undef LOCTEXT_NAMESPACE
diff --git a/Engine/Plugins/FX/Niagara/Source/NiagaraEditor/Private/Stateless/NiagaraStatelessEmitterDetails.h b/Engine/Plugins/FX/Niagara/Source/NiagaraEditor/Private/Stateless/NiagaraStatelessEmitterDetails.h
--- a/Engine/Plugins/FX/Niagara/Source/NiagaraEditor/Private/Stateless/NiagaraStatelessEmitterDetails.h
+++ b/Engine/Plugins/FX/Niagara/Source/NiagaraEditor/Private/Stateless/NiagaraStatelessEmitterDetails.h
@@ -6,16 +6,28 @@
 #include "UObject/WeakObjectPtr.h"
 
 class IDetailLayoutBuilder;
+class SWidget;
+
 class UNiagaraStatelessEmitter;
+class UNiagaraStatelessEmitterTemplate;
 
 class FNiagaraStatelessEmitterDetails : public IDetailCustomization
 {
 public:
 	static TSharedRef<IDetailCustomization> MakeInstance();
+
+	// Begin: IDetailCustomization
 	virtual void CustomizeDetails(IDetailLayoutBuilder& DetailBuilder) override;
+	// End: IDetailCustomization
 
+private:
 	bool GetFixedBoundsEnabled() const;
 
+	FText GetSelectedTemplateName() const;
+	void SetTemplate(UNiagaraStatelessEmitterTemplate* Template) const;
+	void SetTemplate(FSoftObjectPath TemplatePath) const;
+	TSharedRef<SWidget> OnGetTemplateActions() const;
+
 protected:
 	TWeakObjectPtr<UNiagaraStatelessEmitter> WeakEmitter;
 };
diff --git a/Engine/Plugins/FX/Niagara/Source/NiagaraEditor/Private/Toolkits/NiagaraStatelessEmitterTemplateToolkit.cpp b/Engine/Plugins/FX/Niagara/Source/NiagaraEditor/Private/Toolkits/NiagaraStatelessEmitterTemplateToolkit.cpp
new file mode 100644
--- /dev/null
+++ b/Engine/Plugins/FX/Niagara/Source/NiagaraEditor/Private/Toolkits/NiagaraStatelessEmitterTemplateToolkit.cpp
@@ -0,0 +1,171 @@
+// Copyright Epic Games, Inc. All Rights Reserved.
+
+#include "NiagaraStatelessEmitterTemplateToolkit.h"
+#include "NiagaraEditorModule.h"
+#include "Stateless/NiagaraStatelessEmitterTemplate.h"
+#include "ViewModels/NiagaraStatelessEmitterTemplateViewModel.h"
+#include "Widgets/SNiagaraStatelessEmitterTemplateWidgets.h"
+
+#include "Framework/Docking/TabManager.h"
+#include "Widgets/Docking/SDockTab.h"
+
+#define LOCTEXT_NAMESPACE "NiagaraStatelessEmitterTemplateToolkit"
+
+namespace NiagaraStatelessEmitterTemplateToolkitPrivate
+{
+	const FName FeaturesTabID(TEXT("NiagaraStatelessEmitterTemplateEditor_Features"));
+	const FName ModulesTabID(TEXT("NiagaraStatelessEmitterTemplateEditor_Modules"));
+	const FName OutputVariablesTabID(TEXT("NiagaraStatelessEmitterTemplateEditor_OutputVariables"));
+}
+
+FNiagaraStatelessEmitterTemplateToolkit::FNiagaraStatelessEmitterTemplateToolkit()
+{
+}
+
+FNiagaraStatelessEmitterTemplateToolkit::~FNiagaraStatelessEmitterTemplateToolkit()
+{
+}
+
+void FNiagaraStatelessEmitterTemplateToolkit::Initialize(const EToolkitMode::Type Mode, const TSharedPtr<IToolkitHost>& InitToolkitHost, UNiagaraStatelessEmitterTemplate* EmitterTemplate)
+{
+	using namespace NiagaraStatelessEmitterTemplateToolkitPrivate;
+
+	WeakEmitterTemplate = EmitterTemplate;
+	if (EmitterTemplate == nullptr)
+	{
+		return;
+	}
+
+	ViewModel = MakeShared<FNiagaraStatelessEmitterTemplateViewModel>(EmitterTemplate);
+		
+	const TSharedRef<FTabManager::FLayout> StandaloneDefaultLayout =
+		FTabManager::NewLayout("Standalone_Niagara_StatelessEmitterTemplateLayout_V1_Dev")
+		->AddArea
+		(
+			FTabManager::NewPrimaryArea()->SetOrientation(Orient_Horizontal)
+			->Split
+			(
+				FTabManager::NewSplitter()->SetOrientation(Orient_Vertical)
+				->SetSizeCoefficient(0.3f)
+				->Split(
+					FTabManager::NewStack()
+					->SetSizeCoefficient(0.5f)
+					->AddTab(FeaturesTabID, ETabState::OpenedTab)
+				)
+				->Split(
+					FTabManager::NewStack()
+					->SetSizeCoefficient(0.5f)
+					->AddTab(OutputVariablesTabID, ETabState::OpenedTab)
+				)
+			)
+			->Split
+			(
+				FTabManager::NewStack()
+				->SetSizeCoefficient(0.8f)
+				->AddTab(ModulesTabID, ETabState::OpenedTab)
+			)
+		);
+
+	constexpr bool bCreateDefaultStandaloneMenu = true;
+	constexpr bool bCreateDefaultToolbar = true;
+	InitAssetEditor(Mode, InitToolkitHost, FNiagaraEditorModule::NiagaraEditorAppIdentifier, StandaloneDefaultLayout, bCreateDefaultStandaloneMenu, bCreateDefaultToolbar, EmitterTemplate);
+
+	ExtendToolbar();
+	RegenerateMenusAndToolbars();
+}
+
+void FNiagaraStatelessEmitterTemplateToolkit::RegisterTabSpawners(const TSharedRef<class FTabManager>& InTabManager)
+{
+	using namespace NiagaraStatelessEmitterTemplateToolkitPrivate;
+
+	WorkspaceMenuCategory = InTabManager->AddLocalWorkspaceMenuCategory(LOCTEXT("WorkspaceMenu_NiagaraStatelessEmitterTemplate", "Lightweight Emitter Template"));
+	
+	FAssetEditorToolkit::RegisterTabSpawners(InTabManager);
+
+	InTabManager->RegisterTabSpawner(
+		FeaturesTabID,
+		FOnSpawnTab::CreateSP(this, &FNiagaraStatelessEmitterTemplateToolkit::SpawnTab_Features))
+		.SetDisplayName(LOCTEXT("TemplateFeatures", "Features"))
+		.SetGroup(WorkspaceMenuCategory.ToSharedRef()
+	);
+	InTabManager->RegisterTabSpawner(
+		ModulesTabID,
+		FOnSpawnTab::CreateSP(this, &FNiagaraStatelessEmitterTemplateToolkit::SpawnTab_Modules))
+		.SetDisplayName(LOCTEXT("TemplateModules", "Modules"))
+		.SetGroup(WorkspaceMenuCategory.ToSharedRef()
+	);
+	InTabManager->RegisterTabSpawner(
+		OutputVariablesTabID,
+		FOnSpawnTab::CreateSP(this, &FNiagaraStatelessEmitterTemplateToolkit::SpawnTab_OutputVariables))
+		.SetDisplayName(LOCTEXT("TemplateOutputVariables", "Output Variables"))
+		.SetGroup(WorkspaceMenuCategory.ToSharedRef()
+	);
+}
+
+void FNiagaraStatelessEmitterTemplateToolkit::UnregisterTabSpawners(const TSharedRef<class FTabManager>& InTabManager)
+{
+	using namespace NiagaraStatelessEmitterTemplateToolkitPrivate;
+
+	FAssetEditorToolkit::UnregisterTabSpawners(InTabManager);
+
+	InTabManager->UnregisterTabSpawner(FeaturesTabID);
+	InTabManager->UnregisterTabSpawner(ModulesTabID);
+	InTabManager->UnregisterTabSpawner(OutputVariablesTabID);
+}
+
+FName FNiagaraStatelessEmitterTemplateToolkit::GetToolkitFName() const
+{
+	return FName("Niagara");
+}
+
+FText FNiagaraStatelessEmitterTemplateToolkit::GetBaseToolkitName() const
+{
+	return LOCTEXT("AppLabel", "Niagara");
+}
+
+FString FNiagaraStatelessEmitterTemplateToolkit::GetWorldCentricTabPrefix() const
+{
+	return LOCTEXT("WorldCentricTabPrefix", "Niagara ").ToString();
+}
+
+FLinearColor FNiagaraStatelessEmitterTemplateToolkit::GetWorldCentricTabColorScale() const
+{
+	return FNiagaraEditorModule::WorldCentricTabColorScale;
+}
+
+void FNiagaraStatelessEmitterTemplateToolkit::ExtendToolbar()
+{
+}
+
+TSharedRef<SDockTab> FNiagaraStatelessEmitterTemplateToolkit::SpawnTab_Features(const FSpawnTabArgs& Args)
+{
+	return
+		SNew(SDockTab)
+		[
+			SNew(SNiagaraStatelessEmitterTemplateFeatures)
+			.ViewModel(ViewModel)
+		];
+}
+
+TSharedRef<SDockTab> FNiagaraStatelessEmitterTemplateToolkit::SpawnTab_Modules(const FSpawnTabArgs& Args)
+{
+	return
+		SNew(SDockTab)
+		[
+			SNew(SNiagaraStatelessEmitterTemplateModules)
+			.ViewModel(ViewModel)
+		];
+}
+
+TSharedRef<SDockTab> FNiagaraStatelessEmitterTemplateToolkit::SpawnTab_OutputVariables(const FSpawnTabArgs& Args)
+{
+	return
+		SNew(SDockTab)
+		[
+			SNew(SNiagaraStatelessEmitterTemplateOutputVariables)
+			.ViewModel(ViewModel)
+		];
+}
+
+
+#undef LOCTEXT_NAMESPACE
diff --git a/Engine/Plugins/FX/Niagara/Source/NiagaraEditor/Private/Toolkits/NiagaraStatelessEmitterTemplateToolkit.h b/Engine/Plugins/FX/Niagara/Source/NiagaraEditor/Private/Toolkits/NiagaraStatelessEmitterTemplateToolkit.h
new file mode 100644
--- /dev/null
+++ b/Engine/Plugins/FX/Niagara/Source/NiagaraEditor/Private/Toolkits/NiagaraStatelessEmitterTemplateToolkit.h
@@ -0,0 +1,46 @@
+// Copyright Epic Games, Inc. All Rights Reserved.
+
+#pragma once
+
+#include "Toolkits/IToolkitHost.h"
+#include "Toolkits/AssetEditorToolkit.h"
+
+class UNiagaraStatelessEmitterTemplate;
+class FNiagaraStatelessEmitterTemplateViewModel;
+
+class SDockTab;
+class FSpawnTabArgs;
+class FWorkspaceItem;
+
+// Editor for Niagara Stateless Emitter Templates
+class FNiagaraStatelessEmitterTemplateToolkit : public FAssetEditorToolkit
+{
+public:
+	FNiagaraStatelessEmitterTemplateToolkit();
+	~FNiagaraStatelessEmitterTemplateToolkit();
+
+	// Initialize with a sim cache to view. 
+	void Initialize(const EToolkitMode::Type Mode, const TSharedPtr<class IToolkitHost>& InitToolkitHost, UNiagaraStatelessEmitterTemplate* EmitterTemplate);
+
+	//~ Begin IToolkit Interface
+	virtual void RegisterTabSpawners(const TSharedRef<class FTabManager>& InTabManager) override;
+	virtual void UnregisterTabSpawners(const TSharedRef<class FTabManager>& InTabManager) override;
+	virtual FName GetToolkitFName() const override;
+	virtual FText GetBaseToolkitName() const override;
+	virtual FString GetWorldCentricTabPrefix() const override;
+	virtual FLinearColor GetWorldCentricTabColorScale() const override;
+	//~ End IToolkit Interface
+
+protected:
+	void ExtendToolbar();
+
+	TSharedRef<SDockTab> SpawnTab_Features(const FSpawnTabArgs& Args);
+	TSharedRef<SDockTab> SpawnTab_Modules(const FSpawnTabArgs& Args);
+	TSharedRef<SDockTab> SpawnTab_OutputVariables(const FSpawnTabArgs& Args);
+
+private:
+	TWeakObjectPtr<UNiagaraStatelessEmitterTemplate>		WeakEmitterTemplate;
+	TSharedPtr<FNiagaraStatelessEmitterTemplateViewModel>	ViewModel;
+
+	TSharedPtr<FWorkspaceItem>								WorkspaceMenuCategory;	
+};
diff --git a/Engine/Plugins/FX/Niagara/Source/NiagaraEditor/Private/ViewModels/NiagaraStatelessEmitterTemplateViewModel.cpp b/Engine/Plugins/FX/Niagara/Source/NiagaraEditor/Private/ViewModels/NiagaraStatelessEmitterTemplateViewModel.cpp
new file mode 100644
--- /dev/null
+++ b/Engine/Plugins/FX/Niagara/Source/NiagaraEditor/Private/ViewModels/NiagaraStatelessEmitterTemplateViewModel.cpp
@@ -0,0 +1,81 @@
+// Copyright Epic Games, Inc. All Rights Reserved.
+
+#include "ViewModels/NiagaraStatelessEmitterTemplateViewModel.h"
+
+#include "Stateless/NiagaraStatelessEmitterTemplate.h"
+#include "Stateless/NiagaraStatelessEmitter.h"
+#include "Stateless/NiagaraStatelessModule.h"
+
+#include "UObject/UObjectGlobals.h"
+#include "UObject/Package.h"
+#include "Editor.h"
+
+#define LOCTEXT_NAMESPACE "NiagaraStatelessEmitterTemplateViewModel"
+
+FNiagaraStatelessEmitterTemplateViewModel::FNiagaraStatelessEmitterTemplateViewModel(UNiagaraStatelessEmitterTemplate* InEmitterTemplate)
+	: EmitterTemplate(InEmitterTemplate)
+{
+	FObjectDuplicationParameters DuplicateParameters(InEmitterTemplate, GetTransientPackage());
+	DuplicateParameters.FlagMask = InEmitterTemplate->GetFlags() & ~(RF_Standalone | RF_Public);
+	DuplicateParameters.ApplyFlags = RF_Transactional | RF_Transient;
+	EmitterTemplateForEdit.Reset(Cast<UNiagaraStatelessEmitterTemplate>(StaticDuplicateObjectEx(DuplicateParameters)));
+
+	FEditorDelegates::PostUndoRedo.AddRaw(this, &FNiagaraStatelessEmitterTemplateViewModel::PostUndoRedo);
+}
+
+FNiagaraStatelessEmitterTemplateViewModel::~FNiagaraStatelessEmitterTemplateViewModel()
+{
+	FEditorDelegates::PostUndoRedo.RemoveAll(this);
+}
+
+void FNiagaraStatelessEmitterTemplateViewModel::PostUndoRedo()
+{
+	OnTemplateChangedDelegate.Broadcast();
+}
+
+void FNiagaraStatelessEmitterTemplateViewModel::NotifyPreChange(FProperty* PropertyAboutToChange)
+{
+}
+
+void FNiagaraStatelessEmitterTemplateViewModel::NotifyPostChange(const FPropertyChangedEvent& PropertyChangedEvent, FProperty* PropertyThatChanged)
+{
+	EmitterTemplate->ModifyTemplate(EmitterTemplateForEdit.Get());
+	OnTemplateChangedDelegate.Broadcast();
+}
+
+TConstArrayView<FNiagaraVariableBase> FNiagaraStatelessEmitterTemplateViewModel::GetOututputComponents() const
+{
+	return EmitterTemplate->GetOututputComponents();
+}
+
+TPair<ENiagaraStatelessFeatureMask, ENiagaraStatelessFeatureMask> FNiagaraStatelessEmitterTemplateViewModel::GetFeatureMaskRange() const
+{
+	ENiagaraStatelessFeatureMask MinMask = ENiagaraStatelessFeatureMask::All;
+	ENiagaraStatelessFeatureMask MaxMask = ENiagaraStatelessFeatureMask::All;
+
+	if (EmitterTemplate->GetSimulationShader().IsNull())
+	{
+		MinMask &= ~ENiagaraStatelessFeatureMask::ExecuteGPU;
+		MaxMask &= ~ENiagaraStatelessFeatureMask::ExecuteGPU;
+	}
+
+	for (UClass* ModuleClass : EmitterTemplate->GetModules())
+	{
+		UNiagaraStatelessModule* Module = ModuleClass ? Cast<UNiagaraStatelessModule>(ModuleClass->GetDefaultObject()) : nullptr;
+		if (Module == nullptr)
+		{
+			continue;
+		}
+
+		const ENiagaraStatelessFeatureMask ModuleMask = Module->GetFeatureMask();
+		MinMask &= ModuleMask;
+		if ( Module->CanDisableModule() == false )
+		{
+			MaxMask &= ModuleMask;
+		}
+	}
+
+	return {MinMask, MaxMask};
+}
+
+#undef LOCTEXT_NAMESPACE
diff --git a/Engine/Plugins/FX/Niagara/Source/NiagaraEditor/Private/ViewModels/NiagaraSystemViewModel.cpp b/Engine/Plugins/FX/Niagara/Source/NiagaraEditor/Private/ViewModels/NiagaraSystemViewModel.cpp
--- a/Engine/Plugins/FX/Niagara/Source/NiagaraEditor/Private/ViewModels/NiagaraSystemViewModel.cpp
+++ b/Engine/Plugins/FX/Niagara/Source/NiagaraEditor/Private/ViewModels/NiagaraSystemViewModel.cpp
@@ -632,7 +632,7 @@ TSharedPtr<FNiagaraEmitterHandleViewModel> FNiagaraSystemViewModel::AddStateless
 			StatelessEmitter->AddSpawnInfo();
 			if (UClass* TemplateClass = FindObject<UClass>(nullptr, TEXT("/Script/Niagara.NiagaraStatelessEmitterDefault")))
 			{
-				StatelessEmitter->SetEmitterTemplateClass(TemplateClass);
+				StatelessEmitter->SetEmitterTemplate(Cast<UNiagaraStatelessEmitterTemplate>(TemplateClass->GetDefaultObject()));
 			}
 			{
 				FSoftObjectPath DefaultMaterial(TEXT("Material'/Niagara/DefaultAssets/DefaultSpriteMaterial.DefaultSpriteMaterial'"));
diff --git a/Engine/Plugins/FX/Niagara/Source/NiagaraEditor/Private/ViewModels/Stack/NiagaraStackStatelessEmitterSimulateGroup.cpp b/Engine/Plugins/FX/Niagara/Source/NiagaraEditor/Private/ViewModels/Stack/NiagaraStackStatelessEmitterSimulateGroup.cpp
--- a/Engine/Plugins/FX/Niagara/Source/NiagaraEditor/Private/ViewModels/Stack/NiagaraStackStatelessEmitterSimulateGroup.cpp
+++ b/Engine/Plugins/FX/Niagara/Source/NiagaraEditor/Private/ViewModels/Stack/NiagaraStackStatelessEmitterSimulateGroup.cpp
@@ -198,7 +198,10 @@ void UNiagaraStackStatelessEmitterSimulateGroup::Paste(const UNiagaraClipboardCo
 
 void UNiagaraStackStatelessEmitterSimulateGroup::OnTemplateChanged()
 {
-	RefreshChildren();
+	if (IsFinalized() == false)
+	{
+		RefreshChildren();
+	}
 }
 
 void UNiagaraStackStatelessEmitterSimulateGroup::RefreshChildrenInternal(const TArray<UNiagaraStackEntry*>& CurrentChildren, TArray<UNiagaraStackEntry*>& NewChildren, TArray<FStackIssue>& NewIssues)
diff --git a/Engine/Plugins/FX/Niagara/Source/NiagaraEditor/Private/Widgets/SNiagaraStatelessEmitterTemplateWidgets.cpp b/Engine/Plugins/FX/Niagara/Source/NiagaraEditor/Private/Widgets/SNiagaraStatelessEmitterTemplateWidgets.cpp
new file mode 100644
--- /dev/null
+++ b/Engine/Plugins/FX/Niagara/Source/NiagaraEditor/Private/Widgets/SNiagaraStatelessEmitterTemplateWidgets.cpp
@@ -0,0 +1,164 @@
+// Copyright Epic Games, Inc. All Rights Reserved.
+
+#include "SNiagaraStatelessEmitterTemplateWidgets.h"
+#include "Stateless/NiagaraStatelessCommon.h"
+#include "Stateless/NiagaraStatelessEmitterTemplate.h"
+#include "ViewModels/NiagaraStatelessEmitterTemplateViewModel.h"
+#include "Widgets/SNiagaraParameterName.h"
+#include "NiagaraEditorUtilities.h"
+
+#include "Components/VerticalBox.h"
+#include "Widgets/Layout/SGridPanel.h"
+#include "DetailsViewArgs.h"
+#include "IDetailsView.h"
+#include "PropertyEditorModule.h"
+
+#define LOCTEXT_NAMESPACE "NiagaraStatelessEmitterTemplateWidgets"
+
+/////////////////////////////////////////////////////////////////////////////////////////////////////////////// 
+
+void SNiagaraStatelessEmitterTemplateModules::Construct(const FArguments& InArgs)
+{
+	ViewModel = InArgs._ViewModel;
+
+	FPropertyEditorModule& PropertyEditorModule = FModuleManager::GetModuleChecked<FPropertyEditorModule>("PropertyEditor");
+	FDetailsViewArgs DetailsViewArgs;
+	DetailsViewArgs.NameAreaSettings = FDetailsViewArgs::HideNameArea;
+	DetailsViewArgs.bHideSelectionTip = true;
+	DetailsViewArgs.NotifyHook = ViewModel.Get();
+	DetailsViewArgs.bAllowSearch = false;
+
+	TSharedRef<IDetailsView> DetailsView = PropertyEditorModule.CreateDetailView(DetailsViewArgs);
+	DetailsView->SetObject(ViewModel->GetTemplateForEdit());
+
+	ChildSlot
+	[
+		SNew(SVerticalBox)
+		+SVerticalBox::Slot()
+		.AutoHeight()
+		.Padding(FMargin(0.0f, 2.0f))
+		[
+			DetailsView
+		]
+	];
+}
+
+/////////////////////////////////////////////////////////////////////////////////////////////////////////////// 
+
+void SNiagaraStatelessEmitterTemplateFeatures::Construct(const FArguments& InArgs)
+{
+	ViewModel = InArgs._ViewModel;
+	ViewModel->OnTemplateChanged().AddSP(this, &SNiagaraStatelessEmitterTemplateFeatures::OnRebuildWidget);
+
+	ChildSlot
+	[
+		SAssignNew(GridPanel, SGridPanel)
+	];
+
+	OnRebuildWidget();
+}
+
+void SNiagaraStatelessEmitterTemplateFeatures::OnRebuildWidget()
+{
+	GridPanel->ClearChildren();
+
+	const TPair<ENiagaraStatelessFeatureMask, ENiagaraStatelessFeatureMask> FeatureMasks = ViewModel->GetFeatureMaskRange();
+	const auto GetEnabledText = [](bool bEnabled) { return bEnabled ? LOCTEXT("Enabled", "Enabled") : LOCTEXT("Disabled", "Disabled"); };
+
+	UEnum* Enum = StaticEnum<ENiagaraStatelessFeatureMask>();
+	int32 iRow = 0;
+
+	// Setup header
+	{
+		int32 iColumn = 0;
+		GridPanel->AddSlot(iColumn++, iRow)
+		.Padding(5)
+		[
+			SNew(STextBlock)
+			.Text(LOCTEXT("Feature", "Feature"))
+		];
+		GridPanel->AddSlot(iColumn++, iRow)
+		.Padding(5)
+		[
+			SNew(STextBlock)
+			.Text(LOCTEXT("AllModulesEnabled", "All Modules Enabled"))
+		];
+		GridPanel->AddSlot(iColumn++, iRow)
+		.Padding(5)
+		[
+			SNew(STextBlock)
+			.Text(LOCTEXT("RequriedModulesEnabled", "Required Modules Enabled"))
+		];
+		++iRow;
+	}
+
+	// Generated rows
+	for (int32 i = 0; i < Enum->NumEnums() - 1; ++i)		// Note: -1 to skip the implicit MAX field
+	{
+		if (Enum->HasMetaData(TEXT("Hidden"), i))
+		{
+			continue;
+		}
+
+		int32 iColumn = 0;
+
+		GridPanel->AddSlot(iColumn++, iRow)
+		.Padding(5)
+		[
+			SNew(STextBlock)
+			.Text(Enum->GetDisplayNameTextByIndex(i))
+		];		
+
+		const bool bMinEnabled = (int32(FeatureMasks.Key) & int32(Enum->GetValueByIndex(i))) != 0;
+		const bool bMaxEnabled = (int32(FeatureMasks.Value) & int32(Enum->GetValueByIndex(i))) != 0;
+
+		GridPanel->AddSlot(iColumn++, iRow)
+		.Padding(5)
+		[
+			SNew(STextBlock)
+			.Text(GetEnabledText(bMinEnabled))
+		];		
+
+		GridPanel->AddSlot(iColumn++, iRow)
+		.Padding(5)
+		[
+			SNew(STextBlock)
+			.Text(GetEnabledText(bMaxEnabled))
+		];
+
+		++iRow;
+	}
+}
+
+/////////////////////////////////////////////////////////////////////////////////////////////////////////////// 
+
+void SNiagaraStatelessEmitterTemplateOutputVariables::Construct(const FArguments& InArgs)
+{
+	ViewModel = InArgs._ViewModel;
+	ViewModel->OnTemplateChanged().AddSP(this, &SNiagaraStatelessEmitterTemplateOutputVariables::OnRebuildWidget);
+
+	ChildSlot
+	[
+		SAssignNew(GridPanel, SGridPanel)
+	];
+
+	OnRebuildWidget();
+}
+
+void SNiagaraStatelessEmitterTemplateOutputVariables::OnRebuildWidget()
+{
+	GridPanel->ClearChildren();
+
+	int32 NumSlots = 0;
+
+	for (const FNiagaraVariableBase& Variable : ViewModel->GetOututputComponents())
+	{
+		GridPanel->AddSlot(0, NumSlots++)
+		//.Padding(5)
+		[
+			FNiagaraParameterUtilities::GetParameterWidget(FNiagaraVariable(Variable), true, false)
+		];
+	}
+}
+
+#undef LOCTEXT_NAMESPACE
diff --git a/Engine/Plugins/FX/Niagara/Source/NiagaraEditor/Private/Widgets/SNiagaraStatelessEmitterTemplateWidgets.h b/Engine/Plugins/FX/Niagara/Source/NiagaraEditor/Private/Widgets/SNiagaraStatelessEmitterTemplateWidgets.h
new file mode 100644
--- /dev/null
+++ b/Engine/Plugins/FX/Niagara/Source/NiagaraEditor/Private/Widgets/SNiagaraStatelessEmitterTemplateWidgets.h
@@ -0,0 +1,61 @@
+// Copyright Epic Games, Inc. All Rights Reserved.
+
+#pragma once
+
+#include "Widgets/SCompoundWidget.h"
+#include "Templates/SharedPointer.h"
+
+class FNiagaraStatelessEmitterTemplateViewModel;
+class SGridPanel;
+
+/////////////////////////////////////////////////////////////////////////////////////////////////////////////// 
+
+class SNiagaraStatelessEmitterTemplateModules : public SCompoundWidget
+{
+public:
+	SLATE_BEGIN_ARGS(SNiagaraStatelessEmitterTemplateModules) {}
+		SLATE_ARGUMENT(TSharedPtr<FNiagaraStatelessEmitterTemplateViewModel>, ViewModel)
+	SLATE_END_ARGS()
+
+	void Construct(const FArguments& InArgs);
+
+private:
+	TSharedPtr<FNiagaraStatelessEmitterTemplateViewModel> ViewModel;
+};
+
+/////////////////////////////////////////////////////////////////////////////////////////////////////////////// 
+
+class SNiagaraStatelessEmitterTemplateFeatures : public SCompoundWidget
+{
+public:
+	SLATE_BEGIN_ARGS(SNiagaraStatelessEmitterTemplateFeatures) {}
+		SLATE_ARGUMENT(TSharedPtr<FNiagaraStatelessEmitterTemplateViewModel>, ViewModel)
+	SLATE_END_ARGS()
+
+	void Construct(const FArguments& InArgs);
+
+	void OnRebuildWidget();
+
+private:
+	TSharedPtr<FNiagaraStatelessEmitterTemplateViewModel>	ViewModel;
+	TSharedPtr<SGridPanel>									GridPanel;
+};
+
+/////////////////////////////////////////////////////////////////////////////////////////////////////////////// 
+
+class SNiagaraStatelessEmitterTemplateOutputVariables : public SCompoundWidget
+{
+public:
+	SLATE_BEGIN_ARGS(SNiagaraStatelessEmitterTemplateOutputVariables) {}
+		SLATE_ARGUMENT(TSharedPtr<FNiagaraStatelessEmitterTemplateViewModel>, ViewModel)
+	SLATE_END_ARGS()
+
+	void Construct(const FArguments& InArgs);
+
+	void OnRebuildWidget();
+
+private:
+	TSharedPtr<FNiagaraStatelessEmitterTemplateViewModel>	ViewModel;
+	TSharedPtr<SGridPanel>									GridPanel;
+};
+
