Niagara - Traversal cache improvements.
+ Fix an issue where garbage collection would cause incorrect results in the object path->soft object ptr cache.
+ Guard the main cache structures for thread safety during async compiles.
+ Add the cache enabled state to the vm id so that it generates unique ddc entries.
+ Add a special case to the static context generation so that reading emitter.localstate in a system script doesn't fail traversal.
+ Add a commandlet to validate the traversal cache results against the legacy results when garbage collecting and accessing the cache on multiple threads.

diff --git a/Engine/Plugins/FX/Niagara/Source/Niagara/Private/NiagaraScript.cpp b/Engine/Plugins/FX/Niagara/Source/Niagara/Private/NiagaraScript.cpp
--- a/Engine/Plugins/FX/Niagara/Source/Niagara/Private/NiagaraScript.cpp
+++ b/Engine/Plugins/FX/Niagara/Source/Niagara/Private/NiagaraScript.cpp
@@ -1451,6 +1451,8 @@ void UNiagaraScript::ComputeVMCompilationId(FNiagaraVMExecutableDataId& Id, cons
 
 	// Append the state of the "Fail If Not Set severity" cvar, as it may affect the LastCompileEvents which is a member of the script's VMExecutableData.
 	Id.AdditionalDefines.Add(FString::Printf(TEXT("FailIfNotSetMessageSeverity: %d"), GNiagaraTranslatorFailIfNotSetSeverity));
+
+	Id.AdditionalDefines.Add(FString::Printf(TEXT("EnableTraversalCache: %d"), GbNiagaraEnableTraversalCache));
 	Id.ReferencedCompileHashes = MoveTemp(HashCollector.ReferencedHashes);
 
 	FNiagaraVMExecutableDataId& LastGeneratedVMId = GetLastGeneratedVMId(VersionGuid);
@@ -4282,6 +4284,7 @@ void UNiagaraScript::BuildClassSchema(FAppendToClassSchemaContext& Context)
 
 	// settings/configurations
 	Context.Update(&GNiagaraTranslatorFailIfNotSetSeverity, sizeof(GNiagaraTranslatorFailIfNotSetSeverity));
+	Context.Update(&GbNiagaraEnableTraversalCache, sizeof(GbNiagaraEnableTraversalCache));
 
 	const UNiagaraSettings* NiagaraSettings = GetDefault<UNiagaraSettings>();
 	checkf(NiagaraSettings, TEXT("Failed to acquire default UNiagaraSettings during AppendToClassSchema"));
diff --git a/Engine/Plugins/FX/Niagara/Source/NiagaraEditor/Private/Commandlets/NiagaraTraversalCacheAuditCommandlet.cpp b/Engine/Plugins/FX/Niagara/Source/NiagaraEditor/Private/Commandlets/NiagaraTraversalCacheAuditCommandlet.cpp
new file mode 100644
--- /dev/null
+++ b/Engine/Plugins/FX/Niagara/Source/NiagaraEditor/Private/Commandlets/NiagaraTraversalCacheAuditCommandlet.cpp
@@ -0,0 +1,528 @@
+// Copyright Epic Games, Inc. All Rights Reserved.
+
+#include "Commandlets/NiagaraTraversalCacheAuditCommandlet.h"
+
+#include "AssetRegistry/AssetData.h"
+#include "AssetRegistry/AssetRegistryModule.h"
+#include "Async/ParallelFor.h"
+#include "EdGraphSchema_Niagara.h"
+#include "HAL/IConsoleManager.h"
+#include "Logging/LogMacros.h"
+#include "Misc/OutputDeviceArchiveWrapper.h"
+#include "Modules/ModuleManager.h"
+#include "NiagaraEditorModule.h"
+#include "NiagaraEmitterHandle.h"
+#include "NiagaraGraphDataCache.h"
+#include "NiagaraNodeFunctionCall.h"
+#include "NiagaraScript.h"
+#include "NiagaraScriptSource.h"
+#include "NiagaraSimulationStageBase.h"
+#include "NiagaraSystem.h"
+#include "TraversalCache/TraversalCache.h"
+#include "NiagaraEditorModule.h"
+
+#include UE_INLINE_GENERATED_CPP_BY_NAME(NiagaraTraversalCacheAuditCommandlet)
+
+DEFINE_LOG_CATEGORY_STATIC(LogNiagaraTraversalCacheAuditCommandlet, Log, All);
+
+namespace NiagaraTraversalCacheAuditUtility
+{
+
+	FString TopLevelScriptToScriptName(const UNiagaraScript* InScript)
+	{
+		switch (InScript->GetUsage())
+		{
+		case ENiagaraScriptUsage::SystemSpawnScript:
+			return "System spawn script";
+		case ENiagaraScriptUsage::SystemUpdateScript:
+			return "System update script";
+		case ENiagaraScriptUsage::EmitterSpawnScript:
+			return "Emitter spawn script";
+		case ENiagaraScriptUsage::EmitterUpdateScript:
+			return "Emitter update script";
+		case ENiagaraScriptUsage::ParticleSpawnScript:
+		case ENiagaraScriptUsage::ParticleSpawnScriptInterpolated:
+			return "Particle spawn script";
+		case ENiagaraScriptUsage::ParticleUpdateScript:
+			return "Particle update script";
+		case ENiagaraScriptUsage::ParticleEventScript:
+			return "Particle event script - " + InScript->GetUsageId().ToString(EGuidFormats::DigitsWithHyphens);
+		case ENiagaraScriptUsage::ParticleGPUComputeScript:
+			return "Particle combined GPU script";
+		case ENiagaraScriptUsage::ParticleSimulationStageScript:
+		{
+			UNiagaraSimulationStageBase* SimulationStage = InScript->GetTypedOuter<UNiagaraSimulationStageBase>();
+			FString SimulationStageName = SimulationStage != nullptr
+				? SimulationStage->SimulationStageName.ToString()
+				: InScript->GetUsageId().ToString(EGuidFormats::DigitsWithHyphens);
+			return "Particle simulation stage script - " + SimulationStageName;
+		}
+		default:
+			return "Unknown script";
+		}
+	}
+
+	template<typename TCollection, typename TItem, typename TTransformedItem, typename TItemTransformer>
+	void CollectCollectionDifferences(
+		const TCollection& Items, const TCollection& NewItems,
+		TArray<TTransformedItem>& OutAddedItems, TArray<TTransformedItem>& OutRemovedItems,
+		TItemTransformer Transformer)
+	{
+		for (const TItem& NewItem : NewItems)
+		{
+			if (Items.Contains(NewItem) == false)
+			{
+				OutAddedItems.Add(Transformer(NewItem));
+			}
+		}
+		for (const TItem& Item : Items)
+		{
+			if (NewItems.Contains(Item) == false)
+			{
+				OutRemovedItems.Add(Transformer(Item));
+			}
+		}
+	}
+
+	template<typename TItem, typename TItemComparer, typename TTransformedItem, typename TItemTransformer>
+	bool CompareCollections(
+		TArray<TItem>& Items, TArray<TItem>& NewItems,
+		TArray<TTransformedItem>& OutAddedItems, TArray<TTransformedItem>& OutRemovedItems,
+		TSet<TItem>& HiddenItems, TSet<TItem>& NewHiddenItems,
+		TArray<TTransformedItem>& OutAddedHiddenItems, TArray<TTransformedItem>& OutRemovedHiddenItems,
+		TItemComparer Comparer, TItemTransformer Transformer)
+	{
+		auto ItemsMatchNewItems = [&](TArray<TItem>& Items, TArray<TItem>& NewItems)->bool
+		{
+			Items.Sort(Comparer);
+			NewItems.Sort(Comparer);
+			for (int32 ItemIndex = 0; ItemIndex < Items.Num(); ItemIndex++)
+			{
+				if (Items[ItemIndex] != NewItems[ItemIndex])
+				{
+					return false;
+				}
+			}
+			return true;
+		};
+
+		if (NewItems.Num() != Items.Num() ||
+			NewHiddenItems.Num() != HiddenItems.Num() ||
+			ItemsMatchNewItems(Items, NewItems) == false ||
+			HiddenItems.Includes(NewHiddenItems) == false)
+		{
+			CollectCollectionDifferences<TArray<TItem>, TItem, TTransformedItem, TItemTransformer>(Items, NewItems, OutAddedItems, OutRemovedItems, Transformer);
+			CollectCollectionDifferences<TSet<TItem>, TItem, TTransformedItem, TItemTransformer>(HiddenItems, NewHiddenItems, OutAddedHiddenItems, OutRemovedHiddenItems, Transformer);
+			return false;
+		}
+
+		return true;
+	}
+}
+
+FCriticalSection UNiagaraTraversalCacheAuditCommandlet::ResultCS;
+
+int32 UNiagaraTraversalCacheAuditCommandlet::FAuditResultFunctionData::CountDifferences() const
+{
+	return
+		AddedInputVariables.Num() + RemovedInputVariables.Num() +
+		AddedHiddenInputVariables.Num() + RemovedHiddenInputVariables.Num() +
+		AddedStaticInputPinVariables.Num() + RemovedStaticInputPinVariables.Num() +
+		AddedHiddenStaticInputPinVariables.Num() + RemovedHiddenStaticInputPinVariables.Num();
+}
+
+int32 UNiagaraTraversalCacheAuditCommandlet::FAuditResultScriptData::CountDifferences() const
+{
+	int32 DifferenceCount = 0;
+	for (const FAuditResultFunctionData& FunctionData : FunctionDatas)
+	{
+		DifferenceCount += FunctionData.CountDifferences();
+	}
+	return DifferenceCount;
+}
+
+int32 UNiagaraTraversalCacheAuditCommandlet::FAuditResultSystemData::CountDifferences() const
+{
+	int32 DifferenceCount = 0;
+	for (const FAuditResultScriptData& ScriptData : ScriptDatas)
+	{
+		DifferenceCount += ScriptData.CountDifferences();
+	}
+	return DifferenceCount;
+}
+
+UNiagaraTraversalCacheAuditCommandlet::UNiagaraTraversalCacheAuditCommandlet(const FObjectInitializer& ObjectInitializer)
+	: Super(ObjectInitializer)
+{
+}
+
+int32 UNiagaraTraversalCacheAuditCommandlet::Main(const FString& Params)
+{
+	TArray<FString> Tokens;
+	TArray<FString> Switches;
+
+	ParseCommandLine(*Params, Tokens, Switches);
+	if (Switches.Contains(TEXT("pause")))
+	{
+		int32 WaitSeconds = 20;
+
+		const double EndSeconds = FPlatformTime::Seconds() + WaitSeconds;
+		UE_LOG(LogNiagaraTraversalCacheAuditCommandlet, Display, TEXT("Pausing for %i seconds."), WaitSeconds);
+		while (FPlatformTime::Seconds() < EndSeconds)
+		{
+			FPlatformProcess::Sleep(1);
+
+			// Tick the engine
+			CommandletHelpers::TickEngine();
+
+			// Stop if exit was requested
+			if (IsEngineExitRequested())
+			{
+				break;
+			}
+		}
+
+		UE_LOG(LogNiagaraTraversalCacheAuditCommandlet, Display, TEXT("Pause finished."));
+	}
+
+	int32 MaxErrorsParam;
+	if (FParse::Value(*Params, TEXT("MaxErrors="), MaxErrorsParam))
+	{
+		MaxErrorCount = MaxErrorsParam;
+	}
+
+	EnableTraversalCacheCVar = IConsoleManager::Get().FindConsoleVariable(TEXT("fx.Niagara.EnableTraversalCache"), false);
+	int32 EnableTraversalCachePreAuditValue = EnableTraversalCacheCVar->GetInt();
+
+	TArray<FString> SystemPaths;
+	GatherAllSystemPaths(SystemPaths);
+
+	FAuditResults Results;
+	AuditSystems(SystemPaths, Results);
+
+	EnableTraversalCacheCVar->SetWithCurrentPriority(EnableTraversalCachePreAuditValue);
+	return 0;
+}
+
+void UNiagaraTraversalCacheAuditCommandlet::GatherAllSystemPaths(TArray<FString>& OutSystemPaths) const
+{
+	FAssetRegistryModule& AssetRegistryModule = FModuleManager::LoadModuleChecked<FAssetRegistryModule>("AssetRegistry");
+	AssetRegistryModule.Get().SearchAllAssets(true);
+	TArray<FAssetData> SystemAssets;
+	AssetRegistryModule.Get().GetAssetsByClass(UNiagaraSystem::StaticClass()->GetClassPathName(), SystemAssets);
+
+	OutSystemPaths.Reserve(SystemAssets.Num());
+	for (FAssetData& SystemAsset : SystemAssets)
+	{
+		OutSystemPaths.Add(SystemAsset.PackageName.ToString());
+	}
+
+	OutSystemPaths.Sort();
+}
+
+void UNiagaraTraversalCacheAuditCommandlet::CollectFunctionData(UNiagaraSystem& System, const FNiagaraEmitterHandle* EmitterHandle, UNiagaraScript* Script, TArray<FFunctionData>& OutFunctionDatas) const
+{
+	if (Script == nullptr)
+	{
+		return;
+	}
+
+	TArray<UNiagaraNode*> TraversedNodes;
+	const UNiagaraScriptSource* ScriptSource = CastChecked<UNiagaraScriptSource>(Script->GetSource(FGuid()));
+	ScriptSource->NodeGraph->BuildTraversal(TraversedNodes, Script->GetUsage(), Script->GetUsageId());
+	for (UNiagaraNode* TraversedNode : TraversedNodes)
+	{
+		UNiagaraNodeFunctionCall* FunctionCallNode = Cast<UNiagaraNodeFunctionCall>(TraversedNode);
+		if (FunctionCallNode != nullptr)
+		{
+			FFunctionData& FunctionData = OutFunctionDatas.AddDefaulted_GetRef();
+			FunctionData.FunctionCallNode = FunctionCallNode;
+
+			FCompileConstantResolver CompileConstantResolver = EmitterHandle == nullptr
+				? FCompileConstantResolver(&System, Script->GetUsage())
+				: FCompileConstantResolver(EmitterHandle->GetInstance(), Script->GetUsage());
+
+			FNiagaraStackGraphUtilities::ENiagaraGetStackFunctionInputPinsOptions Options = FNiagaraStackGraphUtilities::ENiagaraGetStackFunctionInputPinsOptions::ModuleInputsOnly;
+			FNiagaraStackGraphUtilities::GetStackFunctionInputs(*FunctionCallNode, FunctionData.InputVariables, FunctionData.HiddenInputVariables, CompileConstantResolver, Options);
+			FNiagaraStackGraphUtilities::GetStackFunctionStaticSwitchPins(*FunctionCallNode, FunctionData.InputPins, FunctionData.HiddenInputPins, CompileConstantResolver);
+		}
+	}
+}
+
+void UNiagaraTraversalCacheAuditCommandlet::CollectScriptData(UNiagaraSystem* System, const FNiagaraEmitterHandle* EmitterHandle, UNiagaraScript* Script, TArray<FScriptData>& OutScriptDatas) const
+{
+	FScriptData& ScriptData = OutScriptDatas.AddDefaulted_GetRef();
+	ScriptData.Script = Script;
+	ScriptData.EmitterHandle = EmitterHandle;
+	CollectFunctionData(*System, EmitterHandle, Script, ScriptData.FunctionDatas);
+}
+
+void UNiagaraTraversalCacheAuditCommandlet::CollectScriptDataForSystem(FSystemData& SystemData) const
+{
+	CollectScriptData(SystemData.System, nullptr, SystemData.System->GetSystemSpawnScript(), SystemData.ScriptDatas);
+	CollectScriptData(SystemData.System, nullptr, SystemData.System->GetSystemUpdateScript(), SystemData.ScriptDatas);
+	for (const FNiagaraEmitterHandle& EmitterHandle : SystemData.System->GetEmitterHandles())
+	{
+		FVersionedNiagaraEmitterData* EmitterData = EmitterHandle.GetInstance().GetEmitterData();
+		CollectScriptData(SystemData.System, &EmitterHandle, EmitterData->EmitterSpawnScriptProps.Script, SystemData.ScriptDatas);
+		CollectScriptData(SystemData.System, &EmitterHandle, EmitterData->EmitterUpdateScriptProps.Script, SystemData.ScriptDatas);
+		CollectScriptData(SystemData.System, &EmitterHandle, EmitterData->SpawnScriptProps.Script, SystemData.ScriptDatas);
+		CollectScriptData(SystemData.System, &EmitterHandle, EmitterData->UpdateScriptProps.Script, SystemData.ScriptDatas);
+		for (const FNiagaraEventScriptProperties& EventScript : EmitterData->EventHandlerScriptProps)
+		{
+			CollectScriptData(SystemData.System, &EmitterHandle, EventScript.Script, SystemData.ScriptDatas);
+		}
+		for (UNiagaraSimulationStageBase* SimulationStage : EmitterData->GetSimulationStages())
+		{
+			CollectScriptData(SystemData.System, &EmitterHandle, SimulationStage->Script, SystemData.ScriptDatas);
+		}
+	}
+}
+
+void UNiagaraTraversalCacheAuditCommandlet::AuditScript(UNiagaraSystem& AuditSystem, FScriptData& AuditScriptData, TArray<FAuditResultScriptData>& OutScriptResults) const
+{
+	using namespace NiagaraTraversalCacheAuditUtility;
+
+	TArray<FFunctionData> TraversalCacheFunctionDatas;
+	CollectFunctionData(AuditSystem, AuditScriptData.EmitterHandle, AuditScriptData.Script, TraversalCacheFunctionDatas);
+
+	TArray<FAuditResultFunctionData> FunctionResults;
+	for (FFunctionData& AuditFunctionData : AuditScriptData.FunctionDatas)
+	{
+		FFunctionData* TraversalCacheFunctionData = TraversalCacheFunctionDatas.FindByPredicate([&AuditFunctionData](const FFunctionData& TraversalCacheFunctionData)
+			{ return TraversalCacheFunctionData.FunctionCallNode == AuditFunctionData.FunctionCallNode; });
+		if (TraversalCacheFunctionData == nullptr)
+		{
+			FString ScriptName = AuditScriptData.EmitterHandle != nullptr
+				? FString::Printf(TEXT("%s - %s"), *AuditScriptData.EmitterHandle->GetUniqueInstanceName(), *TopLevelScriptToScriptName(AuditScriptData.Script))
+				: TopLevelScriptToScriptName(AuditScriptData.Script);
+			UE_LOG(LogNiagaraTraversalCacheAuditCommandlet, Error, TEXT("System: %s - Script %s - Function %s - Not found in traversal cache data."),
+				*AuditSystem.GetPathName(), *ScriptName, *AuditFunctionData.FunctionCallNode->GetFunctionName());
+		}
+		else
+		{
+			TArray<FNiagaraVariableBase> AddedInputVariables;
+			TArray<FNiagaraVariableBase> RemovedInputVariables;
+			TArray<FNiagaraVariableBase> AddedHiddenInputVariables;
+			TArray<FNiagaraVariableBase> RemovedHiddenInputVariables;
+			auto CompareVariablesForSort = [](const FNiagaraVariable& A, const FNiagaraVariable& B) { return A.GetName().ToString() < B.GetName().ToString(); };
+			auto TransformVariableToVariableBase = [](const FNiagaraVariable& Item)->FNiagaraVariableBase { return Item; };
+			bool bVariableCompareMatched = CompareCollections(
+				AuditFunctionData.InputVariables, TraversalCacheFunctionData->InputVariables,
+				AddedInputVariables, RemovedInputVariables,
+				AuditFunctionData.HiddenInputVariables, TraversalCacheFunctionData->HiddenInputVariables,
+				AddedHiddenInputVariables, RemovedHiddenInputVariables,
+				CompareVariablesForSort, TransformVariableToVariableBase);
+
+			TArray<FNiagaraVariableBase> AddedStaticInputPinVariables;
+			TArray<FNiagaraVariableBase> RemovedStaticInputPinVariables;
+			TArray<FNiagaraVariableBase> AddedHiddenStaticInputPinVariables;
+			TArray<FNiagaraVariableBase> RemovedHiddenStaticInputPinVariables;
+			auto ComparePinsForSort = [](const UEdGraphPin& A, const UEdGraphPin& B) { return A.PinName.ToString() < B.PinName.ToString(); };
+			auto TransformPinToVariableBase = [](const UEdGraphPin* Pin)->FNiagaraVariableBase { return UEdGraphSchema_Niagara::PinToNiagaraVariable(Pin); };
+			bool bPinCompareMatched = CompareCollections(
+				AuditFunctionData.InputPins, TraversalCacheFunctionData->InputPins,
+				AddedStaticInputPinVariables, RemovedStaticInputPinVariables,
+				AuditFunctionData.HiddenInputPins, TraversalCacheFunctionData->HiddenInputPins,
+				AddedHiddenStaticInputPinVariables, RemovedHiddenStaticInputPinVariables,
+				ComparePinsForSort, TransformPinToVariableBase);
+
+			if (bVariableCompareMatched == false || bPinCompareMatched == false)
+			{
+				FAuditResultFunctionData& FunctionResult = FunctionResults.AddDefaulted_GetRef();
+				FunctionResult.FunctionName = TraversalCacheFunctionData->FunctionCallNode->GetFunctionName();
+				FunctionResult.FunctionScriptPath = *TraversalCacheFunctionData->FunctionCallNode->FunctionScript->GetPathName();
+				FunctionResult.bFunctionEnabled = TraversalCacheFunctionData->FunctionCallNode->GetDesiredEnabledState() != ENodeEnabledState::Disabled;
+
+				FunctionResult.AddedInputVariables = AddedInputVariables;
+				FunctionResult.RemovedInputVariables = RemovedInputVariables;
+				FunctionResult.AddedHiddenInputVariables = AddedHiddenInputVariables;
+				FunctionResult.RemovedHiddenInputVariables = RemovedHiddenInputVariables;
+
+				FunctionResult.AddedStaticInputPinVariables = AddedStaticInputPinVariables;
+				FunctionResult.RemovedStaticInputPinVariables = RemovedStaticInputPinVariables;
+				FunctionResult.AddedHiddenStaticInputPinVariables = AddedHiddenStaticInputPinVariables;
+				FunctionResult.RemovedHiddenStaticInputPinVariables = RemovedHiddenStaticInputPinVariables;
+			}
+		}
+	}
+
+	if (FunctionResults.Num() > 0)
+	{
+		FAuditResultScriptData& ScriptResult = OutScriptResults.AddDefaulted_GetRef();
+		ScriptResult.ScriptName = TopLevelScriptToScriptName(AuditScriptData.Script);
+		ScriptResult.EmitterName = AuditScriptData.EmitterHandle != nullptr ? AuditScriptData.EmitterHandle->GetUniqueInstanceName() : TEXT("");
+		ScriptResult.FunctionDatas = FunctionResults;
+	}
+}
+
+void UNiagaraTraversalCacheAuditCommandlet::LogSystemResult(const FAuditResultSystemData& SystemResult, int32 SystemResultIndex) const
+{
+	for (const FAuditResultScriptData& ScriptResult : SystemResult.ScriptDatas)
+	{
+		UE_LOG(LogNiagaraTraversalCacheAuditCommandlet, Error, TEXT("[%i] System: %s Emitter: %s"),
+			SystemResultIndex + 1, *SystemResult.SystemPath, *ScriptResult.EmitterName);
+		for (const FAuditResultFunctionData& FunctionResult : ScriptResult.FunctionDatas)
+		{
+			TArray<FString> DifferenceStrings;
+			if (FunctionResult.AddedInputVariables.Num() > 0 || FunctionResult.RemovedInputVariables.Num() > 0)
+			{
+				DifferenceStrings.Add(TEXT("Input Variables"));
+			}
+			if (FunctionResult.AddedHiddenInputVariables.Num() > 0 || FunctionResult.RemovedHiddenInputVariables.Num() > 0)
+			{
+				DifferenceStrings.Add(TEXT("Hidden Input Variables"));
+			}
+			if (FunctionResult.AddedStaticInputPinVariables.Num() > 0 || FunctionResult.RemovedStaticInputPinVariables.Num() > 0)
+			{
+				DifferenceStrings.Add(TEXT("Static Pins"));
+			}
+			if (FunctionResult.AddedHiddenStaticInputPinVariables.Num() > 0 || FunctionResult.RemovedHiddenStaticInputPinVariables.Num() > 0)
+			{
+				DifferenceStrings.Add(TEXT("Hidden Static Pins"));
+			}
+			FString DifferenceString = FString::Join(DifferenceStrings, TEXT(", "));;
+			UE_LOG(LogNiagaraTraversalCacheAuditCommandlet, Error, TEXT("[%i] Script: %s Function: %s%s Function Script: %s Difference: %s"),
+				SystemResultIndex + 1,
+				*ScriptResult.ScriptName,
+				*FunctionResult.FunctionName,
+				FunctionResult.bFunctionEnabled ? TEXT("") : TEXT(" (Disabled)"),
+				FunctionResult.FunctionScriptPath != NAME_None ? *FunctionResult.FunctionScriptPath.ToString() : TEXT("(none)"),
+				*DifferenceString);
+		}
+	}
+}
+
+void UNiagaraTraversalCacheAuditCommandlet::AuditSystems(const TArray<FString>& SystemPaths, FAuditResults& Results) const
+{
+	int32 NumBatches = FMath::CeilToInt((float)SystemPaths.Num() / BatchSize);
+	for (int32 BatchIndex = 0; BatchIndex < NumBatches && (MaxErrorCount == 0 || Results.ErrorCount < MaxErrorCount); BatchIndex++)
+	{
+		// Disable the traversal cache and load all of the systems in this batch.
+		EnableTraversalCacheCVar->SetWithCurrentPriority(0);
+		TArray<FSystemData> BatchSystems;
+		int32 SystemPathStartIndex = BatchIndex * BatchSize;
+		int32 SystemPathEndIndex = FMath::Min(SystemPathStartIndex + BatchSize, SystemPaths.Num());
+		for (int32 SystemPathIndex = SystemPathStartIndex; SystemPathIndex < SystemPathEndIndex; SystemPathIndex++)
+		{
+			FSystemData& SystemData = BatchSystems.AddDefaulted_GetRef();
+			SystemData.SystemPath = SystemPaths[SystemPathIndex];
+			SystemData.System = LoadObject<UNiagaraSystem>(GetTransientPackage(), *SystemData.SystemPath);
+			if (SystemData.System == nullptr)
+			{
+				UE_LOG(LogNiagaraTraversalCacheAuditCommandlet, Warning, TEXT("[%i / %i] Load failed - %s"), SystemPathIndex + 1, SystemPaths.Num(), *SystemData.SystemPath);
+			}
+			else
+			{
+				UE_LOG(LogNiagaraTraversalCacheAuditCommandlet, Display, TEXT("[%i / %i] Loaded - %s"), SystemPathIndex + 1, SystemPaths.Num(), *SystemData.SystemPath);
+				SystemData.System->WaitForCompilationComplete();
+			}
+		}
+
+		// Collect the legacy script and function data for the loaded systems on the main thread.
+		for (int32 BatchSystemIndex = 0; BatchSystemIndex < BatchSystems.Num(); BatchSystemIndex++)
+		{
+			FSystemData& SystemData = BatchSystems[BatchSystemIndex];
+			int32 SystemPathIndex = SystemPathStartIndex + BatchSystemIndex;
+			if (SystemData.System == nullptr)
+			{
+				continue;
+			}
+
+			UE_LOG(LogNiagaraTraversalCacheAuditCommandlet, Display, TEXT("[%i / %i] Collecting Legacy Data - %s"), SystemPathIndex + 1, SystemPaths.Num(), *SystemPaths[SystemPathIndex]);
+			CollectScriptDataForSystem(SystemData);
+		}
+
+		// Enable the traversal cache and collect the script and function data again on worker threads and compare with the legacy data.
+		EnableTraversalCacheCVar->SetWithCurrentPriority(1);
+		ParallelFor(BatchSystems.Num(), [&](int32 BatchSystemIndex)
+		{
+			if (MaxErrorCount != 0 && Results.ErrorCount > MaxErrorCount)
+			{
+				return;
+			}
+			int32 SystemPathIndex = SystemPathStartIndex + BatchSystemIndex;
+			FSystemData& AuditSystemData = BatchSystems[BatchSystemIndex];
+			if (AuditSystemData.System == nullptr)
+			{
+				UE_LOG(LogNiagaraTraversalCacheAuditCommandlet, Warning, TEXT("[%i / %i] Skipping Validation - %s wasn't loaded"), SystemPathIndex + 1, SystemPaths.Num(), *SystemPaths[SystemPathIndex]);
+				return;
+			}
+
+			UE_LOG(LogNiagaraTraversalCacheAuditCommandlet, Display, TEXT("[%i / %i] Validating - %s"), SystemPathIndex + 1, SystemPaths.Num(), *SystemPaths[SystemPathIndex]);
+
+			TArray<FAuditResultScriptData> ScriptResults;
+			for (FScriptData& AuditScriptData : AuditSystemData.ScriptDatas)
+			{
+				AuditScript(*AuditSystemData.System, AuditScriptData, ScriptResults);
+			}
+
+			if (ScriptResults.Num() > 0)
+			{
+				FScopeLock ResultLock(&ResultCS);
+				FAuditResultSystemData& SystemResult = Results.SystemDatas.AddDefaulted_GetRef();
+				SystemResult.SystemPath = AuditSystemData.SystemPath;
+				SystemResult.ScriptDatas = ScriptResults;
+				LogSystemResult(SystemResult, Results.SystemDatas.Num() - 1);
+				Results.ErrorCount += SystemResult.CountDifferences();
+			}
+		});
+
+		// Collect garbage.
+		::CollectGarbage(RF_NoFlags);
+	}
+
+	if (Results.SystemDatas.Num() > 0 && bWriteCSV)
+	{
+		WriteCSVFile(Results);
+	}
+}
+
+void UNiagaraTraversalCacheAuditCommandlet::WriteCSVFile(const FAuditResults& Results) const
+{
+	int32 Year, Month, DayOfWeek, Day, Hour, Min, Sec, MSec;
+	FPlatformTime::SystemTime(Year, Month, DayOfWeek, Day, Hour, Min, Sec, MSec);
+	FString FileName = FString::Printf(TEXT("TraversalCache_%04i-%02i-%02i_%02i-%02i-%02i.csv"), Year, Month, Day, Hour, Min, Sec);
+	const FString AuditFilePath = FPaths::ProjectSavedDir() / TEXT("Audit") / FileName;
+
+	TUniquePtr<FArchive> FileArchive(IFileManager::Get().CreateDebugFileWriter(*AuditFilePath));
+	TUniquePtr<FOutputDeviceArchiveWrapper> OutputStream(new FOutputDeviceArchiveWrapper(FileArchive.Get()));
+
+	auto AddRows = [&OutputStream](
+		const FAuditResultSystemData& SystemResult,
+		const FAuditResultScriptData& ScriptResult,
+		const FAuditResultFunctionData& FunctionResult,
+		const TArray<FNiagaraVariableBase>& ItemVariables,
+		const FString& TypeString,
+		const FString& OperationString)
+	{
+		for (const FNiagaraVariableBase& ItemVariable : ItemVariables)
+		{
+			OutputStream->Logf(TEXT("%s, %s, %s, %s, %s, %s, %s, %s"),
+				*SystemResult.SystemPath, *ScriptResult.EmitterName, *ScriptResult.ScriptName,
+				*FunctionResult.FunctionName, FunctionResult.bFunctionEnabled ? TEXT("Enabled") : TEXT("Disabled"),
+				*TypeString, *OperationString, *ItemVariable.GetName().ToString());
+		}
+	};
+
+	OutputStream->Log(TEXT("SystemPath, EmitterName, ScriptName, FunctionName, Enabled, Type, Operation, Name"));
+	for (const FAuditResultSystemData& SystemResult : Results.SystemDatas)
+	{
+		for (const FAuditResultScriptData& ScriptResult : SystemResult.ScriptDatas)
+		{
+			for (const FAuditResultFunctionData& FunctionResult : ScriptResult.FunctionDatas)
+			{
+				AddRows(SystemResult, ScriptResult, FunctionResult, FunctionResult.AddedInputVariables, TEXT("InputVariable"), TEXT("Added"));
+				AddRows(SystemResult, ScriptResult, FunctionResult, FunctionResult.RemovedInputVariables, TEXT("InputVariable"), TEXT("Removed"));
+				AddRows(SystemResult, ScriptResult, FunctionResult, FunctionResult.AddedHiddenInputVariables, TEXT("HiddenInputVariable"), TEXT("Added"));
+				AddRows(SystemResult, ScriptResult, FunctionResult, FunctionResult.RemovedHiddenInputVariables, TEXT("HiddenInputVariable"), TEXT("Removed"));
+
+				AddRows(SystemResult, ScriptResult, FunctionResult, FunctionResult.AddedStaticInputPinVariables, TEXT("StaticInputPin"), TEXT("Added"));
+				AddRows(SystemResult, ScriptResult, FunctionResult, FunctionResult.RemovedStaticInputPinVariables, TEXT("StaticInputPin"), TEXT("Removed"));
+				AddRows(SystemResult, ScriptResult, FunctionResult, FunctionResult.AddedHiddenStaticInputPinVariables, TEXT("HiddenStaticInputPin"), TEXT("Added"));
+				AddRows(SystemResult, ScriptResult, FunctionResult, FunctionResult.RemovedHiddenStaticInputPinVariables, TEXT("HiddenStaticInputPin"), TEXT("Removed"));
+			}
+		}
+	}
+}
+
diff --git a/Engine/Plugins/FX/Niagara/Source/NiagaraEditor/Private/NiagaraEditorModule.cpp b/Engine/Plugins/FX/Niagara/Source/NiagaraEditor/Private/NiagaraEditorModule.cpp
--- a/Engine/Plugins/FX/Niagara/Source/NiagaraEditor/Private/NiagaraEditorModule.cpp
+++ b/Engine/Plugins/FX/Niagara/Source/NiagaraEditor/Private/NiagaraEditorModule.cpp
@@ -233,13 +233,6 @@ static FAutoConsoleVariableRef CVarEnableCustomInlineDynamicInputFormats(
 	TEXT("If > 0 and experimental inline editors for dynamic input trees are enabled, custom formats which are defined on scripts will be applied.\n"),
 	ECVF_Default);
 
-int32 GbEnableTraversalCache = 0;
-static FAutoConsoleVariableRef CVarEnableTraversalCache(
-	TEXT("fx.Niagara.EnableTraversalCache"),
-	GbEnableTraversalCache,
-	TEXT("If > 0 the new traversal cache will be used to speed up utilities like GetStackFunctionInputs.\n"),
-	ECVF_Default);
-
 // this is required for gpu script compilation ticks
 static FNiagaraShaderQueueTickable NiagaraShaderQueueProcessor;
 
diff --git a/Engine/Plugins/FX/Niagara/Source/NiagaraEditor/Private/TraversalCache/StaticContextBuilder.cpp b/Engine/Plugins/FX/Niagara/Source/NiagaraEditor/Private/TraversalCache/StaticContextBuilder.cpp
--- a/Engine/Plugins/FX/Niagara/Source/NiagaraEditor/Private/TraversalCache/StaticContextBuilder.cpp
+++ b/Engine/Plugins/FX/Niagara/Source/NiagaraEditor/Private/TraversalCache/StaticContextBuilder.cpp
@@ -132,6 +132,13 @@ void FStaticContextBuilder::GatherContextFromScriptObject(FStaticContext& Static
 	AddEnumValue(StaticContext, ESelectKeySource::ExternalConstant, SYS_PARAM_SCRIPT_CONTEXT,
 		NAME_None, (int32)FNiagaraUtilities::ConvertScriptUsageToStaticSwitchContext(Script.Usage));
 
+	// This is a special case to handle emitter local space being read in system scripts which really shouldn't be supported since it's not
+	// defined there, but there is already content that uses this and we need to results to be consistent.
+	if (Script.GetUsage() == ENiagaraScriptUsage::SystemSpawnScript || Script.GetUsage() == ENiagaraScriptUsage::SystemUpdateScript)
+	{
+		AddBoolValue(StaticContext, ESelectKeySource::ExternalConstant, SYS_PARAM_EMITTER_LOCALSPACE, NAME_None, false);
+	}
+
 	FTraversalCache::GetTopLevelScriptFunctionInputContexts(Script, FunctionInputContexts, TraversalDebugData);
 }
 
diff --git a/Engine/Plugins/FX/Niagara/Source/NiagaraEditor/Private/TraversalCache/Traversal.cpp b/Engine/Plugins/FX/Niagara/Source/NiagaraEditor/Private/TraversalCache/Traversal.cpp
--- a/Engine/Plugins/FX/Niagara/Source/NiagaraEditor/Private/TraversalCache/Traversal.cpp
+++ b/Engine/Plugins/FX/Niagara/Source/NiagaraEditor/Private/TraversalCache/Traversal.cpp
@@ -25,6 +25,8 @@ TSharedRef<const FTraversalData> FTraversal::TraverseWithVisitor(const FTraversa
 
 bool FTraversal::CanWriteStatics(FTraversalDebugData* InOutDebugData) const
 {
+	UE::TScopeLock ScopeLock(CanWriteStaticsCacheGuard);
+
 	if (bCanWriteStaticsCache.IsSet() == false)
 	{
 		TSet<const FTraversal*> CheckedTraversals;
@@ -35,6 +37,8 @@ bool FTraversal::CanWriteStatics(FTraversalDebugData* InOutDebugData) const
 
 bool FTraversal::CanWriteStaticAttributes(FTraversalDebugData* InOutDebugData) const
 {
+	UE::TScopeLock ScopeLock(CanWriteStaticAttributesCacheGuard);
+
 	if (bCanWriteStaticAttributesCache.IsSet() == false)
 	{
 		TSet<const FTraversal*> CheckedTraversals;
@@ -45,13 +49,24 @@ bool FTraversal::CanWriteStaticAttributes(FTraversalDebugData* InOutDebugData) c
 
 void FTraversal::ResetCachedData() const
 {
-	HashToTraversalDataCache.Empty();
-	bCanWriteStaticsCache.Reset();
-	bCanWriteStaticAttributesCache.Reset();
+	{
+		UE::TScopeLock ScopeLock(HashToTraversalDataCacheGuard);
+		HashToTraversalDataCache.Empty();
+	}
+	{ 
+		UE::TScopeLock ScopeLock(CanWriteStaticsCacheGuard);
+		bCanWriteStaticsCache.Reset();
+	}
+	{
+		UE::TScopeLock ScopeLock(CanWriteStaticAttributesCacheGuard);
+		bCanWriteStaticAttributesCache.Reset();
+	}
 }
 
 void FTraversal::TraverseInternal(const FTraversalCallingContext& CallingContext, FStaticContext& OutLocalStaticContext, FTraversalDebugData* InOutDebugData) const
 {
+	UE::TScopeLock ScopeLock(HashToTraversalDataCacheGuard);
+
 	FTraversalCallingContext::FContextHash CallingContextHash = CallingContext.GenerateHash();
 	FCachedTraversalData* CachedTraversalData = HashToTraversalDataCache.Find(CallingContextHash);
 	if (CachedTraversalData != nullptr)
@@ -72,6 +87,8 @@ void FTraversal::TraverseInternal(const FTraversalCallingContext& CallingContext
 TSharedRef<const FTraversalData> FTraversal::TraverseWithVisitorInternal(const FTraversalCallingContext& CallingContext, const ITraversalVisitor& Visitor,
 	FStaticContext& OutLocalStaticContext, FTraversalDebugData* InOutDebugData) const
 {
+	UE::TScopeLock ScopeLock(HashToTraversalDataCacheGuard);
+
 	FTraversalCallingContext::FContextHash CallingContextHash = CallingContext.GenerateHash();
 	FCachedTraversalData* CachedTraversalData = HashToTraversalDataCache.Find(CallingContextHash);
 	FGuid VisitorId = Visitor.GetVisitorId();
diff --git a/Engine/Plugins/FX/Niagara/Source/NiagaraEditor/Private/TraversalCache/Traversal.h b/Engine/Plugins/FX/Niagara/Source/NiagaraEditor/Private/TraversalCache/Traversal.h
--- a/Engine/Plugins/FX/Niagara/Source/NiagaraEditor/Private/TraversalCache/Traversal.h
+++ b/Engine/Plugins/FX/Niagara/Source/NiagaraEditor/Private/TraversalCache/Traversal.h
@@ -2,6 +2,7 @@
 
 #pragma once
 
+#include "Async/RecursiveMutex.h"
 #include "Delegates/DelegateCombinations.h"
 #include "TraversalCache/TraversalShared.h"
 
@@ -173,8 +174,13 @@ private:
 	};
 
 	mutable TMap<FTraversalCallingContext::FContextHash, FCachedTraversalData> HashToTraversalDataCache;
+	mutable UE::FRecursiveMutex HashToTraversalDataCacheGuard;
+
 	mutable TOptional<bool> bCanWriteStaticsCache;
+	mutable UE::FMutex CanWriteStaticsCacheGuard;
+
 	mutable TOptional<bool> bCanWriteStaticAttributesCache;
+	mutable UE::FMutex CanWriteStaticAttributesCacheGuard;	
 };
 
 } // UE::Niagara::TraversalCache
\ No newline at end of file
diff --git a/Engine/Plugins/FX/Niagara/Source/NiagaraEditor/Private/TraversalCache/TraversalBuilder.cpp b/Engine/Plugins/FX/Niagara/Source/NiagaraEditor/Private/TraversalCache/TraversalBuilder.cpp
--- a/Engine/Plugins/FX/Niagara/Source/NiagaraEditor/Private/TraversalCache/TraversalBuilder.cpp
+++ b/Engine/Plugins/FX/Niagara/Source/NiagaraEditor/Private/TraversalCache/TraversalBuilder.cpp
@@ -583,10 +583,6 @@ namespace BuilderHelpers
 
 FTraversalBuilder* FTraversalBuilder::Instance = nullptr;
 
-FCriticalSection FTraversalBuilder::IsValidSelectValueTypeCacheCS;
-FCriticalSection FTraversalBuilder::SelectValueCacheCS;
-FCriticalSection FTraversalBuilder::ExtractedParameterFlagCacheCS;
-
 void FTraversalBuilder::Initialize()
 {
 	if (bInitialized == false)
@@ -1093,7 +1089,7 @@ void FTraversalBuilder::CollectAdditionalTraversalData(FTraversal& Traversal)
 
 bool FTraversalBuilder::IsValidSelectValueTypeInternal(const FNiagaraTypeDefinition& InValueType) const
 {
-	FScopeLock ScopeLock(&IsValidSelectValueTypeCacheCS);
+	UE::TScopeLock ScopeLock(IsValidSelectValueTypeCacheGuard);
 
 	bool* bIsValidType = IsValidSelectValueTypeCache.Find(InValueType);
 	if (bIsValidType == nullptr)
@@ -1106,7 +1102,7 @@ bool FTraversalBuilder::IsValidSelectValueTypeInternal(const FNiagaraTypeDefinit
 
 FSelectValue FTraversalBuilder::CreateSelectValueInternal(const FNiagaraVariable& InVariableValue) const
 {
-	FScopeLock ScopeLock(&SelectValueCacheCS);
+	UE::TScopeLock ScopeLock(SelectValueCacheGuard);
 
 	const FNiagaraTypeDefinition& VariableType = InVariableValue.GetType();
 	UObject* TypeObject = VariableType.IsEnum() ? (UObject*)VariableType.GetEnum() : (UObject*)VariableType.GetStruct();
@@ -1133,7 +1129,7 @@ FSelectValue FTraversalBuilder::CreateSelectValueInternal(const FNiagaraVariable
 
 FSelectValue FTraversalBuilder::CreateSelectValueInternal(const FNiagaraTypeDefinition& ValueType, int32 SelectNumericValue) const
 {
-	FScopeLock ScopeLock(&SelectValueCacheCS);
+	UE::TScopeLock ScopeLock(SelectValueCacheGuard);
 
 	UObject* TypeObject = ValueType.IsEnum() ? (UObject*)ValueType.GetEnum() : (UObject*)ValueType.GetStruct();
 	if (TypeObject == nullptr)
@@ -1158,7 +1154,7 @@ FSelectValue FTraversalBuilder::CreateSelectValueInternal(const FNiagaraTypeDefi
 
 EParameterFlags FTraversalBuilder::ExtractFlagsFromParameterNameInternal(FName ParameterName) const
 {
-	FScopeLock ScopeLock(&ExtractedParameterFlagCacheCS);
+	UE::TScopeLock ScopeLock(ExtractedParameterFlagCacheGuard);
 
 	EParameterFlags* CachedFlags = ExtractedParameterFlagCache.Find(ParameterName);
 	if (CachedFlags != nullptr)
diff --git a/Engine/Plugins/FX/Niagara/Source/NiagaraEditor/Private/TraversalCache/TraversalBuilder.h b/Engine/Plugins/FX/Niagara/Source/NiagaraEditor/Private/TraversalCache/TraversalBuilder.h
--- a/Engine/Plugins/FX/Niagara/Source/NiagaraEditor/Private/TraversalCache/TraversalBuilder.h
+++ b/Engine/Plugins/FX/Niagara/Source/NiagaraEditor/Private/TraversalCache/TraversalBuilder.h
@@ -210,14 +210,15 @@ private:
 	bool bInitialized = false;
 
 	TMap<TSubclassOf<UEdGraphNode>, TSharedRef<IGraphNodeHandler>> GraphNodeHandlers;
-	
-	static FCriticalSection IsValidSelectValueTypeCacheCS;
-	static FCriticalSection SelectValueCacheCS;
-	static FCriticalSection ExtractedParameterFlagCacheCS;
 
 	mutable TMap<FNiagaraTypeDefinition, bool> IsValidSelectValueTypeCache;
+	mutable UE::FMutex IsValidSelectValueTypeCacheGuard;
+
 	mutable TMap<FSelectValueCacheKey, FSelectValue> SelectValueCache;
+	mutable UE::FMutex SelectValueCacheGuard;
+
 	mutable TMap<FName, EParameterFlags> ExtractedParameterFlagCache;
+	mutable UE::FMutex ExtractedParameterFlagCacheGuard;
 };
 
 } // UE::Niagara::TraversalCache
\ No newline at end of file
diff --git a/Engine/Plugins/FX/Niagara/Source/NiagaraEditor/Private/TraversalCache/TraversalCache.cpp b/Engine/Plugins/FX/Niagara/Source/NiagaraEditor/Private/TraversalCache/TraversalCache.cpp
--- a/Engine/Plugins/FX/Niagara/Source/NiagaraEditor/Private/TraversalCache/TraversalCache.cpp
+++ b/Engine/Plugins/FX/Niagara/Source/NiagaraEditor/Private/TraversalCache/TraversalCache.cpp
@@ -156,17 +156,21 @@ FTraversalCache& FTraversalCache::GetInstance()
 
 const UNiagaraScript* FTraversalCache::GetScriptFromPath(const FSoftObjectPath& Path) const
 {
+	UE::TScopeLock ScopeLock(SoftObjectPathToScriptCacheGuard);
+
 	const UNiagaraScript* Script = nullptr;
 	TWeakObjectPtr<const UNiagaraScript>* CachedScriptPtr = SoftObjectPathToScriptCache.Find(Path);
-	if (CachedScriptPtr == nullptr)
+	if (CachedScriptPtr != nullptr)
 	{
-		Script = Cast<UNiagaraScript>(Path.ResolveObject());
-		SoftObjectPathToScriptCache.Add(Path, TWeakObjectPtr(Script));
+		Script = CachedScriptPtr->Get();
 	}
-	else
+
+	if (Script == nullptr)
 	{
-		Script = CachedScriptPtr->Get();
+		Script = Cast<UNiagaraScript>(Path.ResolveObject());
+		SoftObjectPathToScriptCache.Add(Path, TWeakObjectPtr(Script));
 	}
+
 	return Script;
 }
 
@@ -181,6 +185,8 @@ void FTraversalCache::OnScriptApplied(UNiagaraScript* InScript, FGuid InVersionG
 
 TSharedPtr<const FTraversal> FTraversalCache::GetScriptAssetTraversalInternal(const FSoftObjectPath& ScriptSoftObjectPath, const FGuid& ScriptVersion, FTraversalBuilderDebugData* BuilderDebugData) const
 {
+	UE::TScopeLock ScopeLock(CachedTraversalsGuard);
+
 	const UNiagaraScript* Script = GetScriptFromPath(ScriptSoftObjectPath);
 	if (Script == nullptr)
 	{
@@ -227,6 +233,8 @@ TSharedPtr<const FTraversal> FTraversalCache::GetScriptAssetTraversalInternal(co
 
 TSharedPtr<const FTraversal> FTraversalCache::GetTopLevelScriptTraversalInternal(const UNiagaraScript& Script, FTraversalBuilderDebugData* BuilderDebugData) const
 {
+	UE::TScopeLock ScopeLock(CachedTraversalsGuard);
+
 	const UNiagaraScriptSource* ScriptSource = CastChecked<UNiagaraScriptSource>(Script.GetSource(FGuid()));
 
 	FScriptTraversalCacheKey CacheKey = FScriptTraversalCacheKey::FromTopLevelScript(FObjectKey(ScriptSource->NodeGraph), Script.GetUsage(), Script.GetUsageId());
@@ -264,6 +272,8 @@ TSharedPtr<const FTraversal> FTraversalCache::GetTopLevelScriptTraversalInternal
 
 TSharedPtr<const FTraversal> FTraversalCache::GetTopLevelScriptTraversalInternal(const UNiagaraNode& StackNode, FTraversalBuilderDebugData* BuilderDebugData) const
 {
+	UE::TScopeLock ScopeLock(CachedTraversalsGuard);
+
 	const UNiagaraGraph* OwningGraph = StackNode.GetNiagaraGraph();
 	if (OwningGraph == nullptr)
 	{
@@ -304,6 +314,8 @@ TSharedPtr<const FTraversal> FTraversalCache::GetTopLevelScriptTraversalInternal
 
 void FTraversalCache::GetTopLevelScriptFunctionInputContextsInternal(const UNiagaraScript& Script, TMap<FGuid, FStaticContext>& OutFunctionInputContexts, FTraversalDebugData* TraversalDebugData) const
 {
+	UE::TScopeLock ScopeLock(CachedParsedRapidIterationParameterDataGuard);
+
 	// Static inputs are stored as rapid iteration parameters, and their names are mangled to keep them unique. Modules
 	// will be expecting these to be in the format Module.[InputName] so the names must be unmangled, and disambiguated by
 	// module node guid so that they can be gathered when creating function contexts.  Emitter rapid iteration parameters
@@ -708,9 +720,18 @@ void FTraversalCache::GetStackFunctionStaticInputPinsInternal(
 
 void FTraversalCache::InvalidateCacheInternal() const
 {
-	CachedTraversals.Empty();
-	CachedParsedRapidIterationParameterData.Empty();
-	SoftObjectPathToScriptCache.Empty();
+	{
+		UE::TScopeLock ScopeLock(CachedTraversalsGuard);
+		CachedTraversals.Empty();
+	}
+	{
+		UE::TScopeLock ScopeLock(CachedParsedRapidIterationParameterDataGuard);
+		CachedParsedRapidIterationParameterData.Empty();
+	}
+	{
+		UE::TScopeLock ScopeLock(SoftObjectPathToScriptCacheGuard);
+		SoftObjectPathToScriptCache.Empty();
+	}
 }
 
 EStackFunctionReadFilterFlags operator|(
diff --git a/Engine/Plugins/FX/Niagara/Source/NiagaraEditor/Private/TraversalCache/TraversalCache.h b/Engine/Plugins/FX/Niagara/Source/NiagaraEditor/Private/TraversalCache/TraversalCache.h
--- a/Engine/Plugins/FX/Niagara/Source/NiagaraEditor/Private/TraversalCache/TraversalCache.h
+++ b/Engine/Plugins/FX/Niagara/Source/NiagaraEditor/Private/TraversalCache/TraversalCache.h
@@ -231,8 +231,13 @@ private:
 	TSharedPtr<FTraversalBuilder> TraversalBuilder;
 
 	mutable TMap<FScriptTraversalCacheKey, FScriptTraversalCacheValue> CachedTraversals;
+	mutable UE::FMutex CachedTraversalsGuard;
+
 	mutable TMap<FScriptTraversalCacheKey, FParsedRapidIterationParameterData> CachedParsedRapidIterationParameterData;
+	mutable UE::FMutex CachedParsedRapidIterationParameterDataGuard;
+
 	mutable TMap<FSoftObjectPath, TWeakObjectPtr<const UNiagaraScript>> SoftObjectPathToScriptCache;
+	mutable UE::FMutex SoftObjectPathToScriptCacheGuard;
 };
 
 EStackFunctionReadFilterFlags operator|(
diff --git a/Engine/Plugins/FX/Niagara/Source/NiagaraEditor/Private/ViewModels/Stack/NiagaraStackGraphUtilities.cpp b/Engine/Plugins/FX/Niagara/Source/NiagaraEditor/Private/ViewModels/Stack/NiagaraStackGraphUtilities.cpp
--- a/Engine/Plugins/FX/Niagara/Source/NiagaraEditor/Private/ViewModels/Stack/NiagaraStackGraphUtilities.cpp
+++ b/Engine/Plugins/FX/Niagara/Source/NiagaraEditor/Private/ViewModels/Stack/NiagaraStackGraphUtilities.cpp
@@ -1135,7 +1135,7 @@ namespace UE::Niagara::TraversalCache::StackGraphUtilities
 void FNiagaraStackGraphUtilities::GetStackFunctionInputs(const UNiagaraNodeFunctionCall& FunctionCallNode, TArray<FNiagaraVariable>& OutInputVariables, TSet<FNiagaraVariable>& OutHiddenVariables, FCompileConstantResolver ConstantResolver, ENiagaraGetStackFunctionInputPinsOptions Options, bool bIgnoreDisabled)
 {
 	SCOPE_CYCLE_COUNTER(STAT_NiagaraEditor_StackGraphUtilities_GetStackFunctionInputs);
-	if (GbEnableTraversalCache == 0)
+	if (GbNiagaraEnableTraversalCache == 0)
 	{
 		GetStackFunctionInputsLegacy(FunctionCallNode, OutInputVariables, OutHiddenVariables, ConstantResolver, Options, bIgnoreDisabled);
 	}
@@ -1152,7 +1152,7 @@ void FNiagaraStackGraphUtilities::GetStackFunctionInputs(const UNiagaraNodeFunct
 void FNiagaraStackGraphUtilities::GetStackFunctionInputs(const UNiagaraNodeFunctionCall& FunctionCallNode, TArray<FNiagaraVariable>& OutInputVariables, FCompileConstantResolver ConstantResolver, ENiagaraGetStackFunctionInputPinsOptions Options, bool bIgnoreDisabled)
 {
 	SCOPE_CYCLE_COUNTER(STAT_NiagaraEditor_StackGraphUtilities_GetStackFunctionInputs);
-	if (GbEnableTraversalCache == 0)
+	if (GbNiagaraEnableTraversalCache == 0)
 	{
 		GetStackFunctionInputsLegacy(FunctionCallNode, OutInputVariables, ConstantResolver, Options, bIgnoreDisabled);
 	}
@@ -1173,7 +1173,7 @@ void FNiagaraStackGraphUtilities::GetStackFunctionInputs(
 	ENiagaraGetStackFunctionInputPinsOptions Options, bool bIgnoreDisabled)
 {
 	SCOPE_CYCLE_COUNTER(STAT_NiagaraEditor_StackGraphUtilities_GetStackFunctionInputs);
-	if (GbEnableTraversalCache == 0)
+	if (GbNiagaraEnableTraversalCache == 0)
 	{
 		GetStackFunctionInputsLegacy(FunctionCallNode, OutInputVariables, Options, bIgnoreDisabled);
 	}
@@ -1729,7 +1729,7 @@ void FNiagaraStackGraphUtilities::GetStackFunctionStaticSwitchPinsLegacy(const U
 void FNiagaraStackGraphUtilities::GetStackFunctionStaticSwitchPins(const UNiagaraNodeFunctionCall& FunctionCallNode, TArray<UEdGraphPin*>& OutInputPins, TSet<UEdGraphPin*>& OutHiddenPins, FCompileConstantResolver& ConstantResolver)
 {
 	SCOPE_CYCLE_COUNTER(STAT_NiagaraEditor_StackGraphUtilities_GetStackFunctionStaticSwitchPins);
-	if (GbEnableTraversalCache == 0)
+	if (GbNiagaraEnableTraversalCache == 0)
 	{
 		GetStackFunctionStaticSwitchPinsLegacy(FunctionCallNode, OutInputPins, OutHiddenPins, ConstantResolver);
 	}
diff --git a/Engine/Plugins/FX/Niagara/Source/NiagaraEditor/Public/Commandlets/NiagaraTraversalCacheAuditCommandlet.h b/Engine/Plugins/FX/Niagara/Source/NiagaraEditor/Public/Commandlets/NiagaraTraversalCacheAuditCommandlet.h
new file mode 100644
--- /dev/null
+++ b/Engine/Plugins/FX/Niagara/Source/NiagaraEditor/Public/Commandlets/NiagaraTraversalCacheAuditCommandlet.h
@@ -0,0 +1,105 @@
+// Copyright Epic Games, Inc. All Rights Reserved.
+
+#pragma once
+
+#include "Commandlets/Commandlet.h"
+#include "NiagaraTypes.h"
+
+#include "NiagaraTraversalCacheAuditCommandlet.generated.h"
+
+struct FNiagaraEmitterHandle;
+class  IConsoleVariable;
+class  UEdGraphPin;
+class  UNiagaraNodeFunctionCall;
+class  UNiagaraScript;
+class  UNiagaraSystem;
+
+UCLASS(config=Editor)
+class UNiagaraTraversalCacheAuditCommandlet : public UCommandlet
+{
+	GENERATED_UCLASS_BODY()
+
+	int32 Main(const FString& Params) override;
+
+private:
+	struct FFunctionData
+	{
+		UNiagaraNodeFunctionCall* FunctionCallNode = nullptr;
+		TArray<FNiagaraVariable> InputVariables;
+		TSet<FNiagaraVariable> HiddenInputVariables;
+		TArray<UEdGraphPin*> InputPins;
+		TSet<UEdGraphPin*> HiddenInputPins;
+	};
+
+	struct FScriptData
+	{
+		UNiagaraScript* Script = nullptr;
+		const FNiagaraEmitterHandle* EmitterHandle = nullptr;
+		TArray<FFunctionData> FunctionDatas;
+	};
+
+	struct FSystemData
+	{
+		FString SystemPath;
+		UNiagaraSystem* System = nullptr;
+		TArray<FScriptData> ScriptDatas;
+	};
+
+	struct FAuditResultFunctionData
+	{
+		FString FunctionName;
+		FName FunctionScriptPath = NAME_None;
+		bool bFunctionEnabled = true;
+
+		TArray<FNiagaraVariableBase> AddedInputVariables;
+		TArray<FNiagaraVariableBase> RemovedInputVariables;
+		TArray<FNiagaraVariableBase> AddedHiddenInputVariables;
+		TArray<FNiagaraVariableBase> RemovedHiddenInputVariables;
+
+		TArray<FNiagaraVariableBase> AddedStaticInputPinVariables;
+		TArray<FNiagaraVariableBase> RemovedStaticInputPinVariables;
+		TArray<FNiagaraVariableBase> AddedHiddenStaticInputPinVariables;
+		TArray<FNiagaraVariableBase> RemovedHiddenStaticInputPinVariables;
+
+		int32 CountDifferences() const;
+	};
+
+	struct FAuditResultScriptData
+	{
+		FString ScriptName;
+		FString EmitterName;
+		TArray<FAuditResultFunctionData> FunctionDatas;
+		int32 CountDifferences() const;
+	};
+
+	struct FAuditResultSystemData
+	{
+		FString SystemPath;
+		TArray<FAuditResultScriptData> ScriptDatas;
+		int32 CountDifferences() const;
+	};
+
+	struct FAuditResults
+	{
+		TArray<FAuditResultSystemData> SystemDatas;
+		int32 ErrorCount = 0;
+	};
+
+	void GatherAllSystemPaths(TArray<FString>& OutSystemPaths) const;
+	void CollectFunctionData(UNiagaraSystem& System, const FNiagaraEmitterHandle* EmitterHandle, UNiagaraScript* Script, TArray<FFunctionData>& OutFunctionDatas) const;
+	void CollectScriptData(UNiagaraSystem* System, const FNiagaraEmitterHandle* EmitterHandle, UNiagaraScript* Script, TArray<FScriptData>& OutScriptDatas) const;
+	void CollectScriptDataForSystem(FSystemData& SystemData) const;
+	void AuditScript(UNiagaraSystem& AuditSystem, FScriptData& AuditScriptData, TArray<FAuditResultScriptData>& OutScriptResults) const;
+	void LogSystemResult(const FAuditResultSystemData& SystemResult, int32 SystemResultIndex) const;
+	void AuditSystems(const TArray<FString>& SystemPaths, FAuditResults& Results) const;
+	void WriteCSVFile(const FAuditResults& Results) const;
+
+private:
+	IConsoleVariable* EnableTraversalCacheCVar = nullptr;
+
+	int32 MaxErrorCount = 0;
+	int32 BatchSize = 100;
+	bool bWriteCSV = true;
+
+	static FCriticalSection ResultCS;
+};
\ No newline at end of file
diff --git a/Engine/Plugins/FX/Niagara/Source/NiagaraEditor/Public/NiagaraEditorModule.h b/Engine/Plugins/FX/Niagara/Source/NiagaraEditor/Public/NiagaraEditorModule.h
--- a/Engine/Plugins/FX/Niagara/Source/NiagaraEditor/Public/NiagaraEditorModule.h
+++ b/Engine/Plugins/FX/Niagara/Source/NiagaraEditor/Public/NiagaraEditorModule.h
@@ -77,7 +77,6 @@ extern NIAGARAEDITOR_API int32 GbShowNiagaraDeveloperWindows;
 extern NIAGARAEDITOR_API int32 GbPreloadSelectablePluginAssetsOnDemand;
 extern NIAGARAEDITOR_API int32 GbEnableExperimentalInlineDynamicInputs;
 extern NIAGARAEDITOR_API int32 GbEnableCustomInlineDynamicInputFormats;
-extern NIAGARAEDITOR_API int32 GbEnableTraversalCache;
 
 /* Defines methods for allowing external modules to supply widgets to the core editor module. */
 class INiagaraEditorWidgetProvider
diff --git a/Engine/Plugins/FX/Niagara/Source/NiagaraShader/Private/NiagaraShared.cpp b/Engine/Plugins/FX/Niagara/Source/NiagaraShader/Private/NiagaraShared.cpp
--- a/Engine/Plugins/FX/Niagara/Source/NiagaraShader/Private/NiagaraShared.cpp
+++ b/Engine/Plugins/FX/Niagara/Source/NiagaraShader/Private/NiagaraShared.cpp
@@ -43,6 +43,13 @@ static FAutoConsoleVariableRef CVarNiagaraTranslatorSilenceFailIfNotSet(
 	ECVF_Default
 );
 
+int32 GbNiagaraEnableTraversalCache = 1;
+static FAutoConsoleVariableRef CVarEnableTraversalCache(
+	TEXT("fx.Niagara.EnableTraversalCache"),
+	GbNiagaraEnableTraversalCache,
+	TEXT("If > 0 the new traversal cache will be used to speed up utilities like GetStackFunctionInputs.\n"),
+	ECVF_Default);
+
 #if WITH_EDITOR
 	FNiagaraCompilationQueue* FNiagaraCompilationQueue::Singleton = nullptr;
 #endif
diff --git a/Engine/Plugins/FX/Niagara/Source/NiagaraShader/Public/NiagaraShared.h b/Engine/Plugins/FX/Niagara/Source/NiagaraShader/Public/NiagaraShared.h
--- a/Engine/Plugins/FX/Niagara/Source/NiagaraShader/Public/NiagaraShared.h
+++ b/Engine/Plugins/FX/Niagara/Source/NiagaraShader/Public/NiagaraShared.h
@@ -1009,6 +1009,7 @@ private:
 
 /** CVars */
 extern NIAGARASHADER_API int32 GNiagaraTranslatorFailIfNotSetSeverity;
+extern NIAGARASHADER_API int32 GbNiagaraEnableTraversalCache;
 
 namespace FNiagaraCVarUtilities
 {
