- Only allow stateless modules to be added if they pass the allow list

diff --git a/Engine/Plugins/FX/Niagara/Source/NiagaraEditor/Private/NiagaraEditorUtilities.cpp b/Engine/Plugins/FX/Niagara/Source/NiagaraEditor/Private/NiagaraEditorUtilities.cpp
--- a/Engine/Plugins/FX/Niagara/Source/NiagaraEditor/Private/NiagaraEditorUtilities.cpp
+++ b/Engine/Plugins/FX/Niagara/Source/NiagaraEditor/Private/NiagaraEditorUtilities.cpp
@@ -4,6 +4,8 @@
 
 #include "AssetRegistry/AssetRegistryModule.h"
 #include "AssetToolsModule.h"
+#include "ClassViewerModule.h"
+#include "ClassViewerFilter.h"
 #include "ContentBrowserModule.h"
 #include "EdGraphSchema_Niagara.h"
 #include "GeneralProjectSettings.h"
@@ -4586,6 +4588,17 @@ bool FNiagaraEditorUtilities::IsEnumIndexVisible(const UEnum* Enum, int32 Index)
 	return FNiagaraEnumIndexVisibilityCache::GetVisibility(Enum, Index);
 }
 
+bool FNiagaraEditorUtilities::IsClassVisibileInGlobalFilter(UClass* Class)
+{
+	FClassViewerModule* ClassViewerModule = FModuleManager::GetModulePtr<FClassViewerModule>("ClassViewer");
+	const TSharedPtr<IClassViewerFilter> GlobalClassFilter = ClassViewerModule ? ClassViewerModule->GetGlobalClassViewerFilter() : TSharedPtr<IClassViewerFilter>();
+	if (GlobalClassFilter.IsValid())
+	{
+		return GlobalClassFilter->IsClassAllowed(FClassViewerInitializationOptions(), Class, ClassViewerModule->CreateFilterFuncs());
+	}
+	return true;
+}
+
 void FNiagaraEditorUtilities::GetScriptMessageStores(UNiagaraScript* InScript, TArray<FNiagaraMessageSourceAndStore>& OutNiagaraMessageStores)
 {
 	UNiagaraScriptSource* Source = Cast<UNiagaraScriptSource>(InScript->GetLatestSource());
diff --git a/Engine/Plugins/FX/Niagara/Source/NiagaraEditor/Private/ViewModels/Stack/NiagaraStackStatelessEmitterSimulateGroup.cpp b/Engine/Plugins/FX/Niagara/Source/NiagaraEditor/Private/ViewModels/Stack/NiagaraStackStatelessEmitterSimulateGroup.cpp
--- a/Engine/Plugins/FX/Niagara/Source/NiagaraEditor/Private/ViewModels/Stack/NiagaraStackStatelessEmitterSimulateGroup.cpp
+++ b/Engine/Plugins/FX/Niagara/Source/NiagaraEditor/Private/ViewModels/Stack/NiagaraStackStatelessEmitterSimulateGroup.cpp
@@ -5,6 +5,7 @@
 #include "IDetailTreeNode.h"
 #include "NiagaraClipboard.h"
 #include "NiagaraEditorStyle.h"
+#include "NiagaraEditorUtilities.h"
 #include "NiagaraStackEditorData.h"
 #include "PropertyHandle.h"
 #include "ScopedTransaction.h"
@@ -129,7 +130,10 @@ public:
 				if (StatelessModule->IsModuleEnabled() == false && 
 					StackEditorData->GetStatelessModuleShowWhenDisabled(UNiagaraStackStatelessModuleItem::GenerateStackEditorDataKey(StatelessModule)) == false)
 				{
-					OutAddActions.Add(MakeShared<FNiagaraStatelessEmitterAddModuleAction>(StatelessModule));
+					if (FNiagaraEditorUtilities::IsClassVisibileInGlobalFilter(StatelessModule->GetClass()))
+					{
+						OutAddActions.Add(MakeShared<FNiagaraStatelessEmitterAddModuleAction>(StatelessModule));
+					}
 				}
 			}
 		}
diff --git a/Engine/Plugins/FX/Niagara/Source/NiagaraEditor/Public/NiagaraEditorUtilities.h b/Engine/Plugins/FX/Niagara/Source/NiagaraEditor/Public/NiagaraEditorUtilities.h
--- a/Engine/Plugins/FX/Niagara/Source/NiagaraEditor/Public/NiagaraEditorUtilities.h
+++ b/Engine/Plugins/FX/Niagara/Source/NiagaraEditor/Public/NiagaraEditorUtilities.h
@@ -420,6 +420,9 @@ namespace FNiagaraEditorUtilities
 
 	bool IsEnumIndexVisible(const UEnum* Enum, int32 Index);
 
+	/** Test to see id the global filter allows this class or not.   Classes that are not allowed can fail validation. */
+	NIAGARAEDITOR_API bool IsClassVisibileInGlobalFilter(UClass* Class);
+
 	void GetScriptMessageStores(UNiagaraScript* InScript, TArray<FNiagaraMessageSourceAndStore>& OutNiagaraMessageStores);
 
 	NIAGARAEDITOR_API bool IsEditorDataInterfaceInstance(const UNiagaraDataInterface* DataInterface);
diff --git a/Engine/Plugins/FX/Niagara/Source/NiagaraEditorWidgets/Private/SNiagaraOverviewGraph.cpp b/Engine/Plugins/FX/Niagara/Source/NiagaraEditorWidgets/Private/SNiagaraOverviewGraph.cpp
--- a/Engine/Plugins/FX/Niagara/Source/NiagaraEditorWidgets/Private/SNiagaraOverviewGraph.cpp
+++ b/Engine/Plugins/FX/Niagara/Source/NiagaraEditorWidgets/Private/SNiagaraOverviewGraph.cpp
@@ -1,8 +1,6 @@
 // Copyright Epic Games, Inc. All Rights Reserved.
 
 #include "SNiagaraOverviewGraph.h"
-#include "ClassViewerModule.h"
-#include "ClassViewerFilter.h"
 #include "EdGraph/EdGraph.h"
 #include "EdGraph/EdGraphNode.h"
 #include "EdGraphSchema_Niagara.h"
@@ -295,14 +293,7 @@ void SNiagaraOverviewGraph::OnCreateEmptyEmitter()
 
 bool SNiagaraOverviewGraph::CanCreateStatelessEmitter() const
 {
-	FClassViewerModule* ClassViewerModule = FModuleManager::GetModulePtr<FClassViewerModule>("ClassViewer");
-	const TSharedPtr<IClassViewerFilter> GlobalClassFilter = ClassViewerModule ? ClassViewerModule->GetGlobalClassViewerFilter() : TSharedPtr<IClassViewerFilter>();
-	if (GlobalClassFilter.IsValid())
-	{
-		return GlobalClassFilter->IsClassAllowed(FClassViewerInitializationOptions(), UNiagaraStatelessEmitter::StaticClass(), ClassViewerModule->CreateFilterFuncs());
-	}
-
-	return true;
+	return FNiagaraEditorUtilities::IsClassVisibileInGlobalFilter(UNiagaraStatelessEmitter::StaticClass());
 }
 
 FText SNiagaraOverviewGraph::GetCreateStatelessEmitterTooltip() const
