- Remove experimental from decal attributes

diff --git a/Engine/Plugins/FX/Niagara/Source/Niagara/Internal/Stateless/Modules/NiagaraStatelessModule_DecalAttributes.h b/Engine/Plugins/FX/Niagara/Source/Niagara/Internal/Stateless/Modules/NiagaraStatelessModule_DecalAttributes.h
--- a/Engine/Plugins/FX/Niagara/Source/Niagara/Internal/Stateless/Modules/NiagaraStatelessModule_DecalAttributes.h
+++ b/Engine/Plugins/FX/Niagara/Source/Niagara/Internal/Stateless/Modules/NiagaraStatelessModule_DecalAttributes.h
@@ -7,7 +7,7 @@
 
 #include "NiagaraStatelessModule_DecalAttributes.generated.h"
 
-UCLASS(MinimalAPI, EditInlineNew, Experimental, meta = (DisplayName = "Decal Attributes"))
+UCLASS(MinimalAPI, EditInlineNew, meta = (DisplayName = "Decal Attributes"))
 class UNiagaraStatelessModule_DecalAttributes : public UNiagaraStatelessModule
 {
 	GENERATED_BODY()
