PR #13328: Niagara Data Interface allows us to set material parameters on Material Instance Dynamic objects, from a CPU stage.

diff --git a/Engine/Plugins/FX/Niagara/Source/Niagara/Private/NiagaraDataInterfaceMaterialInstanceDynamic.cpp b/Engine/Plugins/FX/Niagara/Source/Niagara/Private/NiagaraDataInterfaceMaterialInstanceDynamic.cpp
new file mode 100644
--- /dev/null
+++ b/Engine/Plugins/FX/Niagara/Source/Niagara/Private/NiagaraDataInterfaceMaterialInstanceDynamic.cpp
@@ -0,0 +1,240 @@
+// Copyright (C) Microsoft Corporation. All rights reserved.
+
+#include "NiagaraDataInterfaceMaterialInstanceDynamic.h"
+
+#include "NiagaraDataInterfaceUtilities.h"
+#include "NiagaraGpuComputeDispatchInterface.h"
+#include "NiagaraSystemInstance.h"
+
+#include "Async/Async.h"
+#include "UObject/WeakInterfacePtr.h"
+
+#include UE_INLINE_GENERATED_CPP_BY_NAME(NiagaraDataInterfaceMaterialInstanceDynamic)
+
+#define LOCTEXT_NAMESPACE "NiagaraDataInterfaceMaterialInstanceDynamic"
+
+namespace NDIMIDPrivate
+{
+	static const FName SetScalarParameterFuncName(TEXT("SetScalarParameter"));
+	static const FText SetScalarParameterDescription = LOCTEXT("NiagaraMID_SetScalarParameter_Desc", "Sets a scalar parameter to a Material instance dynamic by name.");
+	static const FName SetVector4ParameterFuncName(TEXT("SetVector4Parameter"));
+	static const FText SetVector4ParameterDescription = LOCTEXT("NiagaraMID_SetVector4Parameter_Desc", "Sets a Vector4 parameter to a Material instance dynamic by name.");
+
+	struct FInstanceData_GT
+	{
+		TWeakObjectPtr<UMaterialInstanceDynamic> MID;
+	};
+
+	static void VMSetScalarParameter(FVectorVMExternalFunctionContext& Context, FName MaterialScalarParamName)
+	{
+		VectorVM::FUserPtrHandler<FInstanceData_GT> InstanceData(Context);
+		FNDIInputParam<float> InScalarValue(Context);
+
+		TWeakObjectPtr<UMaterialInstanceDynamic> MID = InstanceData->MID;
+		const float ScalarParam = InScalarValue.Get();
+		if (MID.IsValid())
+		{
+			AsyncTask(ENamedThreads::GameThread,
+				[WeakMID=MID, ParamName=MaterialScalarParamName, ScalarValue=ScalarParam]()
+				{
+					if (WeakMID.IsValid())
+					{
+						auto* MID = WeakMID.Get();
+						if (MID)
+						{
+							MID->SetScalarParameterValue(ParamName, ScalarValue);
+						}
+					}
+				}
+			);
+		}
+	}
+
+	static void VMSetVector4Parameter(FVectorVMExternalFunctionContext& Context, FName MaterialVectorParamName)
+	{
+		VectorVM::FUserPtrHandler<FInstanceData_GT> InstanceData(Context);
+		FNDIInputParam<FVector4f> InVectorValue(Context);
+
+		TWeakObjectPtr<UMaterialInstanceDynamic> MID = InstanceData->MID;
+		const FVector4f Vec4Param = InVectorValue.Get();
+		if (MID.IsValid())
+		{
+			AsyncTask(ENamedThreads::GameThread,
+				[WeakMID=MID, ParamName=MaterialVectorParamName, Vector4Value=Vec4Param]()
+				{
+					if (WeakMID.IsValid())
+					{
+						auto* MID = WeakMID.Get();
+						if (MID)
+						{
+							MID->SetVectorParameterValue(ParamName, Vector4Value);
+						}
+					}
+				}
+			);
+		}
+	}
+
+}
+
+UNiagaraDataInterfaceMaterialInstanceDynamic::UNiagaraDataInterfaceMaterialInstanceDynamic(const FObjectInitializer& Initializer)
+	: Super(Initializer)
+	, DefaultMaterialInst(nullptr)
+	, InstancedMaterialParamBinding()
+{
+}
+
+void UNiagaraDataInterfaceMaterialInstanceDynamic::InitBindingMembers()
+{
+#if WITH_EDITORONLY_DATA
+	InstancedMaterialParamBinding.SetUsage(ENiagaraParameterBindingUsage::User);
+	InstancedMaterialParamBinding.SetAllowedObjects({ UObject::StaticClass() });
+#endif
+}
+
+void UNiagaraDataInterfaceMaterialInstanceDynamic::PostInitProperties()
+{
+	Super::PostInitProperties();
+
+	if (HasAnyFlags(RF_ClassDefaultObject))
+	{
+		const ENiagaraTypeRegistryFlags Flags =
+			ENiagaraTypeRegistryFlags::AllowUserVariable |
+			ENiagaraTypeRegistryFlags::AllowEmitterVariable |
+			ENiagaraTypeRegistryFlags::AllowSystemVariable |
+			ENiagaraTypeRegistryFlags::AllowParameter;
+
+		FNiagaraTypeRegistry::Register(FNiagaraTypeDefinition(GetClass()), Flags);
+
+		InitBindingMembers();
+	}
+}
+
+bool UNiagaraDataInterfaceMaterialInstanceDynamic::InitPerInstanceData(void* PerInstanceData, FNiagaraSystemInstance* SystemInstance)
+{
+	NDIMIDPrivate::FInstanceData_GT* InstanceData = new (PerInstanceData) NDIMIDPrivate::FInstanceData_GT();
+	InstanceData->MID = GetMaterialInstance(SystemInstance);
+	return true;
+}
+
+bool UNiagaraDataInterfaceMaterialInstanceDynamic::PerInstanceTick(void* PerInstanceData, FNiagaraSystemInstance* SystemInstance, float DeltaSeconds)
+{
+	NDIMIDPrivate::FInstanceData_GT* InstanceData = static_cast<NDIMIDPrivate::FInstanceData_GT*>(PerInstanceData);
+	InstanceData->MID = GetMaterialInstance(SystemInstance);
+	return false;
+}
+
+UMaterialInstanceDynamic* UNiagaraDataInterfaceMaterialInstanceDynamic::GetMaterialInstance(FNiagaraSystemInstance* SystemInstance) const
+{
+	UMaterialInstanceDynamic* MaterialInstance = DefaultMaterialInst;
+	if (InstancedMaterialParamBinding.ResolvedParameter.IsValid())
+	{
+		FNiagaraParameterDirectBinding<UObject*> InstancedMaterialDirectBinding;
+		UObject* InstancedMaterialObj = InstancedMaterialDirectBinding.Init(SystemInstance->GetInstanceParameters(), InstancedMaterialParamBinding.ResolvedParameter);
+		if (UMaterialInstanceDynamic* InstancedMat = Cast<UMaterialInstanceDynamic>(InstancedMaterialObj))
+		{
+			MaterialInstance = InstancedMat;
+		}
+	}
+
+	return MaterialInstance;
+}
+
+bool UNiagaraDataInterfaceMaterialInstanceDynamic::CanExecuteOnTarget(ENiagaraSimTarget Target) const
+{
+	return Target == ENiagaraSimTarget::CPUSim;
+}
+
+int32 UNiagaraDataInterfaceMaterialInstanceDynamic::PerInstanceDataSize() const
+{
+	return sizeof(NDIMIDPrivate::FInstanceData_GT);
+}
+
+void UNiagaraDataInterfaceMaterialInstanceDynamic::DestroyPerInstanceData(void* PerInstanceData, FNiagaraSystemInstance* SystemInstance)
+{
+	auto* InstanceData = static_cast<NDIMIDPrivate::FInstanceData_GT*>(PerInstanceData);
+	InstanceData->~FInstanceData_GT();
+}
+
+bool UNiagaraDataInterfaceMaterialInstanceDynamic::CopyToInternal(UNiagaraDataInterface* Destination) const
+{
+	if (!Super::CopyToInternal(Destination))
+	{
+		return false;
+	}
+
+	auto* OtherTyped = CastChecked< UNiagaraDataInterfaceMaterialInstanceDynamic >(Destination);
+
+	OtherTyped->DefaultMaterialInst = DefaultMaterialInst;
+	OtherTyped->InstancedMaterialParamBinding = InstancedMaterialParamBinding;
+
+	return true;
+}
+
+bool UNiagaraDataInterfaceMaterialInstanceDynamic::Equals(const UNiagaraDataInterface* Other) const
+{
+	if (!Super::Equals(Other))
+	{
+		return false;
+	}
+
+	const UNiagaraDataInterfaceMaterialInstanceDynamic* OtherTyped = CastChecked< const UNiagaraDataInterfaceMaterialInstanceDynamic >(Other);
+
+	return OtherTyped->DefaultMaterialInst == DefaultMaterialInst
+		&& OtherTyped->InstancedMaterialParamBinding == InstancedMaterialParamBinding;
+}
+
+
+void UNiagaraDataInterfaceMaterialInstanceDynamic::GetFunctions(TArray< FNiagaraFunctionSignature >& OutFunctions)
+{
+	FNiagaraFunctionSignature SigTemplate;
+	SigTemplate.bMemberFunction = true;
+	SigTemplate.bRequiresExecPin = true;
+	SigTemplate.bRequiresContext = false;
+	SigTemplate.bSupportsCPU = true;
+	SigTemplate.bSupportsGPU = false;
+	SigTemplate.FunctionSpecifiers.Add(FName("Parameter Name"));
+	SigTemplate.Inputs.Add(FNiagaraVariable(FNiagaraTypeDefinition(GetClass()), TEXT("MID")));
+
+	// Scalar parameters
+	{
+		FNiagaraFunctionSignature Sig = SigTemplate;
+		Sig.Name = NDIMIDPrivate::SetScalarParameterFuncName;
+		Sig.SetDescription(NDIMIDPrivate::SetScalarParameterDescription);
+		Sig.Inputs.Add(FNiagaraVariable(FNiagaraTypeDefinition::GetFloatDef(), TEXT("Value")));
+		OutFunctions.Add(Sig);
+	}
+
+	// Vector4 parameters
+	{
+		FNiagaraFunctionSignature Sig = SigTemplate;
+		Sig.Name = NDIMIDPrivate::SetVector4ParameterFuncName;
+		Sig.SetDescription(NDIMIDPrivate::SetVector4ParameterDescription);
+		Sig.Inputs.Add(FNiagaraVariable(FNiagaraTypeDefinition::GetVec4Def(), TEXT("Value")));
+		OutFunctions.Add(Sig);
+	}
+}
+
+void UNiagaraDataInterfaceMaterialInstanceDynamic::GetVMExternalFunction(const FVMExternalFunctionBindingInfo& BindingInfo, void* InstanceData, FVMExternalFunction& OutFunc)
+{
+	if (BindingInfo.Name == NDIMIDPrivate::SetScalarParameterFuncName)
+	{
+		OutFunc = FVMExternalFunction::CreateLambda(
+			[ParamName=BindingInfo.FunctionSpecifiers[0].Value](FVectorVMExternalFunctionContext& Context)
+			{
+				NDIMIDPrivate::VMSetScalarParameter(Context, ParamName);
+			}
+		);
+	}
+	else if (BindingInfo.Name == NDIMIDPrivate::SetVector4ParameterFuncName)
+	{
+		OutFunc = FVMExternalFunction::CreateLambda(
+			[ParamName=BindingInfo.FunctionSpecifiers[0].Value](FVectorVMExternalFunctionContext& Context)
+			{
+				NDIMIDPrivate::VMSetVector4Parameter(Context, ParamName);
+			}
+		);
+	}
+}
+
+#undef LOCTEXT_NAMESPACE
diff --git a/Engine/Plugins/FX/Niagara/Source/Niagara/Private/NiagaraDataInterfaceMaterialInstanceDynamic.h b/Engine/Plugins/FX/Niagara/Source/Niagara/Private/NiagaraDataInterfaceMaterialInstanceDynamic.h
new file mode 100644
--- /dev/null
+++ b/Engine/Plugins/FX/Niagara/Source/Niagara/Private/NiagaraDataInterfaceMaterialInstanceDynamic.h
@@ -0,0 +1,49 @@
+// Copyright (C) Microsoft Corporation. All rights reserved.
+
+#pragma once
+
+#include "Materials/MaterialInstanceDynamic.h"
+#include "NiagaraCommon.h"
+#include "NiagaraDataInterface.h"
+#include "NiagaraParameterBinding.h"
+#include "NiagaraShared.h"
+#include "VectorVM.h"
+
+#include "NiagaraDataInterfaceMaterialInstanceDynamic.generated.h"
+
+UCLASS(EditInlineNew, Category = "Material", CollapseCategories, meta = (DisplayName = "Material Instance Dynamic"), MinimalAPI)
+class UNiagaraDataInterfaceMaterialInstanceDynamic : public UNiagaraDataInterface
+{
+	GENERATED_BODY()
+
+public:
+	UNiagaraDataInterfaceMaterialInstanceDynamic(const FObjectInitializer& Initializer);
+
+	// UObject Interface
+	virtual void PostInitProperties() override;
+	// UObject Interface End
+
+	// UNiagaraDataInterface Interface
+	virtual bool InitPerInstanceData(void* PerInstanceData, FNiagaraSystemInstance* SystemInstance) override;
+	virtual bool PerInstanceTick(void* PerInstanceData, FNiagaraSystemInstance* SystemInstance, float DeltaSeconds) override;
+	virtual bool CanExecuteOnTarget(ENiagaraSimTarget Target) const override;
+	virtual int32 PerInstanceDataSize() const override;
+	virtual void DestroyPerInstanceData(void* PerInstanceData, FNiagaraSystemInstance* SystemInstance) override;
+	virtual bool CopyToInternal(UNiagaraDataInterface* Destination) const override;
+	virtual bool Equals(const UNiagaraDataInterface* Other) const override;
+	virtual void GetFunctions(TArray< FNiagaraFunctionSignature >& OutFunctions) override;
+	virtual void GetVMExternalFunction(const FVMExternalFunctionBindingInfo& BindingInfo, void* InstanceData, FVMExternalFunction& OutFunc) override;
+	virtual bool HasPreSimulateTick() const override { return true; }
+	// UNiagaraDataInterface Interface
+
+	UPROPERTY(EditAnywhere, Category = Material)
+	TObjectPtr<UMaterialInstanceDynamic> DefaultMaterialInst;
+
+	UPROPERTY(EditAnywhere, Category = "Material")
+	FNiagaraParameterBinding InstancedMaterialParamBinding;
+
+private:
+	void InitBindingMembers();
+
+	UMaterialInstanceDynamic* GetMaterialInstance(FNiagaraSystemInstance* SystemInstance) const;
+};
\ No newline at end of file
