- Add missing copyright header

#rnx

diff --git a/Engine/Plugins/FX/Niagara/Source/Niagara/Private/NiagaraDataInterfaceMaterialInstanceDynamic.cpp b/Engine/Plugins/FX/Niagara/Source/Niagara/Private/NiagaraDataInterfaceMaterialInstanceDynamic.cpp
--- a/Engine/Plugins/FX/Niagara/Source/Niagara/Private/NiagaraDataInterfaceMaterialInstanceDynamic.cpp
+++ b/Engine/Plugins/FX/Niagara/Source/Niagara/Private/NiagaraDataInterfaceMaterialInstanceDynamic.cpp
@@ -1,4 +1,4 @@
-// Copyright (C) Microsoft Corporation. All rights reserved.
+// Copyright Epic Games, Inc. All Rights Reserved.
 
 #include "NiagaraDataInterfaceMaterialInstanceDynamic.h"
 
diff --git a/Engine/Plugins/FX/Niagara/Source/Niagara/Private/NiagaraDataInterfaceMaterialInstanceDynamic.h b/Engine/Plugins/FX/Niagara/Source/Niagara/Private/NiagaraDataInterfaceMaterialInstanceDynamic.h
--- a/Engine/Plugins/FX/Niagara/Source/Niagara/Private/NiagaraDataInterfaceMaterialInstanceDynamic.h
+++ b/Engine/Plugins/FX/Niagara/Source/Niagara/Private/NiagaraDataInterfaceMaterialInstanceDynamic.h
@@ -1,4 +1,4 @@
-// Copyright (C) Microsoft Corporation. All rights reserved.
+// Copyright Epic Games, Inc. All Rights Reserved.
 
 #pragma once
 
