- Fixed search & scroll being triggered when modifying the stack while a search term was active
- Added 'go to previous search result' when pressing shift and enter in the stack search. More consistent with other parts of Niagara
- Added better stack entry selection mechanisms on top of the previous "root selection".
- Refactored inline parameter selection to select the input row (blue highlight as if clicked) and pulsates it to draw attention if previously visible
- Added generic function to scroll to, select row and or pulsate a row. This is useful for future features as this is effectively a navigational feature that we didn't have before.

- Caveat: SystemSelectionViewModel selection & contained StackViewModel selection should be unified

diff --git a/Engine/Plugins/FX/Niagara/Source/NiagaraEditor/Private/ViewModels/NiagaraSystemSelectionViewModel.cpp b/Engine/Plugins/FX/Niagara/Source/NiagaraEditor/Private/ViewModels/NiagaraSystemSelectionViewModel.cpp
--- a/Engine/Plugins/FX/Niagara/Source/NiagaraEditor/Private/ViewModels/NiagaraSystemSelectionViewModel.cpp
+++ b/Engine/Plugins/FX/Niagara/Source/NiagaraEditor/Private/ViewModels/NiagaraSystemSelectionViewModel.cpp
@@ -516,8 +516,7 @@ void UNiagaraSystemSelectionViewModel::UpdateExternalSelectionState()
 			StackSelectionEntries.Add(SelectionEntry.Entry.Get());
 		}
 	}
-
-	SelectionStackViewModel->SetFocusedStackEntryOverride(nullptr);
+	
 	StackSelection->SetSelectedEntries(StackSelectionEntries);
 	OnEntrySelectionChangedDelegate.Broadcast();
 	if (bEmitterHandleIdSelectionChanged)
diff --git a/Engine/Plugins/FX/Niagara/Source/NiagaraEditor/Private/ViewModels/Stack/NiagaraStackFunctionInput.cpp b/Engine/Plugins/FX/Niagara/Source/NiagaraEditor/Private/ViewModels/Stack/NiagaraStackFunctionInput.cpp
--- a/Engine/Plugins/FX/Niagara/Source/NiagaraEditor/Private/ViewModels/Stack/NiagaraStackFunctionInput.cpp
+++ b/Engine/Plugins/FX/Niagara/Source/NiagaraEditor/Private/ViewModels/Stack/NiagaraStackFunctionInput.cpp
@@ -457,7 +457,7 @@ void UNiagaraStackFunctionInput::PasteFunctionInput(const UNiagaraClipboardFunct
 		SetClipboardContentViaConversionScript(*ClipboardInput);
 	}
 
-	OnCopyPasteDelegate.ExecuteIfBound();
+	RequestPulse();
 	
 	// For children inputs we are more strict and require name & type match as the user is less intentional about it when pasting
 	if(ClipboardInput->ChildrenInputs.Num() > 0)
@@ -3105,7 +3105,7 @@ const UNiagaraClipboardFunctionInput* UNiagaraStackFunctionInput::ToClipboardFun
 		break;
 	}
 
-	OnCopyPasteDelegate.ExecuteIfBound();
+	RequestPulse();
 
 	if(ClipboardInput)
 	{
diff --git a/Engine/Plugins/FX/Niagara/Source/NiagaraEditor/Private/ViewModels/Stack/NiagaraStackModuleItem.cpp b/Engine/Plugins/FX/Niagara/Source/NiagaraEditor/Private/ViewModels/Stack/NiagaraStackModuleItem.cpp
--- a/Engine/Plugins/FX/Niagara/Source/NiagaraEditor/Private/ViewModels/Stack/NiagaraStackModuleItem.cpp
+++ b/Engine/Plugins/FX/Niagara/Source/NiagaraEditor/Private/ViewModels/Stack/NiagaraStackModuleItem.cpp
@@ -1693,7 +1693,7 @@ void UNiagaraStackModuleItem::Copy(UNiagaraClipboardContent* ClipboardContent) c
 	InputRoot->ToClipboardFunctionInputs(ClipboardFunction, MutableView(ClipboardFunction->Inputs));
 	ClipboardContent->Functions.Add(ClipboardFunction);
 
-	OnCopyPasteDelegate.ExecuteIfBound();
+	RequestPulse();
 }
 
 bool UNiagaraStackModuleItem::TestCanPasteWithMessage(const UNiagaraClipboardContent* ClipboardContent, FText& OutMessage) const
diff --git a/Engine/Plugins/FX/Niagara/Source/NiagaraEditor/Private/ViewModels/Stack/NiagaraStackScriptItemGroup.cpp b/Engine/Plugins/FX/Niagara/Source/NiagaraEditor/Private/ViewModels/Stack/NiagaraStackScriptItemGroup.cpp
--- a/Engine/Plugins/FX/Niagara/Source/NiagaraEditor/Private/ViewModels/Stack/NiagaraStackScriptItemGroup.cpp
+++ b/Engine/Plugins/FX/Niagara/Source/NiagaraEditor/Private/ViewModels/Stack/NiagaraStackScriptItemGroup.cpp
@@ -741,7 +741,7 @@ void UNiagaraStackScriptItemGroup::Paste(const UNiagaraClipboardContent* Clipboa
 		int32 PasteIndex = ModuleItems.Num() > 0 ? ModuleItems.Last()->GetModuleIndex() + 1 : 0;
 		PasteModules(ClipboardContent, PasteIndex, OutPasteWarning);
 
-		OnCopyPasteDelegate.ExecuteIfBound();
+		RequestPulse();
 	}
 }
 
diff --git a/Engine/Plugins/FX/Niagara/Source/NiagaraEditor/Private/ViewModels/Stack/NiagaraStackViewModel.cpp b/Engine/Plugins/FX/Niagara/Source/NiagaraEditor/Private/ViewModels/Stack/NiagaraStackViewModel.cpp
--- a/Engine/Plugins/FX/Niagara/Source/NiagaraEditor/Private/ViewModels/Stack/NiagaraStackViewModel.cpp
+++ b/Engine/Plugins/FX/Niagara/Source/NiagaraEditor/Private/ViewModels/Stack/NiagaraStackViewModel.cpp
@@ -277,7 +277,9 @@ void UNiagaraStackViewModel::Tick()
 		{
 			RootEntry->RefreshChildren();
 			bRefreshPending = false;
-			RestartSearch();
+			FNiagaraStackSearchConfig SearchConfig;
+			SearchConfig.bScrollToResult = false;
+			RestartSearch(SearchConfig);
 			bValidatorUpdatePending = true;
 		}
 
@@ -308,10 +310,7 @@ void UNiagaraStackViewModel::SetSearchTextExternal(const FText& NewSearchText)
 }
 
 void UNiagaraStackViewModel::OnSearchTextChanged(const FText& SearchText)
-{
-	// If the user is searching that means we don't want to have a previous override mess with the search
-	FocusedStackEntryOverride.Reset();
-	
+{	
 	if (RootEntry && !CurrentSearchText.EqualTo(SearchText))
 	{
 		CurrentSearchText = SearchText;
@@ -337,12 +336,7 @@ const TArray<UNiagaraStackViewModel::FSearchResult>& UNiagaraStackViewModel::Get
 }
 
 UNiagaraStackEntry* UNiagaraStackViewModel::GetCurrentFocusedEntry() const
-{
-	if(FocusedStackEntryOverride.IsValid() && FocusedStackEntryOverride->IsFinalized() == false)
-	{
-		return FocusedStackEntryOverride.Get();
-	}
-	
+{	
 	if (FocusedSearchResultCache.IsSet() && FocusedSearchResultCache->GetEntry() != nullptr && !FocusedSearchResultCache->GetEntry()->IsFinalized())
 	{
 		return FocusedSearchResultCache->GetEntry();
@@ -550,7 +544,7 @@ void UNiagaraStackViewModel::SearchTick()
 		// we need to call the SearchCompletedDelegate to go through SynchronizeTreeExpansion in SNiagaraStack so that when exiting search we return the stack expansion to the state it was before searching
 		else
 		{
-			SearchCompletedDelegate.Broadcast();
+			SearchCompletedDelegate.Broadcast(CurrentSearchConfig.Get(FNiagaraStackSearchConfig()));
 		}
 		bRestartSearch = false;
 	}
@@ -611,7 +605,8 @@ void UNiagaraStackViewModel::SearchTick()
 				}
 			}
 			
-			SearchCompletedDelegate.Broadcast();
+			SearchCompletedDelegate.Broadcast(CurrentSearchConfig.Get(FNiagaraStackSearchConfig()));
+			CurrentSearchConfig.Reset();
 		}
 	}
 }
@@ -655,12 +650,13 @@ void UNiagaraStackViewModel::GeneratePathForEntry(UNiagaraStackEntry* Root, UNia
 	}
 }
 
-void UNiagaraStackViewModel::RestartSearch()
+void UNiagaraStackViewModel::RestartSearch(const FNiagaraStackSearchConfig& InConfig)
 {
 	if(bRestartSearch == false)
 	{
+		CurrentSearchConfig = InConfig;
+		
 		bRestartSearch = true;
-		FocusedStackEntryOverride.Reset();
 	}
 }
 
@@ -704,6 +700,11 @@ UNiagaraStackViewModel::FOnDataObjectChanged& UNiagaraStackViewModel::OnDataObje
 	return DataObjectChangedDelegate;
 }
 
+UNiagaraStackViewModel::FOnStackEntrySelectionRequested& UNiagaraStackViewModel::OnStackEntrySelectionRequested()
+{
+	return OnStackEntrySelectionRequestedDelegate;
+}
+
 bool UNiagaraStackViewModel::GetShowOnlyModified() const
 {
 	for (TSharedRef<FTopLevelViewModel> TopLevelViewModel : TopLevelViewModels)
@@ -726,7 +727,9 @@ void UNiagaraStackViewModel::SetShowOnlyModified(bool bInShowOnlyModified)
 		}
 	}
 
-	RestartSearch();
+	FNiagaraStackSearchConfig SearchConfig;
+	SearchConfig.bScrollToResult = false;
+	RestartSearch(SearchConfig);
 	RootEntry->RefreshChildren();
 }
 
@@ -752,7 +755,9 @@ void UNiagaraStackViewModel::SetShowAllAdvanced(bool bInShowAllAdvanced)
 		}
 	}
 
-	RestartSearch();
+	FNiagaraStackSearchConfig SearchConfig;
+	SearchConfig.bScrollToResult = false;
+	RestartSearch(SearchConfig);
 	RootEntry->RefreshFilteredChildren();
 }
 
@@ -779,7 +784,9 @@ void UNiagaraStackViewModel::SetShowOutputs(bool bInShowOutputs)
 	}
 	
 	// Showing outputs changes indenting so a full refresh is needed.
-	RestartSearch();
+	FNiagaraStackSearchConfig SearchConfig;
+	SearchConfig.bScrollToResult = false;
+	RestartSearch(SearchConfig);
 	RootEntry->RefreshChildren();
 }
 
@@ -806,7 +813,9 @@ void UNiagaraStackViewModel::SetShowLinkedInputs(bool bInShowLinkedInputs)
 	}
 
 	// Showing linked inputs changes indenting so a full refresh is needed.
-	RestartSearch();
+	FNiagaraStackSearchConfig SearchConfig;
+	SearchConfig.bScrollToResult = false;
+	RestartSearch(SearchConfig);
 	RootEntry->RefreshChildren();
 }
 
@@ -832,7 +841,9 @@ void UNiagaraStackViewModel::SetShowOnlyIssues(bool bInShowOnlyIssues)
 		}
 	}
 
-	RestartSearch();
+	FNiagaraStackSearchConfig SearchConfig;
+	SearchConfig.bScrollToResult = false;
+	RestartSearch(SearchConfig);
 	RootEntry->RefreshChildren();
 }
 
@@ -855,16 +866,9 @@ void UNiagaraStackViewModel::SetLastScrollPosition(double InLastScrollPosition)
 	}
 }
 
-void UNiagaraStackViewModel::SetFocusedStackEntryOverride(UNiagaraStackEntry* InStackEntry)
+void UNiagaraStackViewModel::SelectStackEntry(UNiagaraStackEntry* StackEntryToSelect, FNiagaraStackEntrySelectionConfig Config) const
 {
-	if(InStackEntry)
-	{
-		FocusedStackEntryOverride = InStackEntry;
-	}
-	else
-	{
-		FocusedStackEntryOverride.Reset();
-	}
+	OnStackEntrySelectionRequestedDelegate.ExecuteIfBound(StackEntryToSelect, Config);
 }
 
 void UNiagaraStackViewModel::EntryExpansionChanged()
@@ -889,7 +893,9 @@ void UNiagaraStackViewModel::EntryStructureChanged(ENiagaraStructureChangedFlags
 	// if the structure changed, additionally invalidate search results
 	if((Flags & StructureChanged) != 0)
 	{
-		RestartSearch();
+		FNiagaraStackSearchConfig SearchConfig;
+		SearchConfig.bScrollToResult = false;
+		RestartSearch(SearchConfig);
 	}
 }
 
@@ -899,7 +905,7 @@ void UNiagaraStackViewModel::EntryDataObjectModified(TArray<UObject*> ChangedObj
 	{
 		SystemViewModel.Pin()->NotifyDataObjectChanged(ChangedObjects, ChangeType);
 	}
-	RestartSearch();
+	
 	DataObjectChangedDelegate.Broadcast(ChangedObjects, ChangeType);
 }
 
diff --git a/Engine/Plugins/FX/Niagara/Source/NiagaraEditor/Public/ViewModels/Stack/NiagaraStackEntry.h b/Engine/Plugins/FX/Niagara/Source/NiagaraEditor/Public/ViewModels/Stack/NiagaraStackEntry.h
--- a/Engine/Plugins/FX/Niagara/Source/NiagaraEditor/Public/ViewModels/Stack/NiagaraStackEntry.h
+++ b/Engine/Plugins/FX/Niagara/Source/NiagaraEditor/Public/ViewModels/Stack/NiagaraStackEntry.h
@@ -538,7 +538,7 @@ public:
 
 	virtual bool TestCanCopyWithMessage(FText& OutMessage) const { return false; }
 
-	virtual void Copy(UNiagaraClipboardContent* ClipboardContent) const { OnCopyPasteDelegate.ExecuteIfBound(); }
+	virtual void Copy(UNiagaraClipboardContent* ClipboardContent) const { RequestPulse(); }
 
 	virtual bool SupportsPaste() const { return false; }
 
@@ -546,9 +546,10 @@ public:
 
 	virtual FText GetPasteTransactionText(const UNiagaraClipboardContent* ClipboardContent) const { return FText(); }
 
-	virtual void Paste(const UNiagaraClipboardContent* ClipboardContent, FText& OutPasteWarning) { OnCopyPasteDelegate.ExecuteIfBound(); }
+	virtual void Paste(const UNiagaraClipboardContent* ClipboardContent, FText& OutPasteWarning) { RequestPulse(); }
 
-	FSimpleDelegate& OnCopyPaste() { return OnCopyPasteDelegate; }
+	void RequestPulse() const { OnPulseRequestedDelegate.ExecuteIfBound(); }
+	FSimpleDelegate& OnPulseRequested() { return OnPulseRequestedDelegate; }
 	
 	virtual bool SupportsDelete() const { return false; }
 
@@ -642,7 +643,7 @@ protected:
 
 	bool IsSystemViewModelValid() const { return SystemViewModel.IsValid(); }
 
-	FSimpleDelegate OnCopyPasteDelegate;
+	FSimpleDelegate OnPulseRequestedDelegate;
 private:
 	NIAGARAEDITOR_API void ChildStructureChanged(ENiagaraStructureChangedFlags Info);
 
diff --git a/Engine/Plugins/FX/Niagara/Source/NiagaraEditor/Public/ViewModels/Stack/NiagaraStackViewModel.h b/Engine/Plugins/FX/Niagara/Source/NiagaraEditor/Public/ViewModels/Stack/NiagaraStackViewModel.h
--- a/Engine/Plugins/FX/Niagara/Source/NiagaraEditor/Public/ViewModels/Stack/NiagaraStackViewModel.h
+++ b/Engine/Plugins/FX/Niagara/Source/NiagaraEditor/Public/ViewModels/Stack/NiagaraStackViewModel.h
@@ -35,6 +35,18 @@ private:
 	bool bIncludeEmitterInformation;
 };
 
+struct FNiagaraStackSearchConfig
+{
+	bool bScrollToResult = true;
+};
+
+struct FNiagaraStackEntrySelectionConfig
+{
+	bool bScrollToAndExpand = true;
+	bool bSelectRow = false;
+	bool bPulse = false;
+};
+
 UCLASS(MinimalAPI)
 class UNiagaraStackViewModel : public UObject
 {
@@ -42,9 +54,11 @@ class UNiagaraStackViewModel : public UObject
 
 public:
 	DECLARE_DELEGATE_OneParam(FOnChangeSearchTextExternal, FText)
+	
+	DECLARE_DELEGATE_TwoParams(FOnStackEntrySelectionRequested, UNiagaraStackEntry* StackEntry, FNiagaraStackEntrySelectionConfig SelectionConfig);
 	DECLARE_MULTICAST_DELEGATE(FOnExpansionChanged);
 	DECLARE_MULTICAST_DELEGATE_OneParam(FOnStructureChanged, ENiagaraStructureChangedFlags);
-	DECLARE_MULTICAST_DELEGATE(FOnSearchCompleted);
+	DECLARE_MULTICAST_DELEGATE_OneParam(FOnSearchCompleted, const FNiagaraStackSearchConfig& SearchConfig);
 	DECLARE_MULTICAST_DELEGATE_TwoParams(FOnDataObjectChanged, TArray<UObject*> /** Changed objects */, ENiagaraDataObjectChange /** Change type */);
 public:
 	struct FSearchResult
@@ -99,6 +113,7 @@ public:
 	NIAGARAEDITOR_API FOnStructureChanged& OnStructureChanged();
 	NIAGARAEDITOR_API FOnSearchCompleted& OnSearchCompleted();
 	NIAGARAEDITOR_API FOnDataObjectChanged& OnDataObjectChanged();
+	NIAGARAEDITOR_API FOnStackEntrySelectionRequested& OnStackEntrySelectionRequested();
 
 	NIAGARAEDITOR_API bool GetShowOnlyModified() const;
 	NIAGARAEDITOR_API void SetShowOnlyModified(bool bInShowOnlyModified);
@@ -118,8 +133,9 @@ public:
 	NIAGARAEDITOR_API double GetLastScrollPosition() const;
 	NIAGARAEDITOR_API void SetLastScrollPosition(double InLastScrollPosition);
 
-	NIAGARAEDITOR_API void SetFocusedStackEntryOverride(UNiagaraStackEntry* InStackEntry);
-	
+	/** Select a stack entry. This selection is distinct from the Selection from the FNiagaraSystemSelectionViewModel which changes the root stack selection. */
+	NIAGARAEDITOR_API void SelectStackEntry(UNiagaraStackEntry* StackEntryToSelect, FNiagaraStackEntrySelectionConfig Config) const;
+
 	NIAGARAEDITOR_API virtual void Tick();
 	//~ stack search stuff
 	NIAGARAEDITOR_API void ResetSearchText();
@@ -185,8 +201,8 @@ private:
 		TArray<FSearchWorkItem>& TraversedArray);
 	NIAGARAEDITOR_API bool ItemMatchesSearchCriteria(UNiagaraStackEntry::FStackSearchItem SearchItem);
 	NIAGARAEDITOR_API void GeneratePathForEntry(UNiagaraStackEntry* Root, UNiagaraStackEntry* Entry, TArray<UNiagaraStackEntry*> CurrentPath, TArray<UNiagaraStackEntry*>& EntryPath) const;
-
-	NIAGARAEDITOR_API void RestartSearch();
+	
+	NIAGARAEDITOR_API void RestartSearch(const FNiagaraStackSearchConfig& InConfig);
 	NIAGARAEDITOR_API void UpdateStackWithValidationResults();
 
 private:
@@ -207,15 +223,15 @@ private:
 	FOnExpansionChanged ExpansionInOverviewChangedDelegate;
 	FOnStructureChanged StructureChangedDelegate;
 	FOnDataObjectChanged DataObjectChangedDelegate;
-
+	FOnStackEntrySelectionRequested OnStackEntrySelectionRequestedDelegate;
 
 	// ~Search stuff
+	TOptional<FNiagaraStackSearchConfig> CurrentSearchConfig;
 	FText CurrentSearchText;
 	FOnSearchCompleted SearchCompletedDelegate;
 	TArray<FSearchWorkItem> ItemsToSearch;
 	TArray<FSearchResult> CurrentSearchResults;
 	static NIAGARAEDITOR_API const double MaxSearchTime;
-	TWeakObjectPtr<UNiagaraStackEntry> FocusedStackEntryOverride;
 	TOptional<FSearchResult> FocusedSearchResultCache;
 	bool bRestartSearch;
 	bool bRefreshPending;
diff --git a/Engine/Plugins/FX/Niagara/Source/NiagaraEditorWidgets/Private/SNiagaraOverviewInlineParameterBox.cpp b/Engine/Plugins/FX/Niagara/Source/NiagaraEditorWidgets/Private/SNiagaraOverviewInlineParameterBox.cpp
--- a/Engine/Plugins/FX/Niagara/Source/NiagaraEditorWidgets/Private/SNiagaraOverviewInlineParameterBox.cpp
+++ b/Engine/Plugins/FX/Niagara/Source/NiagaraEditorWidgets/Private/SNiagaraOverviewInlineParameterBox.cpp
@@ -74,9 +74,13 @@ FReply SNiagaraOverviewInlineParameterBox::NavigateToStack(TWeakObjectPtr<UNiaga
 	{
 		FunctionInput = SubstituteEntry;
 	}
-	
-	ModuleItem->GetSystemViewModel()->GetSelectionViewModel()->GetSelectionStackViewModel()->SetFocusedStackEntryOverride(FunctionInput.Get());
 
+	FNiagaraStackEntrySelectionConfig Config;
+	Config.bScrollToAndExpand = true;
+	Config.bSelectRow = true;
+	Config.bPulse = true;
+	ModuleItem->GetSystemViewModel()->GetSelectionViewModel()->GetSelectionStackViewModel()->SelectStackEntry(FunctionInput.Get(), Config);
+	
 	return FReply::Handled();
 }
 
diff --git a/Engine/Plugins/FX/Niagara/Source/NiagaraEditorWidgets/Private/SNiagaraStack.cpp b/Engine/Plugins/FX/Niagara/Source/NiagaraEditorWidgets/Private/SNiagaraStack.cpp
--- a/Engine/Plugins/FX/Niagara/Source/NiagaraEditorWidgets/Private/SNiagaraStack.cpp
+++ b/Engine/Plugins/FX/Niagara/Source/NiagaraEditorWidgets/Private/SNiagaraStack.cpp
@@ -300,9 +300,10 @@ void SNiagaraStack::Construct(const FArguments& InArgs, UNiagaraStackViewModel*
 {
 	StackViewModel = InStackViewModel;
 	StackViewModel->OnChangeSearchTextExternal().BindSP(this, &SNiagaraStack::UpdateSearchTextFromExternal);
+	StackViewModel->OnStackEntrySelectionRequested().BindSP(this, &SNiagaraStack::SelectStackEntry);
 	StackViewModel->OnStructureChanged().AddSP(this, &SNiagaraStack::StackStructureChanged);
 	StackViewModel->OnExpansionChanged().AddSP(this, &SNiagaraStack::OnStackExpansionChanged);
-	StackViewModel->OnSearchCompleted().AddSP(this, &SNiagaraStack::OnStackSearchComplete); 
+	StackViewModel->OnSearchCompleted().AddSP(this, &SNiagaraStack::OnStackSearchComplete);
 	NameColumnWidth = .3f;
 	ContentColumnWidth = .7f;
 	StackCommandContext = MakeShared<FNiagaraStackCommandContext>();
@@ -518,6 +519,36 @@ bool SNiagaraStack::IsEntryFocusedInSearch(UNiagaraStackEntry* Entry) const
 	return false;
 }
 
+void SNiagaraStack::SelectStackEntry(UNiagaraStackEntry* StackEntry, FNiagaraStackEntrySelectionConfig Config)
+{
+	if(StackViewModel && StackEntry->IsFinalized() == false)
+	{
+		if(Config.bScrollToAndExpand)
+		{
+			TArray<UNiagaraStackEntry*> StackEntryPath;
+			StackViewModel->GetPathForEntry(StackEntry, StackEntryPath);
+			ExpandAllInPath(StackEntryPath);
+			
+			StackTree->RequestScrollIntoView(StackEntry);
+		}
+
+		if(Config.bSelectRow)
+		{
+			StackTree->SetSelection(StackEntry, ESelectInfo::Type::Direct);
+		}
+
+		if(Config.bPulse)
+		{
+			
+			RegisterActiveTimer(0.f, FWidgetActiveTimerDelegate::CreateLambda([StackEntry](double CurrentTime, float DeltaTime)
+			{
+				StackEntry->RequestPulse();
+				return EActiveTimerReturnType::Stop;
+			}));
+		}
+	}
+}
+
 void SNiagaraStack::ShowEmitterInContentBrowser(TWeakPtr<FNiagaraEmitterHandleViewModel> EmitterHandleViewModelWeak)
 {
 	TSharedPtr<FNiagaraEmitterHandleViewModel> EmitterHandleViewModel = EmitterHandleViewModelWeak.Pin();
@@ -657,10 +688,14 @@ FReply SNiagaraStack::OnRowAcceptDrop(const FDragDropEvent& InDragDropEvent, EIt
 	return bHandled ? FReply::Handled() : FReply::Unhandled();
 }
 
-void SNiagaraStack::OnStackSearchComplete()
+void SNiagaraStack::OnStackSearchComplete(const FNiagaraStackSearchConfig& SearchConfig)
 {
 	ExpandSearchResults();
-	ScrollFocusedEntryIntoView();
+
+	if(SearchConfig.bScrollToResult)
+	{
+		ScrollFocusedEntryIntoView();
+	}
 }
 
 void SNiagaraStack::ExpandSearchResults()
@@ -677,7 +712,15 @@ void SNiagaraStack::OnSearchBoxTextCommitted(const FText& NewText, ETextCommit::
 {
 	if (CommitInfo == ETextCommit::OnEnter) // hasn't been autojumped yet or we hit enter
 	{
-		AddSearchScrollOffset(+1);
+		bool bIsShiftDown = FSlateApplication::Get().GetModifierKeys().IsShiftDown();
+		if(bIsShiftDown)
+		{
+			AddSearchScrollOffset(-1);
+		}
+		else
+		{
+			AddSearchScrollOffset(+1);
+		}
 	}
 }
 
diff --git a/Engine/Plugins/FX/Niagara/Source/NiagaraEditorWidgets/Private/SNiagaraStack.h b/Engine/Plugins/FX/Niagara/Source/NiagaraEditorWidgets/Private/SNiagaraStack.h
--- a/Engine/Plugins/FX/Niagara/Source/NiagaraEditorWidgets/Private/SNiagaraStack.h
+++ b/Engine/Plugins/FX/Niagara/Source/NiagaraEditorWidgets/Private/SNiagaraStack.h
@@ -97,10 +97,12 @@ private:
 	FSlateColor GetTextColorForItem(UNiagaraStackEntry* Item) const;
 	void ScrollFocusedEntryIntoView();
 	void AddSearchScrollOffset(int NumberOfSteps);
-	void OnStackSearchComplete();
+	void OnStackSearchComplete(const FNiagaraStackSearchConfig& SearchConfig);
 	void ExpandSearchResults();
 	bool IsEntryFocusedInSearch(UNiagaraStackEntry* Entry) const;
-	
+
+	void SelectStackEntry(UNiagaraStackEntry* StackEntry, FNiagaraStackEntrySelectionConfig Config);
+
 	// Inline menu commands
 	void ShowEmitterInContentBrowser(TWeakPtr<FNiagaraEmitterHandleViewModel> EmitterHandleViewModelWeak);
 	void NavigateTo(UNiagaraStackEntry* Item);
diff --git a/Engine/Plugins/FX/Niagara/Source/NiagaraEditorWidgets/Private/Stack/SNiagaraStackTableRow.cpp b/Engine/Plugins/FX/Niagara/Source/NiagaraEditorWidgets/Private/Stack/SNiagaraStackTableRow.cpp
--- a/Engine/Plugins/FX/Niagara/Source/NiagaraEditorWidgets/Private/Stack/SNiagaraStackTableRow.cpp
+++ b/Engine/Plugins/FX/Niagara/Source/NiagaraEditorWidgets/Private/Stack/SNiagaraStackTableRow.cpp
@@ -43,7 +43,7 @@ SNiagaraStackTableRow::~SNiagaraStackTableRow()
 {
 	if(StackEntry && StackEntry->IsFinalized() == false)
 	{
-		StackEntry->OnCopyPaste().Unbind();
+		StackEntry->OnPulseRequested().Unbind();
 	}
 }
 
@@ -69,7 +69,7 @@ void SNiagaraStackTableRow::Construct(const FArguments& InArgs, UNiagaraStackVie
 
 	PulseAnimation.AddCurve(0.f, 0.5f, ECurveEaseFunction::CubicInOut);
 
-	InStackEntry->OnCopyPaste().BindSP(this, &SNiagaraStackTableRow::Pulse);
+	InStackEntry->OnPulseRequested().BindSP(this, &SNiagaraStackTableRow::Pulse);
 	
 	IndicatorColor = InArgs._IndicatorColor;
 
