[NiagaraComponent]
- Avoid registering to delegate from the loader during serialize since Postload will do it later anyway

diff --git a/Engine/Plugins/FX/Niagara/Source/Niagara/Private/NiagaraComponent.cpp b/Engine/Plugins/FX/Niagara/Source/Niagara/Private/NiagaraComponent.cpp
--- a/Engine/Plugins/FX/Niagara/Source/Niagara/Private/NiagaraComponent.cpp
+++ b/Engine/Plugins/FX/Niagara/Source/Niagara/Private/NiagaraComponent.cpp
@@ -2993,10 +2993,14 @@ void UNiagaraComponent::Serialize(FStructuredArchive::FRecord Record)
 
 		// Check the the event handler here as well as in post load since post load isn't always called due to how components are managed when spawning
 		// actors with templates.
-		if (Asset != nullptr && AssetExposedParametersChangedHandle.IsValid() == false)
+		// When called from the loader, we know postload is going to be called, so avoid race conditions on the delegate by avoiding registering to it while in the loader.
+		if (!IsInAsyncLoadingThread() && !IsInParallelLoadingThread())
 		{
-			AssetExposedParametersChangedHandle = Asset->GetExposedParameters().AddOnChangedHandler(
-				FNiagaraParameterStore::FOnChanged::FDelegate::CreateUObject(this, &UNiagaraComponent::AssetExposedParametersChanged));
+			if (Asset != nullptr && AssetExposedParametersChangedHandle.IsValid() == false)
+			{
+				AssetExposedParametersChangedHandle = Asset->GetExposedParameters().AddOnChangedHandler(
+					FNiagaraParameterStore::FOnChanged::FDelegate::CreateUObject(this, &UNiagaraComponent::AssetExposedParametersChanged));
+			}
 		}
 	}
 #endif
