- Fix for editor undo crashing due to order being out of sync and the thumbnails

diff --git a/Engine/Plugins/FX/Niagara/Source/Niagara/Classes/NiagaraEmitterInstance.h b/Engine/Plugins/FX/Niagara/Source/Niagara/Classes/NiagaraEmitterInstance.h
--- a/Engine/Plugins/FX/Niagara/Source/Niagara/Classes/NiagaraEmitterInstance.h
+++ b/Engine/Plugins/FX/Niagara/Source/Niagara/Classes/NiagaraEmitterInstance.h
@@ -138,7 +138,11 @@ public:
 		}
 	}
 
-	[[nodiscard]] const UNiagaraRendererProperties* GetRenderer(int32 i) const { return GetRenderers()[i]; }
+	[[nodiscard]] const UNiagaraRendererProperties* GetRenderer(int32 i) const
+	{
+		TConstArrayView<UNiagaraRendererProperties*> LocalRenderers = GetRenderers();
+		return LocalRenderers.IsValidIndex(i) ? LocalRenderers[i] : nullptr;
+	}
 
 	////////////////////////////////////////////////////////////////////////////////////////////////////
 	//-TODO:Stateless: Should we cache this data? clean up const casting
diff --git a/Engine/Plugins/FX/Niagara/Source/Niagara/Private/NiagaraSystemRenderData.cpp b/Engine/Plugins/FX/Niagara/Source/Niagara/Private/NiagaraSystemRenderData.cpp
--- a/Engine/Plugins/FX/Niagara/Source/Niagara/Private/NiagaraSystemRenderData.cpp
+++ b/Engine/Plugins/FX/Niagara/Source/Niagara/Private/NiagaraSystemRenderData.cpp
@@ -418,8 +418,10 @@ void FNiagaraSystemRenderData::OnSystemComplete(const FNiagaraSystemInstance& Sy
 			const FNiagaraEmitterInstance& EmitterInstance = EmitterInstances[ExecIdx.EmitterIndex].Get();
 			if (FNiagaraRenderer* EmitterRenderer = EmitterRenderers_GT[ExecIdx.SystemRendererIndex])
 			{
-				const UNiagaraRendererProperties* RendererProperties = EmitterInstance.GetRenderer(ExecIdx.EmitterRendererIndex);
-				EmitterRenderer->OnSystemComplete_GameThread(RendererProperties, &EmitterInstance);
+				if (const UNiagaraRendererProperties* RendererProperties = EmitterInstance.GetRenderer(ExecIdx.EmitterRendererIndex))
+				{
+					EmitterRenderer->OnSystemComplete_GameThread(RendererProperties, &EmitterInstance);
+				}
 			}
 		}
 	}
diff --git a/Engine/Plugins/FX/Niagara/Source/NiagaraEditorWidgets/Private/SNiagaraEmitterRendererThumbnails.cpp b/Engine/Plugins/FX/Niagara/Source/NiagaraEditorWidgets/Private/SNiagaraEmitterRendererThumbnails.cpp
--- a/Engine/Plugins/FX/Niagara/Source/NiagaraEditorWidgets/Private/SNiagaraEmitterRendererThumbnails.cpp
+++ b/Engine/Plugins/FX/Niagara/Source/NiagaraEditorWidgets/Private/SNiagaraEmitterRendererThumbnails.cpp
@@ -69,15 +69,19 @@ void FNiagaraRendererStackEntryObserver::OnMaterialCompiled(class UMaterialInter
 		FNiagaraEmitterInstance* InInstance = EmitterHandleViewModelWeak.Pin()->GetEmitterViewModel()->GetSimulation().IsValid() ? EmitterHandleViewModelWeak.Pin()->GetEmitterViewModel()->GetSimulation().Pin().Get() : nullptr;
 		for (UNiagaraStackEntry* Entry : RendererPreviewStackEntries)
 		{
-			if (UNiagaraStackRendererItem* RendererItem = Cast<UNiagaraStackRendererItem>(Entry))
+			UNiagaraStackRendererItem* RendererItem = Cast<UNiagaraStackRendererItem>(Entry);
+			UNiagaraRendererProperties* RendererProperties = RendererItem ? RendererItem->GetRendererProperties() : nullptr;
+			if (RendererProperties == nullptr)
 			{
-				TArray<UMaterialInterface*> Materials;
-				RendererItem->GetRendererProperties()->GetUsedMaterials(InInstance, Materials);
-				if (Materials.Contains(MaterialInterface))
-				{
-					bUsingThisMaterial = true;
-					break;
-				}
+				continue;
+			}
+
+			TArray<UMaterialInterface*> Materials;
+			RendererProperties->GetUsedMaterials(InInstance, Materials);
+			if (Materials.Contains(MaterialInterface))
+			{
+				bUsingThisMaterial = true;
+				break;
 			}
 		}
 
@@ -129,66 +133,70 @@ void SNiagaraEmitterRendererThumbnails::RefreshThumbnails()
 	for (int32 StackEntryIndex = 0; StackEntryIndex < RendererPreviewStackEntries.Num(); StackEntryIndex++)
 	{
 		UNiagaraStackEntry* Entry = RendererPreviewStackEntries[StackEntryIndex];
-		if (UNiagaraStackRendererItem* RendererItem = Cast<UNiagaraStackRendererItem>(Entry))
+		UNiagaraStackRendererItem* RendererItem = Cast<UNiagaraStackRendererItem>(Entry);
+		UNiagaraRendererProperties* RendererProperties = RendererItem ? RendererItem->GetRendererProperties() : nullptr;
+		if (RendererProperties == nullptr)
 		{
-			TArray<TSharedPtr<SWidget>> Widgets;
-			RendererItem->GetRendererProperties()->GetRendererWidgets(InInstance, Widgets, UThumbnailManager::Get().GetSharedThumbnailPool());
-			TArray<TSharedPtr<SWidget>> TooltipWidgets;
-			RendererItem->GetRendererProperties()->GetRendererTooltipWidgets(InInstance, TooltipWidgets, UThumbnailManager::Get().GetSharedThumbnailPool());
+			continue;
+		}
 
-			// Fallback widget & tooltip
-			if(Widgets.Num() == 0)
-			{
-				RendererItem->GetRendererProperties()->CreateDefaultRendererWidget(Widgets);
-			}
+		TArray<TSharedPtr<SWidget>> Widgets;
+		TArray<TSharedPtr<SWidget>> TooltipWidgets;
+		RendererProperties->GetRendererWidgets(InInstance, Widgets, UThumbnailManager::Get().GetSharedThumbnailPool());
+		RendererProperties->GetRendererTooltipWidgets(InInstance, TooltipWidgets, UThumbnailManager::Get().GetSharedThumbnailPool());
 
-			if(TooltipWidgets.Num() == 0)
-			{
-				FText BaseText = LOCTEXT("DefaultRendererTooltip", "{0}");
-				TooltipWidgets.Add(SNew(STextBlock).Text(FText::FormatOrdered(BaseText, RendererItem->GetRendererProperties()->GetClass()->GetDisplayNameText())));
-			}
+		// Fallback widget & tooltip
+		if(Widgets.Num() == 0)
+		{
+			RendererProperties->CreateDefaultRendererWidget(Widgets);
+		}
+
+		if(TooltipWidgets.Num() == 0)
+		{
+			FText BaseText = LOCTEXT("DefaultRendererTooltip", "{0}");
+			TooltipWidgets.Add(SNew(STextBlock).Text(FText::FormatOrdered(BaseText, RendererProperties->GetClass()->GetDisplayNameText())));
+		}
 			
-			check(Widgets.Num() == TooltipWidgets.Num());
-			for (int32 WidgetIndex = 0; WidgetIndex < Widgets.Num(); WidgetIndex++)
-			{
-				SetupRendererThumbnailTooltip(Widgets[WidgetIndex], TooltipWidgets[WidgetIndex]);
+		check(Widgets.Num() == TooltipWidgets.Num());
+		for (int32 WidgetIndex = 0; WidgetIndex < Widgets.Num(); WidgetIndex++)
+		{
+			SetupRendererThumbnailTooltip(Widgets[WidgetIndex], TooltipWidgets[WidgetIndex]);
 				
-				// We wrap the thumbnail in a scalebox so it will scale to whatever its parent dictates
-				TSharedPtr<SWidget> RendererThumbnailWidget = SNew(SBox)
-					.VAlign(VAlign_Center)
-					.HAlign(HAlign_Center)
-					.HeightOverride(RendererThumbnailSize)
-					.WidthOverride(RendererThumbnailSize)
+			// We wrap the thumbnail in a scalebox so it will scale to whatever its parent dictates
+			TSharedPtr<SWidget> RendererThumbnailWidget = SNew(SBox)
+				.VAlign(VAlign_Center)
+				.HAlign(HAlign_Center)
+				.HeightOverride(RendererThumbnailSize)
+				.WidthOverride(RendererThumbnailSize)
+				[
+					SNew(SScaleBox)
+					.Stretch(EStretch::ScaleToFit)
 					[
-						SNew(SScaleBox)
-						.Stretch(EStretch::ScaleToFit)
-						[
-							Widgets[WidgetIndex].ToSharedRef()	
-						]
-					];
+						Widgets[WidgetIndex].ToSharedRef()	
+					]
+				];
 				
-				// If we have on click functionality, we wrap the thumbnail in a button
-				if(OnRendererThumbnailClickedDelegate.IsBound())
-				{
-					RendererThumbnailWidget = SNew(SButton)
-						.ButtonStyle(&FAppStyle::GetWidgetStyle<FButtonStyle>("HoverHintOnly"))
-						.ContentPadding(FMargin(1.f))
-						.OnClicked(this, &SNiagaraEmitterRendererThumbnails::OnRendererThumbnailClicked, RendererItem)
-						[
-							RendererThumbnailWidget.ToSharedRef()
-						];
-				}
-				
-				ToolBarBuilder.AddWidget(
-					RendererThumbnailWidget.ToSharedRef()
-				);
-			}
-	
-			// if we had a widget for this entry, add a separator for the next entry's widgets, except for the last entry
-			if(Widgets.Num() > 0 && StackEntryIndex < RendererPreviewStackEntries.Num() - 1)
+			// If we have on click functionality, we wrap the thumbnail in a button
+			if(OnRendererThumbnailClickedDelegate.IsBound())
 			{
-				ToolBarBuilder.AddSeparator();
+				RendererThumbnailWidget = SNew(SButton)
+					.ButtonStyle(&FAppStyle::GetWidgetStyle<FButtonStyle>("HoverHintOnly"))
+					.ContentPadding(FMargin(1.f))
+					.OnClicked(this, &SNiagaraEmitterRendererThumbnails::OnRendererThumbnailClicked, RendererItem)
+					[
+						RendererThumbnailWidget.ToSharedRef()
+					];
 			}
+				
+			ToolBarBuilder.AddWidget(
+				RendererThumbnailWidget.ToSharedRef()
+			);
+		}
+	
+		// if we had a widget for this entry, add a separator for the next entry's widgets, except for the last entry
+		if(Widgets.Num() > 0 && StackEntryIndex < RendererPreviewStackEntries.Num() - 1)
+		{
+			ToolBarBuilder.AddSeparator();
 		}
 	}
 	
