[Niagara]
- Convert to the new ConditionalPreload function that is thread-safe

diff --git a/Engine/Plugins/FX/Niagara/Source/Niagara/Private/NiagaraEmitter.cpp b/Engine/Plugins/FX/Niagara/Source/Niagara/Private/NiagaraEmitter.cpp
--- a/Engine/Plugins/FX/Niagara/Source/Niagara/Private/NiagaraEmitter.cpp
+++ b/Engine/Plugins/FX/Niagara/Source/Niagara/Private/NiagaraEmitter.cpp
@@ -549,13 +549,7 @@ void UNiagaraEmitter::PostLoad()
 	// state.  While investigation continues into hwo this happens, this call to Preload() seems to resolve the somewhat reproducible case
 	// that we have.  There's some evidence that UNiagaraScript having instanced properties is leading to an edge case in object serialization
 	// but more investigation is required.
-	if (HasAnyFlags(RF_NeedLoad))
-	{
-		if (FLinkerLoad* LinkerLoad = GetLinker())
-		{
-			LinkerLoad->Preload(this);
-		}
-	}
+	ConditionalPreload();
 #endif
 
 #if WITH_EDITORONLY_DATA
diff --git a/Engine/Plugins/FX/Niagara/Source/Niagara/Private/NiagaraSystem.cpp b/Engine/Plugins/FX/Niagara/Source/Niagara/Private/NiagaraSystem.cpp
--- a/Engine/Plugins/FX/Niagara/Source/Niagara/Private/NiagaraSystem.cpp
+++ b/Engine/Plugins/FX/Niagara/Source/Niagara/Private/NiagaraSystem.cpp
@@ -1224,10 +1224,11 @@ void UNiagaraSystem::PostLoad()
 			if (Dependency->HasAnyFlags(RF_NeedLoad))
 			{
 				UE_LOG(LogNiagara, Log, TEXT("NiagaraSystem::PostLoad() - SubObject(%s) RF_NeedLoad"), *GetFullNameSafe(Dependency));
-				if (GNiagaraObjectNeedsLoadMode == 2)
-				{
-					Dependency->GetLinker()->Preload(Dependency);
-				}
+			}
+
+			if (GNiagaraObjectNeedsLoadMode == 2)
+			{
+				Dependency->ConditionalPreload();
 			}
 		}
 	}
