- Fix for update context killing systems that are waiting compilation
- When switching preview modes the update context is called multiple times
- The first time it's considered active
- The next time it's no longer active but is waiting for compilation (compontent ticking) so is now removed from the reactivate list
- We now check if the component tick is enabled due to waiting for compilation results

#tests Place a looping system in the level, change device profiles, observe it isn't visible after changing and can not be reactivated easily

diff --git a/Engine/Plugins/FX/Niagara/Source/Niagara/Private/NiagaraCommon.cpp b/Engine/Plugins/FX/Niagara/Source/Niagara/Private/NiagaraCommon.cpp
--- a/Engine/Plugins/FX/Niagara/Source/Niagara/Private/NiagaraCommon.cpp
+++ b/Engine/Plugins/FX/Niagara/Source/Niagara/Private/NiagaraCommon.cpp
@@ -298,7 +298,7 @@ void FNiagaraSystemUpdateContext::AddInternal(UNiagaraComponent* Comp, bool bReI
 		}
 	}
 
-	bool bIsActive = (Comp->IsActive() && Comp->GetRequestedExecutionState() == ENiagaraExecutionState::Active) || Comp->IsRegisteredWithScalabilityManager();
+	const bool bIsActive = Comp->IsActiveForUpdateContext();
 
 	if (bDestroyOnAdd)
 	{
diff --git a/Engine/Plugins/FX/Niagara/Source/Niagara/Private/NiagaraComponent.cpp b/Engine/Plugins/FX/Niagara/Source/Niagara/Private/NiagaraComponent.cpp
--- a/Engine/Plugins/FX/Niagara/Source/Niagara/Private/NiagaraComponent.cpp
+++ b/Engine/Plugins/FX/Niagara/Source/Niagara/Private/NiagaraComponent.cpp
@@ -1920,6 +1920,17 @@ void UNiagaraComponent::OnPooledReuse(UWorld* NewWorld)
 	}
 }
 
+bool UNiagaraComponent::IsActiveForUpdateContext() const
+{
+	return
+		// If we are activate and the execution state is also active (anything else means we were winding down / disabled)
+		(IsActive() && GetRequestedExecutionState() == ENiagaraExecutionState::Active) ||
+		// We were waiting on compilation / DDC request to come back
+		(bAwaitingActivationDueToNotReady && PrimaryComponentTick.IsTickFunctionEnabled()) || 
+		// We are registered with scalabilty, in which case allow it to take control
+		IsRegisteredWithScalabilityManager();
+}
+
 void UNiagaraComponent::OnRegister()
 {
 	if (IsActive() && SystemInstanceController.IsValid() == false)
diff --git a/Engine/Plugins/FX/Niagara/Source/Niagara/Public/NiagaraComponent.h b/Engine/Plugins/FX/Niagara/Source/Niagara/Public/NiagaraComponent.h
--- a/Engine/Plugins/FX/Niagara/Source/Niagara/Public/NiagaraComponent.h
+++ b/Engine/Plugins/FX/Niagara/Source/Niagara/Public/NiagaraComponent.h
@@ -275,6 +275,8 @@ public:
 
 	NIAGARA_API void OnPooledReuse(UWorld* NewWorld);
 
+	bool IsActiveForUpdateContext() const;
+
 	/*
 	Switch which asset the component is using.
 	This requires Niagara to wait for concurrent execution and the override parameter store to be synchronized with the new asset.
