PR #13323: CascadeToNiagaraConverter: Add more script input types.

diff --git a/Engine/Plugins/FX/CascadeToNiagaraConverter/Source/CascadeToNiagaraConverter/Private/NiagaraStackGraphUtilitiesAdapterLibrary.cpp b/Engine/Plugins/FX/CascadeToNiagaraConverter/Source/CascadeToNiagaraConverter/Private/NiagaraStackGraphUtilitiesAdapterLibrary.cpp
--- a/Engine/Plugins/FX/CascadeToNiagaraConverter/Source/CascadeToNiagaraConverter/Private/NiagaraStackGraphUtilitiesAdapterLibrary.cpp
+++ b/Engine/Plugins/FX/CascadeToNiagaraConverter/Source/CascadeToNiagaraConverter/Private/NiagaraStackGraphUtilitiesAdapterLibrary.cpp
@@ -115,6 +115,8 @@ FName UFXConverterUtilitiesLibrary::GetNiagaraScriptInputTypeName(ENiagaraScript
 		return FNiagaraTypeDefinition::GetVec4Def().GetFName();
 	case ENiagaraScriptInputType::LinearColor:
 		return FNiagaraTypeDefinition::GetColorDef().GetFName();
+	case ENiagaraScriptInputType::Matrix:
+		return FNiagaraTypeDefinition::GetMatrix4Def().GetFName();
 	case ENiagaraScriptInputType::Quaternion:
 		return FNiagaraTypeDefinition::GetQuatDef().GetFName();
 	case ENiagaraScriptInputType::Bool:
@@ -239,8 +241,28 @@ UNiagaraScriptConversionContext* UFXConverterUtilitiesLibrary::CreateScriptConte
 	return ScriptContext;
 }
 
+UNiagaraScriptConversionContextInput* UFXConverterUtilitiesLibrary::CreateScriptInputLinkedEnum(const FString& ParameterNameString, const FString& UserDefinedEnumAssetPath)
+{
+	UUserDefinedEnum* UserDefinedEnum = LoadObject<UUserDefinedEnum>(nullptr, *UserDefinedEnumAssetPath);
+	if (UserDefinedEnum == nullptr)
+	{
+		return nullptr;
+	}
+
+	UNiagaraClipboardFunctionInput* NewInput = UNiagaraClipboardEditorScriptingUtilities::CreateEnumLinkedValueInput(GetTransientPackage(), FName(), false, false, UserDefinedEnum, FName(ParameterNameString));
+	if (NewInput != nullptr)
+	{
+		UNiagaraScriptConversionContextInput* Input = NewObject<UNiagaraScriptConversionContextInput>();
+		Input->Init(NewInput, ENiagaraScriptInputType::Enum, NewInput->GetTypeDef());
+		return Input;
+	}
+	return nullptr;
+}
+
 UNiagaraScriptConversionContextInput* UFXConverterUtilitiesLibrary::CreateScriptInputLinkedParameter(FString ParameterNameString, ENiagaraScriptInputType InputType)
 {
+	ensureMsgf(InputType != ENiagaraScriptInputType::Enum, TEXT("Use CreateEnumLinkedValueInput for Enum linked parameters"));
+
 	const FName InputTypeName = GetNiagaraScriptInputTypeName(InputType);
 	UNiagaraClipboardFunctionInput* NewInput = UNiagaraClipboardEditorScriptingUtilities::CreateLinkedValueInput(GetTransientPackage(), FName(), InputTypeName, false, false, FName(ParameterNameString));
 	const FNiagaraTypeDefinition& TargetTypeDef = UNiagaraClipboardEditorScriptingUtilities::GetRegisteredTypeDefinitionByName(InputTypeName);
@@ -268,6 +290,33 @@ UNiagaraScriptConversionContextInput* UFXConverterUtilitiesLibrary::CreateScript
 	return Input;
 }
 
+UNiagaraScriptConversionContextInput* UFXConverterUtilitiesLibrary::CreateScriptInputLinearColor(FLinearColor Value)
+{
+	UNiagaraClipboardFunctionInput* NewInput = UNiagaraClipboardEditorScriptingUtilities::CreateColorLocalValueInput(GetTransientPackage(), FName(), false, false, Value);
+	const FNiagaraTypeDefinition& TargetTypeDef = FNiagaraTypeDefinition::GetColorDef();
+	UNiagaraScriptConversionContextInput* Input = NewObject<UNiagaraScriptConversionContextInput>();
+	Input->Init(NewInput, ENiagaraScriptInputType::LinearColor, TargetTypeDef);
+	return Input;
+}
+
+UNiagaraScriptConversionContextInput* UFXConverterUtilitiesLibrary::CreateScriptInputMatrix(FMatrix44f Value)
+{
+	UNiagaraClipboardFunctionInput* NewInput = UNiagaraClipboardEditorScriptingUtilities::CreateMatrixLocalValueInput(GetTransientPackage(), FName(), false, false, Value);
+	const FNiagaraTypeDefinition& TargetTypeDef = FNiagaraTypeDefinition::GetMatrix4Def();
+	UNiagaraScriptConversionContextInput* Input = NewObject<UNiagaraScriptConversionContextInput>();
+	Input->Init(NewInput, ENiagaraScriptInputType::Matrix, TargetTypeDef);
+	return Input;
+}
+
+UNiagaraScriptConversionContextInput* UFXConverterUtilitiesLibrary::CreateScriptInputQuat(FQuat4f Value)
+{
+	UNiagaraClipboardFunctionInput* NewInput = UNiagaraClipboardEditorScriptingUtilities::CreateQuatLocalValueInput(GetTransientPackage(), FName(), false, false, Value);
+	const FNiagaraTypeDefinition& TargetTypeDef = FNiagaraTypeDefinition::GetQuatDef();
+	UNiagaraScriptConversionContextInput* Input = NewObject<UNiagaraScriptConversionContextInput>();
+	Input->Init(NewInput, ENiagaraScriptInputType::Quaternion, TargetTypeDef);
+	return Input;
+}
+
 UNiagaraScriptConversionContextInput* UFXConverterUtilitiesLibrary::CreateScriptInputVec2(FVector2D Value)
 {
 	UNiagaraClipboardFunctionInput* NewInput = UNiagaraClipboardEditorScriptingUtilities::CreateVec2LocalValueInput(GetTransientPackage(), FName(), false, false, Value);
@@ -277,6 +326,15 @@ UNiagaraScriptConversionContextInput* UFXConverterUtilitiesLibrary::CreateScript
 	return Input;
 }
 
+UNiagaraScriptConversionContextInput* UFXConverterUtilitiesLibrary::CreateScriptInputVec4(FVector4f Value)
+{
+	UNiagaraClipboardFunctionInput* NewInput = UNiagaraClipboardEditorScriptingUtilities::CreateVec4LocalValueInput(GetTransientPackage(), FName(), false, false, Value);
+	const FNiagaraTypeDefinition& TargetTypeDef = FNiagaraTypeDefinition::GetVec4Def();
+	UNiagaraScriptConversionContextInput* Input = NewObject<UNiagaraScriptConversionContextInput>();
+	Input->Init(NewInput, ENiagaraScriptInputType::Vec4, TargetTypeDef);
+	return Input;
+}
+
 UNiagaraScriptConversionContextInput* UFXConverterUtilitiesLibrary::CreateScriptInputVector(FVector Value)
 {
 	UNiagaraClipboardFunctionInput* NewInput = UNiagaraClipboardEditorScriptingUtilities::CreateVec3LocalValueInput(GetTransientPackage(), FName(), false, false, Value);
diff --git a/Engine/Plugins/FX/CascadeToNiagaraConverter/Source/CascadeToNiagaraConverter/Public/NiagaraStackGraphUtilitiesAdapterLibrary.h b/Engine/Plugins/FX/CascadeToNiagaraConverter/Source/CascadeToNiagaraConverter/Public/NiagaraStackGraphUtilitiesAdapterLibrary.h
--- a/Engine/Plugins/FX/CascadeToNiagaraConverter/Source/CascadeToNiagaraConverter/Public/NiagaraStackGraphUtilitiesAdapterLibrary.h
+++ b/Engine/Plugins/FX/CascadeToNiagaraConverter/Source/CascadeToNiagaraConverter/Public/NiagaraStackGraphUtilitiesAdapterLibrary.h
@@ -184,6 +184,7 @@ enum class ENiagaraScriptInputType : uint8
 	Vec3,
 	Vec4,
 	LinearColor,
+	Matrix,
 	Quaternion,
 	Struct,
 	Enum,
@@ -990,6 +991,9 @@ public:
 	UFUNCTION(BlueprintCallable, Category = "FXConverterUtilities")
 	static FAssetData CreateAssetData(FString InPath);
 
+	UFUNCTION(BlueprintCallable, Category = "FXConverterUtilities")
+	static UNiagaraScriptConversionContextInput* CreateScriptInputLinkedEnum(const FString& ParameterNameString, const FString& UserDefinedEnumAssetPath);
+
 	UFUNCTION(BlueprintCallable, Category = "FXConverterUtilities")
 	static UNiagaraScriptConversionContextInput* CreateScriptInputLinkedParameter(FString ParameterNameString, ENiagaraScriptInputType InputType);
 
@@ -999,9 +1003,21 @@ public:
 	UFUNCTION(BlueprintCallable, Category = "FXConverterUtilities")
 	static UNiagaraScriptConversionContextInput* CreateScriptInputFloat(float Value);
 
+	UFUNCTION(BlueprintCallable, Category = "FXConverterUtilities")
+	static UNiagaraScriptConversionContextInput* CreateScriptInputLinearColor(FLinearColor Value);
+
+	UFUNCTION(BlueprintCallable, Category = "FXConverterUtilities")
+	static UNiagaraScriptConversionContextInput* CreateScriptInputMatrix(FMatrix44f Value);
+
+	UFUNCTION(BlueprintCallable, Category = "FXConverterUtilities")
+	static UNiagaraScriptConversionContextInput* CreateScriptInputQuat(FQuat4f Value);
+
 	UFUNCTION(BlueprintCallable, Category = "FXConverterUtilities")
 	static UNiagaraScriptConversionContextInput* CreateScriptInputVec2(FVector2D Value);
 
+	UFUNCTION(BlueprintCallable, Category = "FXConverterUtilities")
+	static UNiagaraScriptConversionContextInput* CreateScriptInputVec4(FVector4f Value);
+
 	UFUNCTION(BlueprintCallable, Category = "FXConverterUtilities")
 	static UNiagaraScriptConversionContextInput* CreateScriptInputVector(FVector Value);
 
diff --git a/Engine/Plugins/FX/Niagara/Source/NiagaraEditor/Private/NiagaraClipboard.cpp b/Engine/Plugins/FX/Niagara/Source/NiagaraEditor/Private/NiagaraClipboard.cpp
--- a/Engine/Plugins/FX/Niagara/Source/NiagaraEditor/Private/NiagaraClipboard.cpp
+++ b/Engine/Plugins/FX/Niagara/Source/NiagaraEditor/Private/NiagaraClipboard.cpp
@@ -459,6 +459,16 @@ FNiagaraTypeDefinition UNiagaraClipboardEditorScriptingUtilities::GetRegisteredT
 	return FNiagaraTypeRegistry::GetRegisteredTypeByName(InTypeName).Get(FNiagaraTypeDefinition());
 }
 
+UNiagaraClipboardFunctionInput* UNiagaraClipboardEditorScriptingUtilities::CreateColorLocalValueInput(UObject* InOuter, FName InInputName, bool bInHasEditCondition, bool bInEditConditionValue, FLinearColor InColorValue)
+{
+	FNiagaraTypeDefinition InputType = FNiagaraTypeDefinition::GetColorDef();
+	TArray<uint8> ColorValue;
+	ColorValue.AddUninitialized(InputType.GetSize());
+	FMemory::Memcpy(ColorValue.GetData(), &InColorValue, InputType.GetSize());
+
+	return const_cast<UNiagaraClipboardFunctionInput*>(CreateLocalValue(InOuter, InInputName, InputType, bInHasEditCondition, bInEditConditionValue, ColorValue));
+}
+
 UNiagaraClipboardFunctionInput* UNiagaraClipboardEditorScriptingUtilities::CreateFloatLocalValueInput(UObject* InOuter, FName InInputName, bool bInHasEditCondition, bool bInEditConditionValue, float InFloatValue)
 {
 	FNiagaraTypeDefinition InputType = FNiagaraTypeDefinition::GetFloatDef();
@@ -469,6 +479,31 @@ UNiagaraClipboardFunctionInput* UNiagaraClipboardEditorScriptingUtilities::Creat
 	return const_cast<UNiagaraClipboardFunctionInput*>(CreateLocalValue(InOuter, InInputName, InputType, bInHasEditCondition, bInEditConditionValue, FloatValue));
 }
 
+UNiagaraClipboardFunctionInput* UNiagaraClipboardEditorScriptingUtilities::CreateMatrixLocalValueInput(UObject* InOuter, FName InInputName, bool bInHasEditCondition, bool bInEditConditionValue, FMatrix44f InMatrixValue)
+{
+	FNiagaraTypeDefinition InputType = FNiagaraTypeDefinition::GetMatrix4Def();
+	FNiagaraMatrix Matrix;
+	Matrix.Row0 = FVector4f(InMatrixValue.M[0][0], InMatrixValue.M[0][1], InMatrixValue.M[0][2], InMatrixValue.M[0][3]);
+	Matrix.Row1 = FVector4f(InMatrixValue.M[1][0], InMatrixValue.M[1][1], InMatrixValue.M[1][2], InMatrixValue.M[1][3]);
+	Matrix.Row2 = FVector4f(InMatrixValue.M[2][0], InMatrixValue.M[2][1], InMatrixValue.M[2][2], InMatrixValue.M[2][3]);
+	Matrix.Row3 = FVector4f(InMatrixValue.M[3][0], InMatrixValue.M[3][1], InMatrixValue.M[3][2], InMatrixValue.M[3][3]);
+	TArray<uint8> MatrixValue;
+	MatrixValue.AddUninitialized(InputType.GetSize());
+	FMemory::Memcpy(MatrixValue.GetData(), &Matrix, InputType.GetSize());
+
+	return const_cast<UNiagaraClipboardFunctionInput*>(CreateLocalValue(InOuter, InInputName, InputType, bInHasEditCondition, bInEditConditionValue, MatrixValue));
+}
+
+UNiagaraClipboardFunctionInput* UNiagaraClipboardEditorScriptingUtilities::CreateQuatLocalValueInput(UObject* InOuter, FName InInputName, bool bInHasEditCondition, bool bInEditConditionValue, FQuat4f InQuatValue)
+{
+	FNiagaraTypeDefinition InputType = FNiagaraTypeDefinition::GetQuatDef();
+	TArray<uint8> QuatValue;
+	QuatValue.AddUninitialized(InputType.GetSize());
+	FMemory::Memcpy(QuatValue.GetData(), &InQuatValue, InputType.GetSize());
+
+	return const_cast<UNiagaraClipboardFunctionInput*>(CreateLocalValue(InOuter, InInputName, InputType, bInHasEditCondition, bInEditConditionValue, QuatValue));
+}
+
 UNiagaraClipboardFunctionInput* UNiagaraClipboardEditorScriptingUtilities::CreateVec2LocalValueInput(UObject* InOuter, FName InInputName, bool bInHasEditCondition, bool bInEditConditionValue, FVector2D InVec2Value)
 {
 	FNiagaraTypeDefinition InputType = FNiagaraTypeDefinition::GetVec2Def();
@@ -491,6 +526,16 @@ UNiagaraClipboardFunctionInput* UNiagaraClipboardEditorScriptingUtilities::Creat
 	return const_cast<UNiagaraClipboardFunctionInput*>(CreateLocalValue(InOuter, InInputName, InputType, bInHasEditCondition, bInEditConditionValue, Vec3ValueFloat));
 }
 
+UNiagaraClipboardFunctionInput* UNiagaraClipboardEditorScriptingUtilities::CreateVec4LocalValueInput(UObject* InOuter, FName InInputName, bool bInHasEditCondition, bool bInEditConditionValue, FVector4f InVec4Value)
+{
+	FNiagaraTypeDefinition InputType = FNiagaraTypeDefinition::GetVec4Def();
+	TArray<uint8> Vec4Value;
+	Vec4Value.AddUninitialized(InputType.GetSize());
+	FMemory::Memcpy(Vec4Value.GetData(), &InVec4Value, InputType.GetSize());
+
+	return const_cast<UNiagaraClipboardFunctionInput*>(CreateLocalValue(InOuter, InInputName, InputType, bInHasEditCondition, bInEditConditionValue, Vec4Value));
+}
+
 UNiagaraClipboardFunctionInput* UNiagaraClipboardEditorScriptingUtilities::CreateIntLocalValueInput(UObject* InOuter, FName InInputName, bool bInHasEditCondition, bool bInEditConditionValue, int32 InIntValue)
 {
 	FNiagaraTypeDefinition InputType = FNiagaraTypeDefinition::GetIntDef();
@@ -546,6 +591,22 @@ UNiagaraClipboardFunctionInput* UNiagaraClipboardEditorScriptingUtilities::Creat
 	);
 }
 
+UNiagaraClipboardFunctionInput* UNiagaraClipboardEditorScriptingUtilities::CreateEnumLinkedValueInput(UObject* InOuter, FName InInputName, bool bInHasEditCondition, bool bInEditConditionValue, UUserDefinedEnum* InEnumType, FName InLinkedValue)
+{
+	FNiagaraTypeDefinition InputType = FNiagaraTypeDefinition(Cast<UEnum>(InEnumType));
+	if (InputType.IsValid())
+	{
+		FNiagaraVariableBase LinkedParameter = FNiagaraVariableBase(InputType, InLinkedValue);
+		return const_cast<UNiagaraClipboardFunctionInput*>(UNiagaraClipboardFunctionInput::CreateLinkedValue(
+			InOuter != nullptr ? InOuter : GetTransientPackage(),
+			InInputName,
+			InputType,
+			bInHasEditCondition ? TOptional<bool>(bInEditConditionValue) : TOptional<bool>(),
+			LinkedParameter));
+	}
+	return nullptr;
+}
+
 UNiagaraClipboardFunctionInput* UNiagaraClipboardEditorScriptingUtilities::CreateLinkedValueInput(UObject* InOuter, FName InInputName, FName InInputTypeName, bool bInHasEditCondition, bool bInEditConditionValue, FName InLinkedValue)
 {
 	FNiagaraTypeDefinition InputType = GetRegisteredTypeDefinitionByName(InInputTypeName);
diff --git a/Engine/Plugins/FX/Niagara/Source/NiagaraEditor/Public/NiagaraClipboard.h b/Engine/Plugins/FX/Niagara/Source/NiagaraEditor/Public/NiagaraClipboard.h
--- a/Engine/Plugins/FX/Niagara/Source/NiagaraEditor/Public/NiagaraClipboard.h
+++ b/Engine/Plugins/FX/Niagara/Source/NiagaraEditor/Public/NiagaraClipboard.h
@@ -315,15 +315,27 @@ public:
 	UFUNCTION(BlueprintPure, Category = "Input")
 	static NIAGARAEDITOR_API FName GetTypeName(const UNiagaraClipboardFunctionInput* InInput);
 
+	UFUNCTION(BlueprintPure, Category = "Input")
+	static NIAGARAEDITOR_API UNiagaraClipboardFunctionInput* CreateColorLocalValueInput(UObject* InOuter, FName InInputName, bool bInHasEditCondition, bool bInEditConditionValue, FLinearColor InColorValue);
+
 	UFUNCTION(BlueprintPure, Category = "Input")
 	static NIAGARAEDITOR_API UNiagaraClipboardFunctionInput* CreateFloatLocalValueInput(UObject* InOuter, FName InInputName, bool bInHasEditCondition, bool bInEditConditionValue, float InLocalValue);
 
+	UFUNCTION(BlueprintPure, Category = "Input")
+	static NIAGARAEDITOR_API UNiagaraClipboardFunctionInput* CreateMatrixLocalValueInput(UObject* InOuter, FName InInputName, bool bInHasEditCondition, bool bInEditConditionValue, FMatrix44f InMatrixValue);
+
+	UFUNCTION(BlueprintPure, Category = "Input")
+	static NIAGARAEDITOR_API UNiagaraClipboardFunctionInput* CreateQuatLocalValueInput(UObject* InOuter, FName InInputName, bool bInHasEditCondition, bool bInEditConditionValue, FQuat4f InQuatValue);
+
 	UFUNCTION(BlueprintPure, Category = "Input")
 	static NIAGARAEDITOR_API UNiagaraClipboardFunctionInput* CreateVec2LocalValueInput(UObject* InOuter, FName InInputName, bool bInHasEditCondition, bool bInEditConditionValue, FVector2D InVec2Value);
 
 	UFUNCTION(BlueprintPure, Category = "Input")
 	static NIAGARAEDITOR_API UNiagaraClipboardFunctionInput* CreateVec3LocalValueInput(UObject* InOuter, FName InInputName, bool bInHasEditCondition, bool bInEditConditionValue, FVector InVec3Value);
 
+	UFUNCTION(BlueprintPure, Category = "Input")
+	static NIAGARAEDITOR_API UNiagaraClipboardFunctionInput* CreateVec4LocalValueInput(UObject* InOuter, FName InInputName, bool bInHasEditCondition, bool bInEditConditionValue, FVector4f InVec4Value);
+
 	UFUNCTION(BlueprintPure, Category = "Input")
 	static NIAGARAEDITOR_API UNiagaraClipboardFunctionInput* CreateIntLocalValueInput(UObject* InOuter, FName InInputName, bool bInHasEditCondition, bool bInEditConditionValue, int32 InLocalValue);
 	
@@ -336,6 +348,9 @@ public:
 	UFUNCTION(BlueprintPure, Category = "Input")
 	static NIAGARAEDITOR_API UNiagaraClipboardFunctionInput* CreateEnumLocalValueInput(UObject* InOuter, FName InInputName, bool bInHasEditCondition, bool bInEditCoditionValue, UUserDefinedEnum* InEnumType, int32 InEnumValue);
 
+	UFUNCTION(BlueprintPure, Category = "Input")
+	static NIAGARAEDITOR_API UNiagaraClipboardFunctionInput* CreateEnumLinkedValueInput(UObject* InOuter, FName InInputName, bool bInHasEditCondition, bool bInEditCoditionValue, UUserDefinedEnum* InEnumType, FName InLinkedValue);
+
 	UFUNCTION(BlueprintPure, Category = "Input")
 	static NIAGARAEDITOR_API UNiagaraClipboardFunctionInput* CreateLinkedValueInput(UObject* InOuter, FName InInputName, FName InInputTypeName, bool bInHasEditCondition, bool bInEditConditionValue, FName InLinkedValue);
 
