- When changing the preview platform invalidate the cooked status
- We need to make sure everything is in sync when the device profile changes to avoid crashes inside the VM

diff --git a/Engine/Plugins/FX/Niagara/Source/Niagara/Classes/NiagaraScript.h b/Engine/Plugins/FX/Niagara/Source/Niagara/Classes/NiagaraScript.h
--- a/Engine/Plugins/FX/Niagara/Source/Niagara/Classes/NiagaraScript.h
+++ b/Engine/Plugins/FX/Niagara/Source/Niagara/Classes/NiagaraScript.h
@@ -1018,7 +1018,7 @@ public:
 	const FNiagaraVMExecutableDataId& GetComputedVMCompilationId() const
 	{
 #if WITH_EDITORONLY_DATA
-		if (!IsCooked)
+		if (!IsScriptCooked())
 		{
 			return GetLastGeneratedVMId();
 		}
@@ -1354,14 +1354,7 @@ public:
 	NIAGARA_API const FNiagaraScriptExecutionParameterStore* GetExecutionReadyParameterStore(ENiagaraSimTarget SimTarget);
 	void InvalidateExecutionReadyParameterStores();
 
-	bool IsScriptCooked() const
-	{
-#if WITH_EDITORONLY_DATA
-		return IsCooked;
-#else
-		return true;
-#endif
-	}
+	bool IsScriptCooked() const;
 
 private:
 	bool OwnerCanBeRunOnGpu() const;
diff --git a/Engine/Plugins/FX/Niagara/Source/Niagara/Private/NiagaraScript.cpp b/Engine/Plugins/FX/Niagara/Source/Niagara/Private/NiagaraScript.cpp
--- a/Engine/Plugins/FX/Niagara/Source/Niagara/Private/NiagaraScript.cpp
+++ b/Engine/Plugins/FX/Niagara/Source/Niagara/Private/NiagaraScript.cpp
@@ -1595,7 +1595,7 @@ bool UNiagaraScript::ContainsUsage(ENiagaraScriptUsage InUsage) const
 const FNiagaraScriptExecutionParameterStore* UNiagaraScript::GetExecutionReadyParameterStore(ENiagaraSimTarget SimTarget)
 {
 #if WITH_EDITORONLY_DATA
-	if (!IsCooked)
+	if (!IsScriptCooked())
 	{
 		if (SimTarget == ENiagaraSimTarget::CPUSim && IsReadyToRun(ENiagaraSimTarget::CPUSim))
 		{
@@ -2432,6 +2432,15 @@ void UNiagaraScript::SetPreviewFeatureLevel(ERHIFeatureLevel::Type InPreviewFeat
 		NiagaraScriptInternal::PreviewFeatureLevel = InPreviewFeatureLevel;
 	}
 
+	// We need to invalid IsCooked if we are changing preview level as the scripts can get out of sync
+	for (TObjectIterator<UNiagaraScript> It; It; ++It)
+	{
+		if (UNiagaraScript* Script = *It)
+		{
+			Script->IsCooked = false;
+		}
+	}
+
 	// Force scripts to recache that can run on the GPU and already have a script resource
 	if (!AreGpuScriptsCompiledBySystem())
 	{
@@ -3307,6 +3316,15 @@ void UNiagaraScript::InvalidateExecutionReadyParameterStores()
 #endif
 }
 
+bool UNiagaraScript::IsScriptCooked() const
+{
+#if WITH_EDITORONLY_DATA
+	return IsCooked;
+#else
+	return true;
+#endif
+}
+
 void UNiagaraScript::RequestCompile(const FGuid& ScriptVersion, bool bForceCompile)
 {
 	FVersionedNiagaraScriptData* ScriptData = GetScriptData(ScriptVersion);
