- Move IsScriptCooked back into the header

#rnx

diff --git a/Engine/Plugins/FX/Niagara/Source/Niagara/Classes/NiagaraScript.h b/Engine/Plugins/FX/Niagara/Source/Niagara/Classes/NiagaraScript.h
--- a/Engine/Plugins/FX/Niagara/Source/Niagara/Classes/NiagaraScript.h
+++ b/Engine/Plugins/FX/Niagara/Source/Niagara/Classes/NiagaraScript.h
@@ -1354,7 +1354,14 @@ public:
 	NIAGARA_API const FNiagaraScriptExecutionParameterStore* GetExecutionReadyParameterStore(ENiagaraSimTarget SimTarget);
 	void InvalidateExecutionReadyParameterStores();
 
-	bool IsScriptCooked() const;
+	bool IsScriptCooked() const
+	{
+	#if WITH_EDITORONLY_DATA
+		return IsCooked;
+	#else
+		return true;
+	#endif
+	}
 
 private:
 	bool OwnerCanBeRunOnGpu() const;
diff --git a/Engine/Plugins/FX/Niagara/Source/Niagara/Private/NiagaraScript.cpp b/Engine/Plugins/FX/Niagara/Source/Niagara/Private/NiagaraScript.cpp
--- a/Engine/Plugins/FX/Niagara/Source/Niagara/Private/NiagaraScript.cpp
+++ b/Engine/Plugins/FX/Niagara/Source/Niagara/Private/NiagaraScript.cpp
@@ -3316,15 +3316,6 @@ void UNiagaraScript::InvalidateExecutionReadyParameterStores()
 #endif
 }
 
-bool UNiagaraScript::IsScriptCooked() const
-{
-#if WITH_EDITORONLY_DATA
-	return IsCooked;
-#else
-	return true;
-#endif
-}
-
 void UNiagaraScript::RequestCompile(const FGuid& ScriptVersion, bool bForceCompile)
 {
 	FVersionedNiagaraScriptData* ScriptData = GetScriptData(ScriptVersion);
