RVT fedback from niagara landscape DI.  Enqueues pages of the RVT to be made resident.
#tests Validated pages are being made resident in a large level with detailed landscapes

diff --git a/Engine/Plugins/FX/Niagara/Source/Niagara/Classes/NiagaraDataInterfaceLandscape.h b/Engine/Plugins/FX/Niagara/Source/Niagara/Classes/NiagaraDataInterfaceLandscape.h
--- a/Engine/Plugins/FX/Niagara/Source/Niagara/Classes/NiagaraDataInterfaceLandscape.h
+++ b/Engine/Plugins/FX/Niagara/Source/Niagara/Classes/NiagaraDataInterfaceLandscape.h
@@ -86,6 +86,8 @@ public:
 	NIAGARA_API virtual bool PerInstanceTick(void* PerInstanceData, FNiagaraSystemInstance* SystemInstance, float DeltaSeconds) override;
 	virtual bool HasPreSimulateTick() const override { return true; }
 
+	virtual void GetVMExternalFunction(const FVMExternalFunctionBindingInfo& BindingInfo, void* PerInstanceData, FVMExternalFunction& OutFunc) override;
+
 protected:
 #if WITH_EDITORONLY_DATA
 	virtual void GetFunctionsInternal(TArray<FNiagaraFunctionSignature>& OutFunctions) const override;
@@ -94,10 +96,7 @@ protected:
 	NIAGARA_API void ApplyLandscape(const FNiagaraSystemInstance& SystemInstance, FNDILandscapeData_GameThread& InstanceData) const;
 	NIAGARA_API ALandscape* GetLandscape(const FNiagaraSystemInstance& SystemInstance, FNDILandscapeData_GameThread& Hint) const;
 
-	static NIAGARA_API const FName GetBaseColorName;
-	static NIAGARA_API const FName GetHeightName;
-	static NIAGARA_API const FName GetWorldNormalName;
-	static NIAGARA_API const FName GetPhysicalMaterialIndexName;
+	void VMGenerateRvtFeedback(FVectorVMExternalFunctionContext& Context);
 };
 
 
diff --git a/Engine/Plugins/FX/Niagara/Source/Niagara/Private/NiagaraDataInterfaceLandscape.cpp b/Engine/Plugins/FX/Niagara/Source/Niagara/Private/NiagaraDataInterfaceLandscape.cpp
--- a/Engine/Plugins/FX/Niagara/Source/Niagara/Private/NiagaraDataInterfaceLandscape.cpp
+++ b/Engine/Plugins/FX/Niagara/Source/Niagara/Private/NiagaraDataInterfaceLandscape.cpp
@@ -2,25 +2,25 @@
 
 #include "NiagaraDataInterfaceLandscape.h"
 
-#include "Containers/ResourceArray.h"
-#include "Misc/LargeWorldRenderPosition.h"
-#include "NiagaraShaderParametersBuilder.h"
-#include "GlobalRenderResources.h"
-#include "NiagaraStats.h"
-#include "NiagaraCompileHashVisitor.h"
-#include "NiagaraWorldManager.h"
-
 #include "Algo/RemoveIf.h"
-#include "RHIStaticStates.h"
-#include "VT/RuntimeVirtualTexture.h"
+#include "Containers/ResourceArray.h"
+#include "EngineModule.h"
 #include "EngineUtils.h"
+#include "GlobalRenderResources.h"
 #include "Landscape.h"
 #include "LandscapeHeightfieldCollisionComponent.h"
 #include "LandscapeInfo.h"
 #include "LandscapeProxy.h"
+#include "Misc/LargeWorldRenderPosition.h"
+#include "NiagaraCompileHashVisitor.h"
+#include "NiagaraShaderParametersBuilder.h"
+#include "NiagaraStats.h"
+#include "NiagaraWorldManager.h"
+#include "PhysicalMaterials/PhysicalMaterial.h"
 #include "RenderUtils.h"
+#include "RHIStaticStates.h"
 #include "ShaderCompilerCore.h"
-#include "PhysicalMaterials/PhysicalMaterial.h"
+#include "VT/RuntimeVirtualTexture.h"
 
 #include UE_INLINE_GENERATED_CPP_BY_NAME(NiagaraDataInterfaceLandscape)
 
@@ -36,6 +36,12 @@ struct FNDILandscapeData_GameThread;
 
 namespace NiagaraDataInterfaceLandscape
 {
+	const FName NAME_GetBaseColor("GetBaseColor");
+	const FName NAME_GetHeight("GetHeight");
+	const FName NAME_GetWorldNormal("GetWorldNormal");
+	const FName NAME_GetPhysicalMaterialIndex("GetPhysicalMaterialIndex");
+	const FName NAME_GenerateRvtFeedback("GenerateVirtualTextureFeedback");
+
 	// This controls the maximum number of regions that will be evaluated for capture in FNDI_Landscape_GeneratedData.
 	// A value of -1 means no limit.
 	static int32 GMaxRegionSearchCount = -1;
@@ -126,11 +132,6 @@ namespace NiagaraDataInterfaceLandscape
 	};
 }
 
-const FName UNiagaraDataInterfaceLandscape::GetBaseColorName(TEXT("GetBaseColor"));
-const FName UNiagaraDataInterfaceLandscape::GetHeightName(TEXT("GetHeight"));
-const FName UNiagaraDataInterfaceLandscape::GetWorldNormalName(TEXT("GetWorldNormal"));
-const FName UNiagaraDataInterfaceLandscape::GetPhysicalMaterialIndexName(TEXT("GetPhysicalMaterialIndex"));
-
 // RenderResource used to hold textures generated by this DI, pulled from the collision geometry of the terrain
 class FLandscapeTextureResource : public FRenderResource
 {
@@ -226,12 +227,22 @@ private:
 
 using FNDI_Landscape_SharedResourceHandle = FNDI_SharedResourceHandle<FNDI_Landscape_SharedResource, FNDI_SharedResourceUsage>;
 
+struct FRegionRequest
+{
+	FBox2D Region = FBox2D(FVector2D::Zero(), FVector2D::One());
+	int32 MipLevel = 0;
+	bool bIncludeBaseColor = true;
+	bool bIncludeHeight = true;
+	bool bIncludeNormal = true;
+};
+
 // Landscape data used for the game thread
 struct FNDILandscapeData_GameThread
 {
 	TWeakObjectPtr<ALandscape> Landscape;
 	TWeakObjectPtr<ULandscapeHeightfieldCollisionComponent> CollisionComponent;
 	FNDI_Landscape_SharedResourceHandle SharedResourceHandle;
+	TArray<FRegionRequest> RegionRequests;
 	bool BaseColorVirtualTextureSRGB = false;
 	NiagaraDataInterfaceLandscape::EBaseColorUnpackType BaseColorVirtualTextureUnpackType = NiagaraDataInterfaceLandscape::EBaseColorUnpackType::None;
 	int32 BaseColorVirtualTextureIndex = INDEX_NONE;
@@ -249,6 +260,7 @@ struct FNDILandscapeData_GameThread
 	{
 		Landscape = nullptr;
 		SharedResourceHandle = FNDI_Landscape_SharedResourceHandle();
+		RegionRequests.Empty();
 		BaseColorVirtualTextureSRGB = false;
 		BaseColorVirtualTextureUnpackType = NiagaraDataInterfaceLandscape::EBaseColorUnpackType::None;
 		BaseColorVirtualTextureIndex = INDEX_NONE;
@@ -273,6 +285,7 @@ struct FNDILandscapeData_GameToRenderThread
 	FMatrix CachedHeightTextureUvToWorldTransform = FMatrix::Identity;
 	FVector4 CachedHeightTextureUvScaleBias = FVector4::Zero();
 	FVector2D CachedHeightTextureWorldGridSize = FVector2D(1.0f, 1.0f);
+	TArray<FRegionRequest> RegionRequests;
 
 	NiagaraDataInterfaceLandscape::EBaseColorUnpackType BaseColorVirtualTextureUnpackType = NiagaraDataInterfaceLandscape::EBaseColorUnpackType::None;
 };
@@ -339,6 +352,9 @@ struct FNDILandscapeData_RenderThread
 			WorldToUvParameters[1] = FVector4::Zero();
 			WorldToUvParameters[2] = FVector4::Zero();
 			WorldToUvParameters[3] = FVector4::Zero();
+
+			AllocatedTexture = nullptr;
+			SizeInPages = FUintPoint(0, 0);
 		}
 
 		void Update(const URuntimeVirtualTexture* VirtualTexture, uint32 PageTableIndex, bool bIncludeWorldToUv, bool bIncludeHeightUnpack)
@@ -347,7 +363,8 @@ struct FNDILandscapeData_RenderThread
 
 			if (VirtualTexture && ::IsValid(VirtualTexture))
 			{
-				if (IAllocatedVirtualTexture* AllocatedTexture = VirtualTexture->GetAllocatedVirtualTexture())
+				AllocatedTexture = VirtualTexture->GetAllocatedVirtualTexture();
+				if (AllocatedTexture)
 				{
 					PageTableRef = AllocatedTexture->GetPageTableTexture(PageTableIndex);
 					if (PageTableRef.IsValid())
@@ -366,14 +383,27 @@ struct FNDILandscapeData_RenderThread
 							WorldToUvParameters[3] = VirtualTexture->GetUniformParameter(ERuntimeVirtualTextureShaderUniform_WorldHeightUnpack);
 						}
 					}
+
+					SizeInPages = FUintPoint(AllocatedTexture->GetWidthInTiles(), AllocatedTexture->GetHeightInTiles());
 				}
 			}
 		}
 
+		FVector2D WorldToUv(const FVector& WorldPosition) const
+		{
+			const FVector NormalizedPosition = WorldPosition - WorldToUvParameters[0];
+			return FVector2D(
+				FVector(WorldToUvParameters[1]).Dot(NormalizedPosition),
+				FVector(WorldToUvParameters[2]).Dot(NormalizedPosition));
+		}
+
 		FTextureRHIRef PageTableRef;
 		FUintVector4 PageTableUniforms[2];
 
 		FVector4 WorldToUvParameters[4];
+
+		IAllocatedVirtualTexture* AllocatedTexture = nullptr;
+		FUintPoint SizeInPages = FUintPoint(0, 0);
 	};
 
 	FVirtualTexturePage BaseColorVirtualPage;
@@ -663,6 +693,35 @@ struct FNiagaraDataInterfaceProxyLandscape : public FNiagaraDataInterfaceProxy
 				Proxy_RT->NormalUnpackMode = FNDILandscapeData_RenderThread::ENormalUnpackType::None;
 				break;
 			}
+
+			TArray<FRegionRequest> RegionsToProcess;
+			Swap(RegionsToProcess, Proxy_RT->LandscapeData.RegionRequests);
+
+			if (!RegionsToProcess.IsEmpty())
+			{
+				IRendererModule& RenderModule = GetRendererModule();
+
+				for (const FRegionRequest& RegionRequest : RegionsToProcess)
+				{
+					const FNDILandscapeData_RenderThread::FVirtualTexturePage* TexturePages[] =
+					{
+						RegionRequest.bIncludeBaseColor ? &Proxy_RT->BaseColorVirtualPage : nullptr,
+						RegionRequest.bIncludeHeight ? &Proxy_RT->HeightVirtualPage : nullptr,
+						RegionRequest.bIncludeNormal ? &Proxy_RT->NormalVirtualPage : nullptr,
+					};
+
+					for (const FNDILandscapeData_RenderThread::FVirtualTexturePage* TexturePage : TexturePages)
+					{
+						if (TexturePage && TexturePage->IsValid() && TexturePage->AllocatedTexture)
+						{
+							const FVector2D UVMin = TexturePage->WorldToUv(FVector(RegionRequest.Region.Min.X, RegionRequest.Region.Min.Y, 0.0)) * TexturePage->SizeInPages;
+							const FVector2D UVMax = TexturePage->WorldToUv(FVector(RegionRequest.Region.Max.X, RegionRequest.Region.Max.Y, 0.0)) * TexturePage->SizeInPages;
+
+							RenderModule.RequestVirtualTextureTiles(TexturePage->AllocatedTexture, FVector2D::One(), FVector2D::Zero(), FVector2D::One(), UVMin, UVMax, RegionRequest.MipLevel);
+						}
+					}
+				}
+			}
 		}
 	}
 
@@ -1350,7 +1409,7 @@ void UNiagaraDataInterfaceLandscape::GetFunctionsInternal(TArray<FNiagaraFunctio
 {
 	{
 		FNiagaraFunctionSignature Sig;
-		Sig.Name = GetBaseColorName;
+		Sig.Name = NiagaraDataInterfaceLandscape::NAME_GetBaseColor;
 		Sig.bMemberFunction = true;
 		Sig.bRequiresContext = false;
 		Sig.bSupportsCPU = false;
@@ -1364,7 +1423,7 @@ void UNiagaraDataInterfaceLandscape::GetFunctionsInternal(TArray<FNiagaraFunctio
 
 	{
 		FNiagaraFunctionSignature Sig;
-		Sig.Name = GetHeightName;
+		Sig.Name = NiagaraDataInterfaceLandscape::NAME_GetHeight;
 		Sig.bMemberFunction = true;
 		Sig.bRequiresContext = false;
 		Sig.bSupportsCPU = false;
@@ -1378,7 +1437,7 @@ void UNiagaraDataInterfaceLandscape::GetFunctionsInternal(TArray<FNiagaraFunctio
 
 	{
 		FNiagaraFunctionSignature Sig;
-		Sig.Name = GetWorldNormalName;
+		Sig.Name = NiagaraDataInterfaceLandscape::NAME_GetWorldNormal;
 		Sig.bMemberFunction = true;
 		Sig.bRequiresContext = false;
 		Sig.bSupportsCPU = false;
@@ -1392,7 +1451,7 @@ void UNiagaraDataInterfaceLandscape::GetFunctionsInternal(TArray<FNiagaraFunctio
 
 	{
 		FNiagaraFunctionSignature Sig;
-		Sig.Name = GetPhysicalMaterialIndexName;
+		Sig.Name = NiagaraDataInterfaceLandscape::NAME_GetPhysicalMaterialIndex;
 		Sig.bMemberFunction = true;
 		Sig.bRequiresContext = false;
 		Sig.bSupportsCPU = false;
@@ -1404,6 +1463,28 @@ void UNiagaraDataInterfaceLandscape::GetFunctionsInternal(TArray<FNiagaraFunctio
 		OutFunctions.Add(Sig);
 	}
 
+	{
+		FNiagaraFunctionSignature Sig;
+		Sig.Name = NiagaraDataInterfaceLandscape::NAME_GenerateRvtFeedback;
+		Sig.bExperimental = true;
+		Sig.bMemberFunction = true;
+		Sig.bRequiresContext = false;
+		Sig.bRequiresExecPin = true;
+		Sig.bSupportsCPU = true;
+		Sig.bSupportsGPU = false;
+		Sig.ModuleUsageBitmask = ENiagaraScriptUsageMask::Emitter | ENiagaraScriptUsageMask::System;
+		Sig.Inputs.Add(FNiagaraVariable(FNiagaraTypeDefinition(GetClass()), TEXT("Landscape")));
+		Sig.Inputs.Add(FNiagaraVariable(FNiagaraTypeDefinition::GetVec2Def(), TEXT("RegionExtentMin")));
+		Sig.Inputs.Add(FNiagaraVariable(FNiagaraTypeDefinition::GetVec2Def(), TEXT("RegionExtentMax")));
+		Sig.Inputs.Emplace_GetRef(FNiagaraVariable(FNiagaraTypeDefinition::GetIntDef(), TEXT("MipLevel"))).SetValue(0);
+		Sig.Inputs.Emplace_GetRef(FNiagaraVariable(FNiagaraTypeDefinition::GetBoolDef(), TEXT("IncludeBaseColor"))).SetValue(true);
+		Sig.Inputs.Emplace_GetRef(FNiagaraVariable(FNiagaraTypeDefinition::GetBoolDef(), TEXT("IncludeHeight"))).SetValue(true);
+		Sig.Inputs.Emplace_GetRef(FNiagaraVariable(FNiagaraTypeDefinition::GetBoolDef(), TEXT("IncludeNormal"))).SetValue(true);
+		Sig.Inputs.Emplace_GetRef(FNiagaraVariable(FNiagaraTypeDefinition::GetBoolDef(), TEXT("Enabled"))).SetValue(true);
+		Sig.SetDescription(LOCTEXT("GenerateRVTFeedback", "Generates a request to refine the RVT in the specified region."));
+
+		OutFunctions.Add(Sig);
+	}
 	for (FNiagaraFunctionSignature& Sig : OutFunctions)
 	{
 		Sig.FunctionVersion = NiagaraDataInterfaceLandscape::LatestVersion;
@@ -1431,9 +1512,52 @@ bool UNiagaraDataInterfaceLandscape::UpgradeFunctionCall(FNiagaraFunctionSignatu
 }
 #endif
 
+void UNiagaraDataInterfaceLandscape::GetVMExternalFunction(const FVMExternalFunctionBindingInfo& BindingInfo, void* PerInstanceData, FVMExternalFunction& OutFunc)
+{
+	using namespace NiagaraDataInterfaceLandscape;
+
+	if (BindingInfo.Name == NiagaraDataInterfaceLandscape::NAME_GenerateRvtFeedback)
+	{
+		OutFunc = FVMExternalFunction::CreateLambda([this](FVectorVMExternalFunctionContext& Context) { VMGenerateRvtFeedback(Context); });
+	}
+}
+
+void UNiagaraDataInterfaceLandscape::VMGenerateRvtFeedback(FVectorVMExternalFunctionContext& Context)
+{
+	VectorVM::FUserPtrHandler<FNDILandscapeData_GameThread> InstanceData(Context);
+	FNDIInputParam<FVector2f> InExtentMin(Context);
+	FNDIInputParam<FVector2f> InExtentMax(Context);
+	FNDIInputParam<int32> InMipLevel(Context);
+	FNDIInputParam<bool> InIncludeBaseColor(Context);
+	FNDIInputParam<bool> InIncludeHeight(Context);
+	FNDIInputParam<bool> InIncludeNormal(Context);
+	FNDIInputParam<bool> InEnabled(Context);
+
+	for (int32 InstanceIdx = 0; InstanceIdx < Context.GetNumInstances(); ++InstanceIdx)
+	{
+		const FVector2f ExtentMin = InExtentMin.GetAndAdvance();
+		const FVector2f ExtentMax = InExtentMax.GetAndAdvance();
+		const int32 RequestedMipLevel = InMipLevel.GetAndAdvance();
+		const bool bIncludeBaseColor = InIncludeBaseColor.GetAndAdvance();
+		const bool bIncludeHeight = InIncludeHeight.GetAndAdvance();
+		const bool bIncludeNormal = InIncludeNormal.GetAndAdvance();
+		const bool bEnabled = InEnabled.GetAndAdvance();
+
+		if (bEnabled)
+		{
+			FRegionRequest& Request = InstanceData->RegionRequests.AddDefaulted_GetRef();
+			Request.Region = FBox2D(FVector2D(ExtentMin), FVector2D(ExtentMax));
+			Request.MipLevel = RequestedMipLevel;
+			Request.bIncludeBaseColor = bIncludeBaseColor;
+			Request.bIncludeHeight = bIncludeHeight;
+			Request.bIncludeNormal = bIncludeNormal;
+		}
+	}
+}
+
 void UNiagaraDataInterfaceLandscape::ProvidePerInstanceDataForRenderThread(void* DataForRenderThread, void* PerInstanceData, const FNiagaraSystemInstanceID& SystemInstance)
 {
-	const FNDILandscapeData_GameThread& SourceData = *reinterpret_cast<const FNDILandscapeData_GameThread*>(PerInstanceData);
+	FNDILandscapeData_GameThread& SourceData = *reinterpret_cast<FNDILandscapeData_GameThread*>(PerInstanceData);
 	FNDILandscapeData_GameToRenderThread* TargetData = new(DataForRenderThread) FNDILandscapeData_GameToRenderThread();
 
 	if (SourceData.SharedResourceHandle)
@@ -1450,6 +1574,7 @@ void UNiagaraDataInterfaceLandscape::ProvidePerInstanceDataForRenderThread(void*
 	}
 
 	TargetData->BaseColorVirtualTextureUnpackType = SourceData.BaseColorVirtualTextureUnpackType;
+	TargetData->RegionRequests = MoveTemp(SourceData.RegionRequests);
 }
 
 #if WITH_EDITORONLY_DATA
@@ -1480,10 +1605,10 @@ void UNiagaraDataInterfaceLandscape::GetParameterDefinitionHLSL(const FNiagaraDa
 
 bool UNiagaraDataInterfaceLandscape::GetFunctionHLSL(const FNiagaraDataInterfaceGPUParamInfo& ParamInfo, const FNiagaraDataInterfaceGeneratedFunction& FunctionInfo, int FunctionInstanceIndex, FString& OutHLSL)
 {
-	if ( (FunctionInfo.DefinitionName == GetBaseColorName) ||
-		 (FunctionInfo.DefinitionName == GetHeightName) ||
-		 (FunctionInfo.DefinitionName == GetWorldNormalName) ||
-		 (FunctionInfo.DefinitionName == GetPhysicalMaterialIndexName) )
+	if ( (FunctionInfo.DefinitionName == NiagaraDataInterfaceLandscape::NAME_GetBaseColor) ||
+		 (FunctionInfo.DefinitionName == NiagaraDataInterfaceLandscape::NAME_GetHeight) ||
+		 (FunctionInfo.DefinitionName == NiagaraDataInterfaceLandscape::NAME_GetWorldNormal) ||
+		 (FunctionInfo.DefinitionName == NiagaraDataInterfaceLandscape::NAME_GetPhysicalMaterialIndex) )
 	{
 		return true;
 	}
@@ -1556,19 +1681,19 @@ bool UNiagaraDataInterfaceLandscape::InitPerInstanceData(void* PerInstanceData,
 
 	bool SystemRequiresHeightsCpu = false;
 	bool SystemRequiresHeightsGpu = false;
-	SystemInstance->EvaluateBoundFunction(GetHeightName, SystemRequiresHeightsCpu, SystemRequiresHeightsGpu);
+	SystemInstance->EvaluateBoundFunction(NiagaraDataInterfaceLandscape::NAME_GetHeight, SystemRequiresHeightsCpu, SystemRequiresHeightsGpu);
 
 	bool SystemRequiresNormalsCpu = false;
 	bool SystemRequiresNormalsGpu = false;
-	SystemInstance->EvaluateBoundFunction(GetWorldNormalName, SystemRequiresNormalsCpu, SystemRequiresNormalsGpu);
+	SystemInstance->EvaluateBoundFunction(NiagaraDataInterfaceLandscape::NAME_GetWorldNormal, SystemRequiresNormalsCpu, SystemRequiresNormalsGpu);
 
 	bool SystemRequiresPhysMatCpu = false;
 	bool SystemRequiresPhysMatGpu = false;
-	SystemInstance->EvaluateBoundFunction(GetPhysicalMaterialIndexName, SystemRequiresPhysMatCpu, SystemRequiresPhysMatGpu);
+	SystemInstance->EvaluateBoundFunction(NiagaraDataInterfaceLandscape::NAME_GetPhysicalMaterialIndex, SystemRequiresPhysMatCpu, SystemRequiresPhysMatGpu);
 
 	bool SystemRequiresBaseColorCpu = false;
 	bool SystemRequiresBaseColorGpu = false;
-	SystemInstance->EvaluateBoundFunction(GetBaseColorName, SystemRequiresBaseColorCpu, SystemRequiresBaseColorGpu);
+	SystemInstance->EvaluateBoundFunction(NiagaraDataInterfaceLandscape::NAME_GetBaseColor, SystemRequiresBaseColorCpu, SystemRequiresBaseColorGpu);
 
 	InstanceData->SystemRequiresBaseColorGpu = SystemRequiresBaseColorGpu;
 	InstanceData->SystemRequiresHeightsGpu = SystemRequiresHeightsGpu;
@@ -1615,7 +1740,7 @@ bool UNiagaraDataInterfaceLandscape::PerInstanceTick(void* PerInstanceData, FNia
 	if (FNDILandscapeData_GameThread* InstanceData = reinterpret_cast<FNDILandscapeData_GameThread*>(PerInstanceData))
 	{
 		// todo - can we avoid checking this every tick?  currently it is required in case the landscape beneath us changes or new data
-		// has stremaed in for the landscape and we need to update our capture of the data
+		// has streamed in for the landscape and we need to update our capture of the data
 		ApplyLandscape(*SystemInstance, *InstanceData);
 
 		if (InstanceData->RequiresCollisionCacheGpu || InstanceData->RequiresPhysMatCacheGpu)
