- Remove all the async optimize / decompression of VM byte code
- The new VM does not use this path so the code is dead

#tests Ran ET

diff --git a/Engine/Plugins/FX/Niagara/Source/Niagara/Classes/NiagaraScript.h b/Engine/Plugins/FX/Niagara/Source/Niagara/Classes/NiagaraScript.h
--- a/Engine/Plugins/FX/Niagara/Source/Niagara/Classes/NiagaraScript.h
+++ b/Engine/Plugins/FX/Niagara/Source/Niagara/Classes/NiagaraScript.h
@@ -394,24 +394,6 @@ struct TStructOpsTypeTraits<FNiagaraVMExecutableByteCode> : public TStructOpsTyp
 	};
 };
 
-// Carrier for uncompressed, and optimmized bytecode returned from optimization task
-struct FNiagaraScriptAsyncOptimizeTaskState
-{
-	FNiagaraVMExecutableDataId CachedScriptVMId;
-	FNiagaraVMExecutableByteCode SourceByteCode;
-	FNiagaraVMExecutableByteCode OptimizedByteCode;
-
-	TArray<uint8, TInlineAllocator<32>> ExternalFunctionRegisterCounts;
-	
-	bool bShouldOptimizeByteCode = false;
-	bool bShouldFreeSourceByteCodeOnCooked = false;
-	bool bOptimizationComplete = false;
-		
-	NIAGARA_API void OptimizeByteCode();
-};
-using FNiagaraScriptAsyncOptimizeTaskStatePtr = TSharedPtr<FNiagaraScriptAsyncOptimizeTaskState, ESPMode::ThreadSafe>;
-
-
 /** Struct containing all of the data needed to run a Niagara VM executable script.*/
 USTRUCT()
 struct FNiagaraVMExecutableData
@@ -424,23 +406,6 @@ public:
 	UPROPERTY()
 	FNiagaraVMExecutableByteCode ByteCode;
 
-	/** Optimized version of the byte code to execute for this system */
-	UPROPERTY(Transient)
-	FNiagaraVMExecutableByteCode OptimizedByteCode;
-	// Container for the optimization task that doesn't copy the task over when the executable data is copied
-	struct FOptimizationTask
-	{
-		FOptimizationTask() { }
-		FOptimizationTask(const FOptimizationTask&) { }
-		FOptimizationTask& operator=(const FOptimizationTask&) { return *this; }
-
-		// Optimization task if one is pending, that should be applied before execution
-		FNiagaraScriptAsyncOptimizeTaskStatePtr State;
-
-		// Lock used to apply the optimization results
-		FCriticalSection Lock;
-	} OptimizationTask;
-
 	/** Number of temp registers used by this script. */
 	UPROPERTY()
 	int32 NumTempRegisters;
@@ -582,8 +547,6 @@ public:
 	UPROPERTY()
 	uint32 bNeedsGPUContextInit : 1;
 
-	NIAGARA_API void ApplyFinishedOptimization(const FNiagaraVMExecutableDataId& CachedScriptVMId, const FNiagaraScriptAsyncOptimizeTaskStatePtr& Result);
-
 	NIAGARA_API void SerializeData(FArchive& Ar, bool bDDCData);
 	
 	NIAGARA_API bool IsValid() const;
@@ -1373,16 +1336,9 @@ private:
 	/** Return the expected SimTarget for this script. Only returns a valid target if there is valid data to run with. */
 	TOptional<ENiagaraSimTarget> GetSimTarget() const;
 
-	bool ShouldDecompressByteCode() const;
-	bool ShouldOptimizeByteCode() const;
-	bool ShouldFreeUnoptimizedByteCode() const;
-
 	void ProcessCompilationErrors(bool bCompileSucceeded, TConstArrayView<FShaderCompilerError> CompilationErrors, TArray<FString>& ProcessedErrors);
 
 public:
-	/** Kicks off an async job to convert the ByteCode into an optimized version for the platform we are running on. */
-	FGraphEventRef HandleByteCodeOptimization(bool bShouldForceNow = false);
-
 	static int32 NIAGARA_API LogCompileStaticVars;
 
 private:
diff --git a/Engine/Plugins/FX/Niagara/Source/Niagara/Classes/NiagaraSystem.h b/Engine/Plugins/FX/Niagara/Source/Niagara/Classes/NiagaraSystem.h
--- a/Engine/Plugins/FX/Niagara/Source/Niagara/Classes/NiagaraSystem.h
+++ b/Engine/Plugins/FX/Niagara/Source/Niagara/Classes/NiagaraSystem.h
@@ -1065,14 +1065,9 @@ protected:
 	TArray<FNiagaraDataSetAccessor<ENiagaraExecutionState>> EmitterExecutionStateAccessors;
 	TArray<TArray<FNiagaraDataSetAccessor<FNiagaraSpawnInfo>>> EmitterSpawnInfoAccessors;
 
-	FGraphEventRef ScriptOptimizationCompletionEvent;
 	FGraphEventRef PSOPrecacheCompletionEvent;
 
 public:
-	NIAGARA_API void AsyncOptimizeAllScripts();
-
-	NIAGARA_API FGraphEventRef GetScriptOptimizationCompletionEvent();
-
 	bool AllowCullingForLocalPlayers() const { return bAllowCullingForLocalPlayers; }
 
 	const FNiagaraSystemStaticBuffers* GetStaticBuffers() const { return StaticBuffers.Get(); }
diff --git a/Engine/Plugins/FX/Niagara/Source/Niagara/Private/NiagaraComponent.cpp b/Engine/Plugins/FX/Niagara/Source/Niagara/Private/NiagaraComponent.cpp
--- a/Engine/Plugins/FX/Niagara/Source/Niagara/Private/NiagaraComponent.cpp
+++ b/Engine/Plugins/FX/Niagara/Source/Niagara/Private/NiagaraComponent.cpp
@@ -1423,8 +1423,6 @@ void UNiagaraComponent::ActivateInternal(bool bReset /* = false */, bool bIsScal
 
 	DestroyCullProxy();
 
-	Asset->AsyncOptimizeAllScripts();
-
 	// We can't call 'Super::Activate(bReset);' as this will enable the component tick
 	if (bReset || ShouldActivate() == true)
 	{
diff --git a/Engine/Plugins/FX/Niagara/Source/Niagara/Private/NiagaraScript.cpp b/Engine/Plugins/FX/Niagara/Source/Niagara/Private/NiagaraScript.cpp
--- a/Engine/Plugins/FX/Niagara/Source/Niagara/Private/NiagaraScript.cpp
+++ b/Engine/Plugins/FX/Niagara/Source/Niagara/Private/NiagaraScript.cpp
@@ -93,14 +93,6 @@ static FAutoConsoleVariableRef CVarNiagaraCompressScriptByteCode(
 	ECVF_Default
 );
 
-bool GNiagaraDelayScriptAsyncOptimization = true;
-static FAutoConsoleVariableRef CVarNiagaraDelayScriptAsyncOptimization(
-	TEXT("fx.Niagara.DelayScriptAsyncOptimization"),
-	GNiagaraDelayScriptAsyncOptimization,
-	TEXT("Should we delay the async optimization until the emitter is activated?"),
-	ECVF_Default
-);
-
 int32 UNiagaraScript::LogCompileStaticVars =0;
 static FAutoConsoleVariableRef CVarLogCompileStaticVars(
 	TEXT("fx.LogCompileStaticVars"),
@@ -377,20 +369,6 @@ void FNiagaraVMExecutableData::Reset()
 	*this = FNiagaraVMExecutableData();
 }
 
-void FNiagaraVMExecutableData::ApplyFinishedOptimization(const FNiagaraVMExecutableDataId& CachedScriptVMId, const FNiagaraScriptAsyncOptimizeTaskStatePtr& Result)
-{
-	check(Result->bOptimizationComplete == true);
-	if ( CachedScriptVMId == Result->CachedScriptVMId )
-	{
-		ByteCode = MoveTemp(Result->SourceByteCode);
-		OptimizedByteCode = MoveTemp(Result->OptimizedByteCode);
-	}
-	else
-	{
- 		UE_LOG(LogNiagara, Warning, TEXT("Skipping ApplyFinishedOptimization as the ID is invalid"));
-	}
-}
-
 void FNiagaraVMExecutableData::SerializeData(FArchive& Ar, bool bDDCData)
 {
 	UScriptStruct* FNiagaraVMExecutableDataType = FNiagaraVMExecutableData::StaticStruct();
@@ -1664,98 +1642,6 @@ TOptional<ENiagaraSimTarget> UNiagaraScript::GetSimTarget() const
 	return TOptional<ENiagaraSimTarget>();
 }
 
-void FNiagaraScriptAsyncOptimizeTaskState::OptimizeByteCode()
-{		
-	SourceByteCode.Uncompress();
-
-	bOptimizationComplete = true;
-}
-
-struct FNiagaraScriptAsyncOptimizeTask
-{
-	FNiagaraScriptAsyncOptimizeTask(FNiagaraScriptAsyncOptimizeTaskStatePtr InResults)
-		: Results(InResults)
-	{
-	}
-
-	FORCEINLINE TStatId GetStatId() const { RETURN_QUICK_DECLARE_CYCLE_STAT(FNiagaraScriptAsyncOptimizeTask, STATGROUP_TaskGraphTasks); }
-	ENamedThreads::Type GetDesiredThread() { return ENamedThreads::AnyThread; }
-	static ESubsequentsMode::Type GetSubsequentsMode() { return ESubsequentsMode::TrackSubsequents; }
-
-	void DoTask(ENamedThreads::Type CurrentThread, const FGraphEventRef& MyCompletionGraphEvent)
-	{
-		Results->OptimizeByteCode();
-	}
-	
-	FNiagaraScriptAsyncOptimizeTaskStatePtr Results;
-};
-
-bool UNiagaraScript::ShouldDecompressByteCode() const
-{
-#if VECTORVM_SUPPORTS_LEGACY
-	return CachedScriptVM.IsValid() && CachedScriptVM.ByteCode.IsCompressed();
-#else
-	return false;
-#endif
-}
-
-bool UNiagaraScript::ShouldOptimizeByteCode() const
-{
-	static const IConsoleVariable* CVarOptimizeVMCode = IConsoleManager::Get().FindConsoleVariable(TEXT("vm.OptimizeVMByteCode"));
-	if (!CVarOptimizeVMCode || CVarOptimizeVMCode->GetInt() == 0 )
-	{
-		return false;
-	}
-	
-	return CachedScriptVM.IsValid() && CachedScriptVM.ByteCode.HasByteCode() && !CachedScriptVM.OptimizedByteCode.HasByteCode();
-}
-
-bool UNiagaraScript::ShouldFreeUnoptimizedByteCode() const
-{
-	static const IConsoleVariable* CVarFreeUnoptimizedByteCode = IConsoleManager::Get().FindConsoleVariable(TEXT("vm.FreeUnoptimizedByteCode"));
-	return (FPlatformProperties::RequiresCookedData() || IsScriptCooked()) && CVarFreeUnoptimizedByteCode && CVarFreeUnoptimizedByteCode->GetInt() != 0;
-}
-
-FGraphEventRef UNiagaraScript::HandleByteCodeOptimization(bool bShouldForceNow)
-{
-	check(IsInGameThread());
-
-	const bool bHasOptimizationTask = CachedScriptVM.OptimizationTask.State.IsValid();
-	
-	if ((!ShouldDecompressByteCode() && !ShouldOptimizeByteCode()) || bHasOptimizationTask)
-	{
-		return nullptr;
-	}
-	
-	auto OptimizedByteCodeResults = MakeShared<FNiagaraScriptAsyncOptimizeTaskState, ESPMode::ThreadSafe>();
-	OptimizedByteCodeResults->CachedScriptVMId = CachedScriptVMId;
-	OptimizedByteCodeResults->SourceByteCode = CachedScriptVM.ByteCode;
-	OptimizedByteCodeResults->bShouldOptimizeByteCode = ShouldOptimizeByteCode();
-	OptimizedByteCodeResults->bShouldFreeSourceByteCodeOnCooked = ShouldFreeUnoptimizedByteCode();
-
-	// This has to be done game code side as we can not access anything in CachedScriptVM
-	OptimizedByteCodeResults->ExternalFunctionRegisterCounts.Reserve(CachedScriptVM.CalledVMExternalFunctions.Num());
-	for (const FVMExternalFunctionBindingInfo& FunctionBindingInfo : CachedScriptVM.CalledVMExternalFunctions)
-	{
-		ensure(FunctionBindingInfo.GetNumInputs() + FunctionBindingInfo.GetNumOutputs() <= TNumericLimits<uint8>::Max());
-		const uint8 RegisterCount = uint8(FunctionBindingInfo.GetNumInputs() + FunctionBindingInfo.GetNumOutputs());
-		OptimizedByteCodeResults->ExternalFunctionRegisterCounts.Add(RegisterCount);
-	}
-	
-	if (bShouldForceNow)
-	{		
-		OptimizedByteCodeResults->OptimizeByteCode();
-		CachedScriptVM.ApplyFinishedOptimization(CachedScriptVMId, OptimizedByteCodeResults);
-		return nullptr;
-	}
-	else
-	{
-		// Dispatch the async task to decompress/optimize the byte code
-		CachedScriptVM.OptimizationTask.State = OptimizedByteCodeResults;
-		return TGraphTask<FNiagaraScriptAsyncOptimizeTask>::CreateTask().ConstructAndDispatchWhenReady(OptimizedByteCodeResults);
-	}	
-}
-
 void UNiagaraScript::GenerateDefaultFunctionBindings()
 {
 	// generate the function bindings for those external functions where there's no user (per-instance) data required
@@ -2388,12 +2274,6 @@ void UNiagaraScript::PostLoad()
 	ProcessSerializedShaderMaps();
 
 	GenerateStatIDs();
-
-	// Optimize the VM script for runtime usage
-	if (!GNiagaraDelayScriptAsyncOptimization)
-	{		
-		HandleByteCodeOptimization(true);
-	}
 }
 
 bool UNiagaraScript::IsReadyToRun(ENiagaraSimTarget SimTarget) const
@@ -2994,7 +2874,6 @@ void UNiagaraScript::SetVMCompilationResults(const FNiagaraVMExecutableDataId& I
 {
 	CachedScriptVMId = InCompileId;
 	CachedScriptVM = InScriptVM;
-	CachedScriptVM.OptimizationTask.State = nullptr;
 	CachedParameterCollectionReferences.Empty();
 	// Proactively clear out the script resource, because it might be stale now.
 	ScriptResource->Invalidate();
@@ -3089,12 +2968,6 @@ void UNiagaraScript::SetVMCompilationResults(const FNiagaraVMExecutableDataId& I
 
 	InvalidateExecutionReadyParameterStores();
 
-	// If we don't have a root system force optimization now as it may never get called otherwise
-	if (FindRootSystem() == nullptr)
-	{
-		HandleByteCodeOptimization(true);
-	}
-	
 	OnVMScriptCompiled().Broadcast(this, InCompileId.ScriptVersionID);
 }
 
@@ -3322,7 +3195,6 @@ void UNiagaraScript::RequestCompile(const FGuid& ScriptVersion, bool bForceCompi
 	if (ScriptData && (!AreScriptAndSourceSynchronized(ScriptVersion) || bForceCompile))
 	{
 		check(IsInGameThread());
-		CachedScriptVM.OptimizationTask.State = nullptr;
 
 		FNiagaraVMExecutableDataId& LastGeneratedVMId = GetLastGeneratedVMId(ScriptVersion);
 		if (IsCompilable() == false)
@@ -4172,7 +4044,6 @@ bool UNiagaraScript::SynchronizeExecutablesWithCompilation(const UNiagaraScript*
 
 		CachedScriptVM = Script->CachedScriptVM;
 		CachedScriptVMId = Script->CachedScriptVMId;
-		CachedScriptVM.OptimizationTask.State = nullptr;
 		CachedParameterCollectionReferences = Script->CachedParameterCollectionReferences;
 		CachedDefaultUObjects = Script->CachedDefaultUObjects;
 		InvalidateAndEmptyDefaultDataInterfaces(Script->CachedDefaultDataInterfaces.Num(), EDataInterfaceInvalidationMode::ProcessInvaliation);
@@ -4218,7 +4089,6 @@ void UNiagaraScript::InvalidateCompileResults(const FString& Reason)
 	CachedDefaultUObjects.Reset();
 	InvalidateAndEmptyDefaultDataInterfaces(0, EDataInterfaceInvalidationMode::QueueInvaliation);
 	ResolvedDataInterfaces.Reset();
-	CachedScriptVM.OptimizationTask.State = nullptr;
 }
 
 FText UNiagaraScript::GetDescription(const FGuid& VersionGuid)
diff --git a/Engine/Plugins/FX/Niagara/Source/Niagara/Private/NiagaraSystem.cpp b/Engine/Plugins/FX/Niagara/Source/Niagara/Private/NiagaraSystem.cpp
--- a/Engine/Plugins/FX/Niagara/Source/Niagara/Private/NiagaraSystem.cpp
+++ b/Engine/Plugins/FX/Niagara/Source/Niagara/Private/NiagaraSystem.cpp
@@ -3909,63 +3909,6 @@ void UNiagaraSystem::AddToInstanceCountStat(int32 NumInstances, bool bSolo)const
 #endif
 }
 
-void UNiagaraSystem::AsyncOptimizeAllScripts()
-{
-	check(IsInGameThread());
-
-	// Optimize is either in flight or done
-	if (bNeedsAsyncOptimize == false)
-	{
-		return;
-	}
-
-	if ( !IsReadyToRun() )
-	{
-		return;
-	}
-
-	FGraphEventArray Prereqs;
-	ForEachScript(
-		[&](UNiagaraScript* Script)
-		{
-			// Kick off the async optimize, which we'll wait on when the script is actually needed
-			if (Script != nullptr)
-			{
-				FGraphEventRef CompletionEvent = Script->HandleByteCodeOptimization(false);
-				if (CompletionEvent.IsValid())
-				{
-					Prereqs.Add(CompletionEvent);
-				}
-			}
-		}
-	);
-
-	if ( Prereqs.Num() > 0 )
-	{
-		DECLARE_CYCLE_STAT(TEXT("FNullGraphTask.NiagaraScriptOptimizationCompletion"), STAT_FNullGraphTask_NiagaraScriptOptimizationCompletion, STATGROUP_TaskGraphTasks);
-		ScriptOptimizationCompletionEvent = TGraphTask<FNullGraphTask>::CreateTask(&Prereqs, ENamedThreads::GameThread).ConstructAndDispatchWhenReady(GET_STATID(STAT_FNullGraphTask_NiagaraScriptOptimizationCompletion), ENamedThreads::AnyThread);
-
-		// Inject task to clear out the reference to the graph task
-		FGraphEventArray CompleteTaskPrereqs({ ScriptOptimizationCompletionEvent });
-		TGraphTask<FNiagaraSystemNullCompletionTask>::CreateTask(&CompleteTaskPrereqs, ENamedThreads::GameThread).ConstructAndDispatchWhenReady(this, &ScriptOptimizationCompletionEvent);
-	}
-
-	bNeedsAsyncOptimize = false;
-}
-
-FGraphEventRef UNiagaraSystem::GetScriptOptimizationCompletionEvent()
-{
-	if (ScriptOptimizationCompletionEvent.IsValid())
-	{
-		if (!ScriptOptimizationCompletionEvent->IsComplete())
-		{
-			return ScriptOptimizationCompletionEvent;
-		}
-		ScriptOptimizationCompletionEvent = nullptr;
-	}
-	return nullptr;
-}
-
 void UNiagaraSystem::UpdateStatID() const
 {
 #if STATS
diff --git a/Engine/Plugins/FX/Niagara/Source/Niagara/Private/NiagaraSystemMemReport.cpp b/Engine/Plugins/FX/Niagara/Source/Niagara/Private/NiagaraSystemMemReport.cpp
--- a/Engine/Plugins/FX/Niagara/Source/Niagara/Private/NiagaraSystemMemReport.cpp
+++ b/Engine/Plugins/FX/Niagara/Source/Niagara/Private/NiagaraSystemMemReport.cpp
@@ -203,7 +203,7 @@ uint64 FNiagaraSystemMemReport::GatherResourceMemory(EReportType ReportType, UOb
 		{
 			{
 				const FNiagaraVMExecutableData& VMData = NiagaraScript->GetVMExecutableData();
-				Nodes.Emplace(NAME_ByteCode, Depth + 1, VMData.ByteCode.GetLength() + VMData.OptimizedByteCode.GetLength());
+				Nodes.Emplace(NAME_ByteCode, Depth + 1, VMData.ByteCode.GetLength());
 				Nodes[NodeIndex].ExclusiveSizeBytes -= Nodes.Last().ExclusiveSizeBytes;
 			}
 
diff --git a/Engine/Plugins/FX/Niagara/Source/Niagara/Private/NiagaraSystemSimulation.cpp b/Engine/Plugins/FX/Niagara/Source/Niagara/Private/NiagaraSystemSimulation.cpp
--- a/Engine/Plugins/FX/Niagara/Source/Niagara/Private/NiagaraSystemSimulation.cpp
+++ b/Engine/Plugins/FX/Niagara/Source/Niagara/Private/NiagaraSystemSimulation.cpp
@@ -1360,14 +1360,7 @@ void FNiagaraSystemSimulation::Tick_GameThread_Internal(float DeltaSeconds, cons
 	FNiagaraSystemSimulationTickContext Context(this, SystemInstances, MainDataSet, DeltaSeconds, SpawnNum, MyCompletionGraphEvent.IsValid());
 	if ( Context.IsRunningAsync() )
 	{
-		FGraphEventArray Prereqs;
-		auto ScriptTask = System->GetScriptOptimizationCompletionEvent();
-		if (ScriptTask.IsValid())
-		{
-			Prereqs.Add(ScriptTask);
-		}
-		
-		auto ConcurrentTickTask = TGraphTask<FNiagaraSystemSimulationTickConcurrentTask>::CreateTask(&Prereqs, ENamedThreads::GameThread).ConstructAndHold(Context, AllWorkCompleteGraphEvent);
+		auto ConcurrentTickTask = TGraphTask<FNiagaraSystemSimulationTickConcurrentTask>::CreateTask(nullptr, ENamedThreads::GameThread).ConstructAndHold(Context, AllWorkCompleteGraphEvent);
 		ConcurrentTickGraphEvent = ConcurrentTickTask->GetCompletionEvent();
 		for (FNiagaraSystemInstance* Instance : Context.Instances)
 		{
@@ -1392,11 +1385,6 @@ void FNiagaraSystemSimulation::Tick_GameThread_Internal(float DeltaSeconds, cons
 	}
 	else
 	{
-		auto ScriptTask = System->GetScriptOptimizationCompletionEvent();
-		if (ScriptTask.IsValid())
-		{
-			ScriptTask->Wait(ENamedThreads::GameThread);
-		}
 		Tick_Concurrent(Context);
 	}
 }
@@ -1546,14 +1534,7 @@ void FNiagaraSystemSimulation::Spawn_GameThread(float DeltaSeconds, bool bPostAc
 		FNiagaraSystemSimulationTickContext Context(this, SpawningInstances, SpawningDataSet, DeltaSeconds, SpawningInstances.Num(), !bIsSolo);
 		if ( Context.IsRunningAsync() )
 		{
-			FGraphEventArray Prereqs;
-			auto ScriptTask = System->GetScriptOptimizationCompletionEvent();
-			if (ScriptTask.IsValid())
-			{
-				Prereqs.Add(ScriptTask);
-			}
-
-			auto ConcurrentTickTask = TGraphTask<FNiagaraSystemSimulationSpawnConcurrentTask>::CreateTask(&Prereqs, ENamedThreads::GameThread).ConstructAndHold(Context, AllWorkCompleteGraphEvent);
+			auto ConcurrentTickTask = TGraphTask<FNiagaraSystemSimulationSpawnConcurrentTask>::CreateTask(nullptr, ENamedThreads::GameThread).ConstructAndHold(Context, AllWorkCompleteGraphEvent);
 			ConcurrentTickGraphEvent = ConcurrentTickTask->GetCompletionEvent();
 			for (FNiagaraSystemInstance* Instance : Context.Instances)
 			{
@@ -1573,12 +1554,6 @@ void FNiagaraSystemSimulation::Spawn_GameThread(float DeltaSeconds, bool bPostAc
 		}
 		else
 		{
-			auto ScriptTask = System->GetScriptOptimizationCompletionEvent();
-			if (ScriptTask.IsValid())
-			{
-				ScriptTask->Wait(ENamedThreads::GameThread);
-			}
-
 			Spawn_Concurrent(Context);
 		}
 	}
