- Disable validation of objects being loaded inside UNiagaraSystem::PostLoad
- This was added to help track down a load issue, but that was fixed

#test Run ET
#rnx

diff --git a/Engine/Plugins/FX/Niagara/Source/Niagara/Private/NiagaraSystem.cpp b/Engine/Plugins/FX/Niagara/Source/Niagara/Private/NiagaraSystem.cpp
--- a/Engine/Plugins/FX/Niagara/Source/Niagara/Private/NiagaraSystem.cpp
+++ b/Engine/Plugins/FX/Niagara/Source/Niagara/Private/NiagaraSystem.cpp
@@ -200,7 +200,7 @@ static FAutoConsoleVariableRef CVarNiagaraOnDemandCompileEnabled(
 	ECVF_Default
 );
 
-static int GNiagaraObjectNeedsLoadMode = 1;
+static int GNiagaraObjectNeedsLoadMode = 0;
 static FAutoConsoleVariableRef CVarNiagaraObjectNeedsLoadMode(
 	TEXT("fx.Niagara.ObjectNeedsLoadMode"),
 	GNiagaraObjectNeedsLoadMode,
