- Enable Sub UV cutouts ES3.1

#tests Use Sub UV cutouts to see they match in ES3.1 preview and not

diff --git a/Engine/Plugins/FX/Niagara/Shaders/Private/NiagaraSpriteVertexFactory.ush b/Engine/Plugins/FX/Niagara/Shaders/Private/NiagaraSpriteVertexFactory.ush
--- a/Engine/Plugins/FX/Niagara/Shaders/Private/NiagaraSpriteVertexFactory.ush
+++ b/Engine/Plugins/FX/Niagara/Shaders/Private/NiagaraSpriteVertexFactory.ush
@@ -672,7 +672,7 @@ float3x3 CalcTangentBasis(FVertexFactoryIntermediates Intermediates)
 
 void ComputeBillboardUVs(FVertexFactoryInput Input, float2 ParticleSize, float SubImageIndex, float2 UVScale, out float2 UVForPosition, out float2 UVForTexturing, out float2 UVForTexturingUnflipped)
 {
-#if FEATURE_LEVEL >= FEATURE_LEVEL_SM4
+#if FEATURE_LEVEL >= FEATURE_LEVEL_ES3_1
 	const uint NumCutOutVertices = GetCutoutNumVertices();
 
 	BRANCH
