- When calling GetWorldSettings set bChecked to false
- This can happen if the current persistent level is nullptr

#test Ran ET / test project
#rnx

diff --git a/Engine/Plugins/FX/Niagara/Source/Niagara/Private/NiagaraComponentPool.cpp b/Engine/Plugins/FX/Niagara/Source/Niagara/Private/NiagaraComponentPool.cpp
--- a/Engine/Plugins/FX/Niagara/Source/Niagara/Private/NiagaraComponentPool.cpp
+++ b/Engine/Plugins/FX/Niagara/Source/Niagara/Private/NiagaraComponentPool.cpp
@@ -244,7 +244,7 @@ void UNiagaraComponentPool::Cleanup(UWorld* World)
 	// This is generally used on world cleanup to ensure all pooled components are cleaned up as they can not be returned back to the component pool.
 	if ( World != nullptr )
 	{
-		if (AWorldSettings* WorldSettings = World->GetWorldSettings())
+		if (AWorldSettings* WorldSettings = World->GetWorldSettings(false, false))
 		{
 			TArray<UNiagaraComponent*> ActiveComponents;
 			WorldSettings->GetComponents(ActiveComponents);
diff --git a/Engine/Source/Runtime/Engine/Private/Particles/WorldPSCPool.cpp b/Engine/Source/Runtime/Engine/Private/Particles/WorldPSCPool.cpp
--- a/Engine/Source/Runtime/Engine/Private/Particles/WorldPSCPool.cpp
+++ b/Engine/Source/Runtime/Engine/Private/Particles/WorldPSCPool.cpp
@@ -234,7 +234,7 @@ void FWorldPSCPool::Cleanup(UWorld* World)
 	// This is generally used on world cleanup to ensure all pooled components are cleaned up as they can not be returned back to the component pool.
 	if (World != nullptr)
 	{
-		if (AWorldSettings* WorldSettings = World->GetWorldSettings())
+		if (AWorldSettings* WorldSettings = World->GetWorldSettings(false, false))
 		{
 			TArray<UParticleSystemComponent*> ActiveComponents;
 			WorldSettings->GetComponents(ActiveComponents);
