- Add exposure option to stateless templates

diff --git a/Engine/Plugins/FX/Niagara/Source/Niagara/Internal/Stateless/NiagaraStatelessEmitterTemplate.h b/Engine/Plugins/FX/Niagara/Source/Niagara/Internal/Stateless/NiagaraStatelessEmitterTemplate.h
--- a/Engine/Plugins/FX/Niagara/Source/Niagara/Internal/Stateless/NiagaraStatelessEmitterTemplate.h
+++ b/Engine/Plugins/FX/Niagara/Source/Niagara/Internal/Stateless/NiagaraStatelessEmitterTemplate.h
@@ -10,6 +10,9 @@
 
 #include "NiagaraStatelessEmitterTemplate.generated.h"
 
+struct FAssetData;
+class FAssetRegistryTagsContext;
+
 UCLASS(MinimalAPI)
 class UNiagaraStatelessEmitterTemplate : public UObject
 {
@@ -19,12 +22,16 @@ public:
 	//~Begin: UObject Interface
 	NIAGARA_API virtual void PostInitProperties() override;
 	NIAGARA_API virtual void PostLoad() override;
+	NIAGARA_API virtual void GetAssetRegistryTags(FAssetRegistryTagsContext Context) const override;
 	//~End: UObject Interface
 
 	static void InitCDOPropertiesAfterModuleStartup();
 	virtual void InitModulesAndAttributes();
 #if WITH_EDITOR
 	NIAGARA_API void ModifyTemplate(UNiagaraStatelessEmitterTemplate* FromTemplate);
+
+	NIAGARA_API static bool IsExposedToLibrary(const FAssetData& AssetData);
+	bool IsExposedToLibrary() const { return bExposedToLibrary; }
 #endif
 
 	TConstArrayView<TObjectPtr<UClass>> GetModules() const { return Modules; }
@@ -41,6 +48,9 @@ protected:
 	TArray<TObjectPtr<UClass>> Modules;
 
 #if WITH_EDITORONLY_DATA
+	UPROPERTY(EditAnywhere, Category="Modules")
+	bool bExposedToLibrary = true;
+
 	TArray<FNiagaraVariableBase> OutputComponents;
 #endif
 	TSharedPtr<FShaderParametersMetadata> ShaderParametersMetadata;
diff --git a/Engine/Plugins/FX/Niagara/Source/Niagara/Private/Stateless/NiagaraStatelessEmitterTemplate.cpp b/Engine/Plugins/FX/Niagara/Source/Niagara/Private/Stateless/NiagaraStatelessEmitterTemplate.cpp
--- a/Engine/Plugins/FX/Niagara/Source/Niagara/Private/Stateless/NiagaraStatelessEmitterTemplate.cpp
+++ b/Engine/Plugins/FX/Niagara/Source/Niagara/Private/Stateless/NiagaraStatelessEmitterTemplate.cpp
@@ -38,13 +38,18 @@
 #include "Stateless/Modules/NiagaraStatelessModule_SubUVAnimation.h"
 #include "Stateless/Modules/NiagaraStatelessModule_DynamicMaterialParameters.h"
 
+#include "UObject/AssetRegistryTagsContext.h"
+
 #if WITH_EDITOR
+#include "AssetRegistry/AssetData.h"
 #include "UObject/UObjectIterator.h"
 #include "Engine/Engine.h"
 #endif
 
-namespace NiagaraStatelessEmitterTemplateImpl
+namespace NiagaraStatelessEmitterTemplatePrivate
 {
+	static const FName NAME_ExposedToLibrary("bExposedToLibrary");
+
 	TArray<TWeakObjectPtr<UNiagaraStatelessEmitterTemplate>> ObjectsToDeferredInit;
 }
 
@@ -59,7 +64,7 @@ void UNiagaraStatelessEmitterTemplate::PostInitProperties()
 	{
 		if (FModuleManager::Get().IsModuleLoaded("Niagara") == false)
 		{
-			NiagaraStatelessEmitterTemplateImpl::ObjectsToDeferredInit.Add(this);
+			NiagaraStatelessEmitterTemplatePrivate::ObjectsToDeferredInit.Add(this);
 		}
 		else
 		{
@@ -75,9 +80,19 @@ void UNiagaraStatelessEmitterTemplate::PostLoad()
 	InitModulesAndAttributes();
 }
 
+void UNiagaraStatelessEmitterTemplate::GetAssetRegistryTags(FAssetRegistryTagsContext Context) const
+{
+	using namespace NiagaraStatelessEmitterTemplatePrivate;
+
+	Super::GetAssetRegistryTags(Context);
+#if WITH_EDITORONLY_DATA
+	Context.AddTag(FAssetRegistryTag(NAME_ExposedToLibrary, LexToString(bExposedToLibrary), FAssetRegistryTag::TT_Hidden));
+#endif
+}
+
 void UNiagaraStatelessEmitterTemplate::InitCDOPropertiesAfterModuleStartup()
 {
-	for (TWeakObjectPtr<UNiagaraStatelessEmitterTemplate>& WeakObject : NiagaraStatelessEmitterTemplateImpl::ObjectsToDeferredInit)
+	for (TWeakObjectPtr<UNiagaraStatelessEmitterTemplate>& WeakObject : NiagaraStatelessEmitterTemplatePrivate::ObjectsToDeferredInit)
 	{
 		if (UNiagaraStatelessEmitterTemplate* EmitterObject = WeakObject.Get())
 		{
@@ -85,7 +100,7 @@ void UNiagaraStatelessEmitterTemplate::InitCDOPropertiesAfterModuleStartup()
 			checkf(EmitterObject->Modules.Num() != 0, TEXT("StatelessTemplate(%s) has no modules, this is likely a code error"), *GetNameSafe(EmitterObject));
 		}
 	}
-	NiagaraStatelessEmitterTemplateImpl::ObjectsToDeferredInit.Empty();
+	NiagaraStatelessEmitterTemplatePrivate::ObjectsToDeferredInit.Empty();
 }
 
 void UNiagaraStatelessEmitterTemplate::InitModulesAndAttributes()
@@ -142,8 +157,12 @@ void UNiagaraStatelessEmitterTemplate::ModifyTemplate(UNiagaraStatelessEmitterTe
 		}
 	}
 
-	// Copy and initialize the template
-	UEngine::CopyPropertiesForUnrelatedObjects(FromTemplate, this);
+	// Copy template over
+	UEngine::FCopyPropertiesForUnrelatedObjectsParams CopyParams;
+	CopyParams.bDoDelta = false;
+	UEngine::CopyPropertiesForUnrelatedObjects(FromTemplate, this, CopyParams);
+
+	// Initialize template
 	InitModulesAndAttributes();
 
 	for (UNiagaraStatelessEmitter* Emitter : EmittersToUpdate)
@@ -152,6 +171,15 @@ void UNiagaraStatelessEmitterTemplate::ModifyTemplate(UNiagaraStatelessEmitterTe
 		Emitter->CacheFromCompiledData();
 	}
 }
+
+bool UNiagaraStatelessEmitterTemplate::IsExposedToLibrary(const FAssetData& AssetData)
+{
+	using namespace NiagaraStatelessEmitterTemplatePrivate;
+
+	bool bExposedToLibrary = false;
+	AssetData.GetTagValue(NAME_ExposedToLibrary, bExposedToLibrary);
+	return bExposedToLibrary;
+}
 #endif
 
 void UNiagaraStatelessEmitterTemplate::SetShaderParameters(uint8* ShaderParametersBase, TConstArrayView<int32> ComponentOffsets) const
diff --git a/Engine/Plugins/FX/Niagara/Source/NiagaraEditor/Private/Stateless/NiagaraStatelessEmitterDetails.cpp b/Engine/Plugins/FX/Niagara/Source/NiagaraEditor/Private/Stateless/NiagaraStatelessEmitterDetails.cpp
--- a/Engine/Plugins/FX/Niagara/Source/NiagaraEditor/Private/Stateless/NiagaraStatelessEmitterDetails.cpp
+++ b/Engine/Plugins/FX/Niagara/Source/NiagaraEditor/Private/Stateless/NiagaraStatelessEmitterDetails.cpp
@@ -176,6 +176,11 @@ TSharedRef<SWidget> FNiagaraStatelessEmitterDetails::OnGetTemplateActions() cons
 
 		for (const FAssetData& AssetData : TemplateAssetDatas)
 		{
+			if (UNiagaraStatelessEmitterTemplate::IsExposedToLibrary(AssetData) == false)
+			{
+				continue;
+			}
+
 			const FSoftObjectPath ObjectPath = AssetData.GetSoftObjectPath();
 			const FText DisplayName = FText::FromString(ObjectPath.GetAssetName());
 			const FText AssetName = FText::FromString(ObjectPath.ToString());
