- Disable delta serialize when resetting / pasting stateless modules

#rnx

diff --git a/Engine/Plugins/FX/Niagara/Source/NiagaraEditor/Private/ViewModels/Stack/NiagaraStackStatelessEmitterSimulateGroup.cpp b/Engine/Plugins/FX/Niagara/Source/NiagaraEditor/Private/ViewModels/Stack/NiagaraStackStatelessEmitterSimulateGroup.cpp
--- a/Engine/Plugins/FX/Niagara/Source/NiagaraEditor/Private/ViewModels/Stack/NiagaraStackStatelessEmitterSimulateGroup.cpp
+++ b/Engine/Plugins/FX/Niagara/Source/NiagaraEditor/Private/ViewModels/Stack/NiagaraStackStatelessEmitterSimulateGroup.cpp
@@ -65,7 +65,11 @@ namespace NiagaraStackStatelessEmitterSimulateGroupPrivate
 				}
 
 				StatelessModule->Modify();
-				UEngine::CopyPropertiesForUnrelatedObjects(const_cast<UNiagaraStatelessModule*>(ClipboardModule), StatelessModule);
+
+				UEngine::FCopyPropertiesForUnrelatedObjectsParams CopyParams;
+				CopyParams.bDoDelta = false;
+				UEngine::CopyPropertiesForUnrelatedObjects(const_cast<UNiagaraStatelessModule*>(ClipboardModule), StatelessModule, CopyParams);
+
 				if (StatelessModule->CanDisableModule())
 				{
 					StatelessModule->SetIsModuleEnabled(ClipboardModule->IsModuleEnabled());
@@ -345,7 +349,11 @@ void UNiagaraStackStatelessModuleItem::Delete()
 	if (StatelessModule != nullptr && StatelessModule->CanDisableModule())
 	{
 		StatelessModule->Modify();
-		UEngine::CopyPropertiesForUnrelatedObjects(StatelessModule->GetClass()->GetDefaultObject(), StatelessModule);
+
+		UEngine::FCopyPropertiesForUnrelatedObjectsParams CopyParams;
+		CopyParams.bDoDelta = false;
+		UEngine::CopyPropertiesForUnrelatedObjects(StatelessModule->GetClass()->GetDefaultObject(), StatelessModule, CopyParams);
+
 		StatelessModule->SetIsModuleEnabled(false);
 		StatelessModule->PostEditChange();
 		OnDataObjectModified().Broadcast({ StatelessModule }, ENiagaraDataObjectChange::Changed);
