- Add missing init for StatDatabase

#tests Add emitters, delete emitters, undo / redo, using both unversioned and versioned emtiters (with multiple versions)
#rnx

diff --git a/Engine/Plugins/FX/Niagara/Source/Niagara/Private/NiagaraEmitter.cpp b/Engine/Plugins/FX/Niagara/Source/Niagara/Private/NiagaraEmitter.cpp
--- a/Engine/Plugins/FX/Niagara/Source/Niagara/Private/NiagaraEmitter.cpp
+++ b/Engine/Plugins/FX/Niagara/Source/Niagara/Private/NiagaraEmitter.cpp
@@ -2318,6 +2318,7 @@ void FVersionedNiagaraEmitterData::OnPostCompile(const UNiagaraEmitter& InEmitte
 	RuntimeEstimation.AllocationEstimate = 0;
 	RuntimeEstimation.IsEstimationDirty = false;
 #if STATS
+	StatDatabase.Init();
 	StatDatabase.ClearStatCaptures();
 #endif
 }
