- Handle nullptr for data interfaces inside UNiagaraScript::InvalidateAndEmptyDefaultDataInterfaces
- This could happen if we didn't process invalidations before a GC occurs

#tests Run ET
#rnx

diff --git a/Engine/Plugins/FX/Niagara/Source/Niagara/Private/NiagaraScript.cpp b/Engine/Plugins/FX/Niagara/Source/Niagara/Private/NiagaraScript.cpp
--- a/Engine/Plugins/FX/Niagara/Source/Niagara/Private/NiagaraScript.cpp
+++ b/Engine/Plugins/FX/Niagara/Source/Niagara/Private/NiagaraScript.cpp
@@ -1512,6 +1512,12 @@ void UNiagaraScript::InvalidateAndEmptyDefaultDataInterfaces(int32 SlackCount, E
 
 		for (UNiagaraDataInterface* DataInterface : DataInterfacesPendingInvalidation)
 		{
+			// The data interface could have been GCed before we process invalidations
+			if (DataInterface == nullptr)
+			{
+				continue;
+			}
+
 			const FName OrigName = DataInterface->GetFName();
 
 			FNameBuilder InvalidatedName;
