- Use Position as the default for Previous.Position when solving velocities
- The attribute might not be part of the payload if we using approximate motion vector calculations
- While this can result in potentially missing velcoity data, if we had other modules which modify position for example, it's close enough

#tests Make an LWE that uses velocity alignment, disable accurate moition vectors and force the CPU simulation path
#rnx

diff --git a/Engine/Plugins/FX/Niagara/Source/Niagara/Private/Stateless/Modules/NiagaraStatelessModule_SolveVelocitiesAndForces.cpp b/Engine/Plugins/FX/Niagara/Source/Niagara/Private/Stateless/Modules/NiagaraStatelessModule_SolveVelocitiesAndForces.cpp
--- a/Engine/Plugins/FX/Niagara/Source/Niagara/Private/Stateless/Modules/NiagaraStatelessModule_SolveVelocitiesAndForces.cpp
+++ b/Engine/Plugins/FX/Niagara/Source/Niagara/Private/Stateless/Modules/NiagaraStatelessModule_SolveVelocitiesAndForces.cpp
@@ -99,7 +99,7 @@ namespace NiagaraStatelessModuleSolveVelocitiesAndForcesPrivate
 			Acceleration			+= ParticleSimulationContext.RandomScaleBiasFloat(i, 41, Parameters->SolveVelocitiesAndForces_GravityScale, Parameters->SolveVelocitiesAndForces_GravityBias);
 
 			FVector3f Position			= ParticleSimulationContext.ReadParticleVariable(ModuleBuiltData->PositionVariableOffset, i, FVector3f::ZeroVector);
-			FVector3f PreviousPosition	= ParticleSimulationContext.ReadParticleVariable(ModuleBuiltData->PreviousPositionVariableOffset, i, FVector3f::ZeroVector);
+			FVector3f PreviousPosition	= ParticleSimulationContext.ReadParticleVariable(ModuleBuiltData->PreviousPositionVariableOffset, i, Position);
 
 			if (Parameters->SolveVelocitiesAndForces_ConeVelocityEnabled != 0)
 			{
