Fix issue with incorrect asset /pool size settings being used in Naigara Component Pool.

diff --git a/Engine/Plugins/FX/Niagara/Source/Niagara/Private/NiagaraComponentPool.cpp b/Engine/Plugins/FX/Niagara/Source/Niagara/Private/NiagaraComponentPool.cpp
--- a/Engine/Plugins/FX/Niagara/Source/Niagara/Private/NiagaraComponentPool.cpp
+++ b/Engine/Plugins/FX/Niagara/Source/Niagara/Private/NiagaraComponentPool.cpp
@@ -458,7 +458,7 @@ void UNiagaraComponentPool::ReclaimWorldParticleSystem(UNiagaraComponent* Compon
 			LastParticleSytemPoolCleanTime = CurrentTime;
 			for (auto& Pair : WorldParticleSystemPools)
 			{
-				Pair.Value.KillUnusedComponents(CurrentTime - GNiagaraSystemPoolKillUnusedTime, Asset);
+				Pair.Value.KillUnusedComponents(CurrentTime - GNiagaraSystemPoolKillUnusedTime, Pair.Key);
 			}
 		}
 
