Ensure handler systems are active when items are being written to the NDC.
Handles any edge cases where an NDC handler has been culled but things are still expecting it to be

diff --git a/Engine/Plugins/FX/Niagara/Source/Niagara/Private/NiagaraDataChannel_Islands.cpp b/Engine/Plugins/FX/Niagara/Source/Niagara/Private/NiagaraDataChannel_Islands.cpp
--- a/Engine/Plugins/FX/Niagara/Source/Niagara/Private/NiagaraDataChannel_Islands.cpp
+++ b/Engine/Plugins/FX/Niagara/Source/Niagara/Private/NiagaraDataChannel_Islands.cpp
@@ -405,7 +405,7 @@ void FNDCIsland::EndFrame()
 void FNDCIsland::Tick(const ETickingGroup& TickGroup)
 {
 	int32 AddedData = Data->ConsumePublishRequests(Owner, TickGroup);
-	if (IsBeingUsed() && AddedData > 0)
+	if (AddedData > 0)
 	{
 		for (UNiagaraComponent* Comp : SpawnedComponents)
 		{
