PR #9887: Ensure Types used in Data Interface Function Signatures are defined in HLSL

#rnx

diff --git a/Engine/Plugins/FX/Niagara/Source/NiagaraEditor/Private/NiagaraHlslTranslator.cpp b/Engine/Plugins/FX/Niagara/Source/NiagaraEditor/Private/NiagaraHlslTranslator.cpp
--- a/Engine/Plugins/FX/Niagara/Source/NiagaraEditor/Private/NiagaraHlslTranslator.cpp
+++ b/Engine/Plugins/FX/Niagara/Source/NiagaraEditor/Private/NiagaraHlslTranslator.cpp
@@ -2241,6 +2241,37 @@ FNiagaraTranslateResults TNiagaraHlslTranslator<GraphBridge>::Translate(const FN
 		TArray<FNiagaraVariable> PrimaryDataSetOutputEntries;
 		FString ParameterMapDefinitionStr = BuildParameterMapHlslDefinitions(PrimaryDataSetOutputEntries);
 
+		// Traverse all the used data interfaces and generate any used structs
+		// This step is for template files that might use structs which need defining via the translator
+		{
+			TSet<FName> InterfaceClasses;
+			for (int32 i = 0; i < CompilationOutput.ScriptData.DataInterfaceInfo.Num(); i++)
+			{
+				FNiagaraScriptDataInterfaceCompileInfo& Info = CompilationOutput.ScriptData.DataInterfaceInfo[i];
+				UNiagaraDataInterface* CDO = CompileDuplicateData->GetDuplicatedDataInterfaceCDOForClass(Info.Type.GetClass());
+				check(CDO != nullptr);
+				if (CDO && CDO->CanExecuteOnTarget(ENiagaraSimTarget::GPUComputeSim))
+				{
+					if (!InterfaceClasses.Contains(Info.Type.GetFName()))
+					{
+						TArray<FNiagaraFunctionSignature> FunctionSigs;
+						CDO->GetFunctionSignatures(FunctionSigs);
+						for (const FNiagaraFunctionSignature& FunctionSig : FunctionSigs)
+						{
+							for (const FNiagaraVariable& Input : FunctionSig.Inputs)
+							{
+								AddStructToDefinitionSet(Input.GetType());
+							}
+							for (const FNiagaraVariableBase& Output : FunctionSig.Outputs)
+							{
+								AddStructToDefinitionSet(Output.GetType());
+							}
+						}
+					}
+				}
+			}
+		}
+
 		// Ensure some structures are always added as we use them in custom HLSL / data interface templates
 		// Remove some structures which we define inside NiagaraEmitterInstanceShader.usf as we want a common set of functions
 		// Ensure we always add structures that are fundamental to custom HLSL or data interface templates shader files
