Track RT lag to prevent RT hitches from impacting rendering based visibility culling.

diff --git a/Engine/Plugins/FX/Niagara/Source/Niagara/Private/NiagaraWorldManager.cpp b/Engine/Plugins/FX/Niagara/Source/Niagara/Private/NiagaraWorldManager.cpp
--- a/Engine/Plugins/FX/Niagara/Source/Niagara/Private/NiagaraWorldManager.cpp
+++ b/Engine/Plugins/FX/Niagara/Source/Niagara/Private/NiagaraWorldManager.cpp
@@ -388,6 +388,9 @@ void FNiagaraWorldManager::Init(UWorld* InWorld)
 
 	World = InWorld;
 
+	LastRenderThreadUpdateTime = MakeShared<std::atomic<float>>();
+	LastRenderThreadUpdateTime->store(InWorld->GetTimeSeconds());
+
 	ComponentPool = NewObject<UNiagaraComponentPool>();
 
 	for (int32 TickGroup = 0; TickGroup < NiagaraNumTickGroups; ++TickGroup)
@@ -1123,6 +1126,20 @@ void FNiagaraWorldManager::TickStart(float DeltaSeconds)
 	LLM_SCOPE(ELLMTag::Niagara);
 	SCOPE_CYCLE_COUNTER(STAT_NiagaraWorldManTick);
 
+	//Have the render thread signal the most recent world time it has consumed.
+	//This will cause a several frame delay in render based visibility culling.
+	//However it will make the culling less sensitive to variable RT times and hitches etc
+	//that would otherwise cause systems to be culled in error.
+	check(LastRenderThreadUpdateTime.IsValid());
+	float WorldTime = GetWorld()->GetTimeSeconds();
+	RenderThreadWorldTimeLag = WorldTime - LastRenderThreadUpdateTime->load(std::memory_order_relaxed);
+	ENQUEUE_RENDER_COMMAND(FNiagaraWorldManagerRTUpdateTime)(
+		[LRT = LastRenderThreadUpdateTime, WorldTime](FRHICommandListImmediate& RHICmdList)
+		{
+			LRT->store(WorldTime, std::memory_order_relaxed);
+		}
+	);
+
 	DataChannelManager->BeginFrame(DeltaSeconds);
 
 #if WITH_NIAGARA_LEAK_DETECTOR
@@ -1906,7 +1923,7 @@ void FNiagaraWorldManager::CalculateScalabilityState(UNiagaraSystem* System, con
 	{
 		FBoxSphereBounds Bounds = Component->CalcBounds(Component->GetComponentToWorld());		
 
-		const float TimeSinceRendered = float(FMath::Max(0.0f, GetWorld()->LastRenderTime - Component->GetLastRenderTime() - World->GetDeltaSeconds() - GLastRenderTimeSafetyBias));
+		const float TimeSinceRendered = float(FMath::Max(0.0f, GetWorld()->LastRenderTime - Component->GetLastRenderTime() - World->GetDeltaSeconds() - RenderThreadWorldTimeLag - GLastRenderTimeSafetyBias));
 
 		if(bIsPreCull)
 		{
diff --git a/Engine/Plugins/FX/Niagara/Source/Niagara/Public/NiagaraWorldManager.h b/Engine/Plugins/FX/Niagara/Source/Niagara/Public/NiagaraWorldManager.h
--- a/Engine/Plugins/FX/Niagara/Source/Niagara/Public/NiagaraWorldManager.h
+++ b/Engine/Plugins/FX/Niagara/Source/Niagara/Public/NiagaraWorldManager.h
@@ -474,6 +474,9 @@ private:
 	using FSystemDelegateMap = TMap<TWeakObjectPtr<UNiagaraSystem>, FDelegateHandle>;
 	FSystemDelegateMap SystemPostChangeDelegateHandles;
 #endif
+
+	TSharedPtr<std::atomic<float>> LastRenderThreadUpdateTime;
+	float RenderThreadWorldTimeLag = 0.0f;
 };
 
 
