- When modifying parameters in the detail panel use editor setting GetResetSimulationOnParameterChange to see if we should reset the system or not
- Disabling reset can be useful to iterate on content live

diff --git a/Engine/Plugins/FX/Niagara/Source/Niagara/Private/NiagaraComponent.cpp b/Engine/Plugins/FX/Niagara/Source/Niagara/Private/NiagaraComponent.cpp
--- a/Engine/Plugins/FX/Niagara/Source/Niagara/Private/NiagaraComponent.cpp
+++ b/Engine/Plugins/FX/Niagara/Source/Niagara/Private/NiagaraComponent.cpp
@@ -34,6 +34,12 @@
 #include "PrimitiveUniformShaderParametersBuilder.h"
 #include "NiagaraActor.h"
 
+#if WITH_EDITOR
+#include "INiagaraEditorOnlyDataUtlities.h"
+#include "NiagaraModule.h"
+#include "Modules/ModuleManager.h"
+#endif
+
 #include UE_INLINE_GENERATED_CPP_BY_NAME(NiagaraComponent)
 
 DECLARE_CYCLE_STAT(TEXT("Sceneproxy create (GT)"), STAT_NiagaraCreateSceneProxy, STATGROUP_Niagara);
@@ -3243,6 +3249,7 @@ void UNiagaraComponent::PostEditChangeProperty(FPropertyChangedEvent& PropertyCh
 		PropertyName = PropertyChangedEvent.Property->GetFName();
 	}
 
+	bool bNeedsReinitialize = true;
 	if (PropertyName == GET_MEMBER_NAME_CHECKED(UNiagaraComponent, Asset))
 	{
 		SynchronizeWithSourceSystem();
@@ -3255,14 +3262,23 @@ void UNiagaraComponent::PostEditChangeProperty(FPropertyChangedEvent& PropertyCh
 	else if(PropertyName == GET_MEMBER_NAME_CHECKED(UNiagaraComponent, OverrideParameters))
 	{
 		SynchronizeWithSourceSystem();
+
+		INiagaraModule& NiagaraModule = FModuleManager::GetModuleChecked<INiagaraModule>("Niagara");
+		bNeedsReinitialize = NiagaraModule.GetEditorOnlyDataUtilities().GetResetSimulationOnParameterChange();
 	}
 	else if (PropertyName == GET_MEMBER_NAME_CHECKED(UNiagaraComponent, TemplateParameterOverrides) ||
 		PropertyName == GET_MEMBER_NAME_CHECKED(UNiagaraComponent, InstanceParameterOverrides))
 	{
 		ApplyOverridesToParameterStore();
+
+		INiagaraModule& NiagaraModule = FModuleManager::GetModuleChecked<INiagaraModule>("Niagara");
+		bNeedsReinitialize = NiagaraModule.GetEditorOnlyDataUtilities().GetResetSimulationOnParameterChange();
 	}
 	
-	ReinitializeSystem();
+	if (bNeedsReinitialize)
+	{
+		ReinitializeSystem();
+	}
 	
 	Super::PostEditChangeProperty(PropertyChangedEvent);
 }
diff --git a/Engine/Plugins/FX/Niagara/Source/Niagara/Public/INiagaraEditorOnlyDataUtlities.h b/Engine/Plugins/FX/Niagara/Source/Niagara/Public/INiagaraEditorOnlyDataUtlities.h
--- a/Engine/Plugins/FX/Niagara/Source/Niagara/Public/INiagaraEditorOnlyDataUtlities.h
+++ b/Engine/Plugins/FX/Niagara/Source/Niagara/Public/INiagaraEditorOnlyDataUtlities.h
@@ -33,4 +33,6 @@ public:
 	virtual UNiagaraDataInterface* GetResolvedRuntimeInstanceForEditorDataInterfaceInstance(const UNiagaraSystem& OwningSystem, UNiagaraDataInterface& EditorDataInterfaceInstance) const = 0;
 
 	virtual TOptional<FNiagaraSystemStateData> GetSystemStateData(const UNiagaraSystem& System) const = 0;
+
+	virtual bool GetResetSimulationOnParameterChange() const = 0;
 };
diff --git a/Engine/Plugins/FX/Niagara/Source/NiagaraEditor/Private/NiagaraEditorModule.cpp b/Engine/Plugins/FX/Niagara/Source/NiagaraEditor/Private/NiagaraEditorModule.cpp
--- a/Engine/Plugins/FX/Niagara/Source/NiagaraEditor/Private/NiagaraEditorModule.cpp
+++ b/Engine/Plugins/FX/Niagara/Source/NiagaraEditor/Private/NiagaraEditorModule.cpp
@@ -367,6 +367,12 @@ class FNiagaraEditorOnlyDataUtilities : public INiagaraEditorOnlyDataUtilities
 
 		return SystemStateData;
 	}
+
+	virtual bool GetResetSimulationOnParameterChange() const override
+	{
+		const UNiagaraEditorSettings* NiagaraEditorSettings = GetDefault<UNiagaraEditorSettings>();
+		return NiagaraEditorSettings->GetResetSimulationOnChange();
+	}
 };
 
 class FNiagaraScriptGraphPanelPinFactory : public FGraphPanelPinFactory
