- Skip any signatures that do not run on the GPU during HLSL translation

#rnx

diff --git a/Engine/Plugins/FX/Niagara/Source/NiagaraEditor/Private/NiagaraHlslTranslator.cpp b/Engine/Plugins/FX/Niagara/Source/NiagaraEditor/Private/NiagaraHlslTranslator.cpp
--- a/Engine/Plugins/FX/Niagara/Source/NiagaraEditor/Private/NiagaraHlslTranslator.cpp
+++ b/Engine/Plugins/FX/Niagara/Source/NiagaraEditor/Private/NiagaraHlslTranslator.cpp
@@ -2258,6 +2258,11 @@ FNiagaraTranslateResults TNiagaraHlslTranslator<GraphBridge>::Translate(const FN
 						CDO->GetFunctionSignatures(FunctionSigs);
 						for (const FNiagaraFunctionSignature& FunctionSig : FunctionSigs)
 						{
+							if (FunctionSig.bSupportsGPU == false)
+							{
+								continue;
+							}
+
 							for (const FNiagaraVariable& Input : FunctionSig.Inputs)
 							{
 								AddStructToDefinitionSet(Input.GetType());
