- Adjust code to be more readable / silence PVS warning

#rnx

diff --git a/Engine/Plugins/FX/Niagara/Source/Niagara/Private/NiagaraFunctionLibrary.cpp b/Engine/Plugins/FX/Niagara/Source/Niagara/Private/NiagaraFunctionLibrary.cpp
--- a/Engine/Plugins/FX/Niagara/Source/Niagara/Private/NiagaraFunctionLibrary.cpp
+++ b/Engine/Plugins/FX/Niagara/Source/Niagara/Private/NiagaraFunctionLibrary.cpp
@@ -1092,9 +1092,8 @@ struct FVectorKernelFastMatrixToQuaternion
 					if (Mat.M[2][2] > Mat.M[i][i])
 						i = 2;
 
-					static const int32 nxt[3] = { 1, 2, 0 };
-					const int32 j = nxt[i];
-					const int32 k = nxt[j];
+					const int32 j = (i + 1) % 3;
+					const int32 k = (j + 1) % 3;
 
 					float s = Mat.M[i][i] - Mat.M[j][j] - Mat.M[k][k] + 1.0f;
 
