- Weighted Distribution Int Array
- Each entry has a valid & weight
- The Weight is used in the random calculation, so the higher the weight the more likely it is to return that value

#virtualized

diff --git a/Engine/Plugins/FX/Niagara/Shaders/Private/NiagaraDataInterfaceArrayDistributionIntTemplate.ush b/Engine/Plugins/FX/Niagara/Shaders/Private/NiagaraDataInterfaceArrayDistributionIntTemplate.ush
new file mode 100644
--- /dev/null
+++ b/Engine/Plugins/FX/Niagara/Shaders/Private/NiagaraDataInterfaceArrayDistributionIntTemplate.ush
@@ -0,0 +1,40 @@
+// Copyright Epic Games, Inc. All Rights Reserved.
+
+int					{ParameterName}_Length;
+ByteAddressBuffer	{ParameterName}_BuiltTableData;
+
+void Length_{ParameterName}(out int OutLength)
+{
+	OutLength = {ParameterName}_Length;
+}
+
+void IsValidIndex_{ParameterName}(in int InIndex, out bool OutIsValid)
+{
+	OutIsValid = InIndex >= 0 && InIndex < {ParameterName}_Length;
+}
+
+void LastIndex_{ParameterName}(out int OutLastIndex)
+{
+	OutLastIndex = {ParameterName}_Length - 1;
+}
+
+void Get_{ParameterName}(in int InIndex, out NDIDistributionIntArrayEntry OutEntry)
+{
+	const int ClampedIndex = clamp(InIndex, 0, {ParameterName}_Length - 1);
+	const int Offset = ClampedIndex * 16; 
+	const uint2 Data = {ParameterName}_BuiltTableData.Load2(Offset + 8);
+	OutEntry.Value = asint(Data.x);
+	OutEntry.Weight = asfloat(Data.y);
+}
+
+
+void GetRandomValue_{ParameterName}(NiagaraRandInfo RandInfo, out int OutValue)
+{
+	int Index = NiagaraRandomInt(RandInfo, {ParameterName}_Length);
+	const uint2 ProbAlias = {ParameterName}_BuiltTableData.Load2(Index * 16);
+
+	const float RandT0 = NiagaraRandomFloat(RandInfo);
+	Index = RandT0 > asfloat(ProbAlias.x) ? asint(ProbAlias.y) : Index;
+
+	OutValue = {ParameterName}_BuiltTableData.Load((Index * 16) + 8);
+}
diff --git a/Engine/Plugins/FX/Niagara/Source/Niagara/Classes/NiagaraDataInterfaceArrayImpl.h b/Engine/Plugins/FX/Niagara/Source/Niagara/Classes/NiagaraDataInterfaceArrayImpl.h
--- a/Engine/Plugins/FX/Niagara/Source/Niagara/Classes/NiagaraDataInterfaceArrayImpl.h
+++ b/Engine/Plugins/FX/Niagara/Source/Niagara/Classes/NiagaraDataInterfaceArrayImpl.h
@@ -195,6 +195,9 @@ struct FNiagaraDataInterfaceArrayImplInternal
 
 #if WITH_EDITORONLY_DATA
 	NIAGARA_API static void GetFunctions(TArray<FNiagaraFunctionSignature>& OutFunctions, UClass* DIClass, FNiagaraTypeDefinition ValueTypeDef, bool bSupportsCPU, bool bSupportsGPU, bool bSupportsAtomicOps);
+	static void GetImmutableFunctions(TArray<FNiagaraFunctionSignature>& OutFunctions, const FNiagaraFunctionSignature& DefaultImmutableSig, FNiagaraTypeDefinition ValueTypeDef);
+	static void GetMutableFunctions(TArray<FNiagaraFunctionSignature>& OutFunctions, const FNiagaraFunctionSignature& DefaultMutableSig, FNiagaraTypeDefinition ValueTypeDef);
+	static void GetAtomicOpFunctions(TArray<FNiagaraFunctionSignature>& OutFunctions, const FNiagaraFunctionSignature& DefaultMutableSig, FNiagaraTypeDefinition ValueTypeDef);
 	NIAGARA_API static bool UpgradeFunctionCall(FNiagaraFunctionSignature& FunctionSignature);
 #endif
 	static const TCHAR* GetHLSLTemplateFile(bool bIsRWArray)
diff --git a/Engine/Plugins/FX/Niagara/Source/Niagara/Internal/DataInterface/NiagaraDataInterfaceArrayDistributionInt.h b/Engine/Plugins/FX/Niagara/Source/Niagara/Internal/DataInterface/NiagaraDataInterfaceArrayDistributionInt.h
new file mode 100644
--- /dev/null
+++ b/Engine/Plugins/FX/Niagara/Source/Niagara/Internal/DataInterface/NiagaraDataInterfaceArrayDistributionInt.h
@@ -0,0 +1,98 @@
+// Copyright Epic Games, Inc. All Rights Reserved.
+
+#pragma once
+
+#include "NiagaraDataInterfaceArray.h"
+#include "NiagaraDataInterfaceArrayImpl.h"
+#include "NiagaraDataInterfaceArrayDistributionInt.generated.h"
+
+USTRUCT(BlueprintType, meta = (DisplayName = "Weighted Distribution Int Entry"))
+struct FNDIDistributionIntArrayEntry
+{
+	GENERATED_BODY()
+
+	// The value returned if we randomly select this value
+	UPROPERTY(EditAnywhere, Category="Entry")
+	int Value = 0;
+
+	// Used to determine the probability this value will be selected when using random
+	// The higher the weight relative to the other weights the higher the chance this value is selected
+	UPROPERTY(EditAnywhere, Category="Entry")
+	float Weight = 0.0f;
+
+	bool operator==(const FNDIDistributionIntArrayEntry& Other) const
+	{
+		return
+			Value == Other.Value &&
+			FMath::IsNearlyEqual(Weight, Other.Weight);
+	}
+};
+
+namespace NDIArrayDistributionIntPrivate
+{
+	struct FNDIProxy;
+}
+
+// Array data interface used for randomly selecting a weighted integer
+UCLASS(EditInlineNew, Category = "Array", meta = (DisplayName = "Weighted Distribution Int Array", Experimental))
+class UNiagaraDataInterfaceArrayDistributionInt : public UNiagaraDataInterface
+{
+	friend NDIArrayDistributionIntPrivate::FNDIProxy;
+
+public:
+	GENERATED_BODY()
+
+	UNiagaraDataInterfaceArrayDistributionInt(FObjectInitializer const& ObjectInitializer);
+
+	// UObject Interface
+	virtual void PostInitProperties() override;
+	virtual void PostLoad() override;
+#if WITH_EDITOR
+	virtual void PostEditChangeProperty(struct FPropertyChangedEvent& PropertyChangedEvent) override;
+#endif
+	// UObject Interface End
+
+	// UNiagaraDataInterface Interface Begin
+	virtual bool CanExecuteOnTarget(ENiagaraSimTarget Target) const override { return true; }
+
+	virtual bool Equals(const UNiagaraDataInterface* Other) const override;
+	virtual bool CopyToInternal(UNiagaraDataInterface* Destination) const override;
+
+protected:
+#if WITH_EDITORONLY_DATA
+	virtual void GetFunctionsInternal(TArray<FNiagaraFunctionSignature>& OutFunctions) const override;
+#endif
+public:
+	virtual void GetVMExternalFunction(const FVMExternalFunctionBindingInfo& BindingInfo, void* InstanceData, FVMExternalFunction& OutFunc) override;
+#if WITH_EDITORONLY_DATA
+	virtual void GetParameterDefinitionHLSL(const FNiagaraDataInterfaceGPUParamInfo& ParamInfo, FString& OutHLSL) override;
+	virtual bool GetFunctionHLSL(const FNiagaraDataInterfaceHlslGenerationContext& HlslGenContext, FString& OutHLSL) override;
+	virtual bool AppendCompileHash(FNiagaraCompileHashVisitor* InVisitor) const override;
+#endif
+	virtual void BuildShaderParameters(FNiagaraShaderParametersBuilder& ShaderParametersBuilder) const override;
+	virtual void SetShaderParameters(const FNiagaraDataInterfaceSetShaderParametersContext& Context) const override;
+	// UNiagaraDataInterface Interface End
+
+	UFUNCTION(BlueprintCallable, Category = Niagara, meta = (DisplayName = "Niagara Distribution Int Array", UnsafeDuringActorConstruction = "true"))
+	static void SetNiagaraArrayDistributionInt(UNiagaraComponent* NiagaraComponent, UPARAM(DisplayName = "Parameter Name") FName OverrideName, const TArray<FNDIDistributionIntArrayEntry>& ArrayData);
+
+private:
+	void BuildTableData();
+
+	int32 GetBuiltTableLength() const;
+
+	void VMLength(FVectorVMExternalFunctionContext& Context) const;
+	void VMIsLastIndex(FVectorVMExternalFunctionContext& Context) const;
+	void VMGetLastIndex(FVectorVMExternalFunctionContext& Context) const;
+	void VMGet(FVectorVMExternalFunctionContext& Context) const;
+	void VMGetRandomValue(FVectorVMExternalFunctionContext& Context) const;
+
+private:
+	mutable FTransactionallySafeRWLock BuiltTableDataGuard;
+
+	UPROPERTY(EditAnywhere, Category = "Array")
+	TArray<FNDIDistributionIntArrayEntry> ArrayData;
+
+	UPROPERTY()
+	TArray<uint8> BuiltTableData;
+};
diff --git a/Engine/Plugins/FX/Niagara/Source/Niagara/Private/DataInterface/NiagaraDataInterfaceArrayDistributionInt.cpp b/Engine/Plugins/FX/Niagara/Source/Niagara/Private/DataInterface/NiagaraDataInterfaceArrayDistributionInt.cpp
new file mode 100644
--- /dev/null
+++ b/Engine/Plugins/FX/Niagara/Source/Niagara/Private/DataInterface/NiagaraDataInterfaceArrayDistributionInt.cpp
@@ -0,0 +1,429 @@
+// Copyright Epic Games, Inc. All Rights Reserved.
+
+#include "DataInterface/NiagaraDataInterfaceArrayDistributionInt.h"
+#include "NiagaraSystemInstance.h"
+#include "NiagaraEmitterInstanceImpl.h"
+#include "NiagaraEmitterInstance.h"
+#include "NiagaraFunctionLibrary.h"
+#include "NiagaraShaderParametersBuilder.h"
+
+#include "RenderGraphBuilder.h"
+#include "WeightedRandomSampler.h"
+
+#include UE_INLINE_GENERATED_CPP_BY_NAME(NiagaraDataInterfaceArrayDistributionInt)
+
+#define LOCTEXT_NAMESPACE "NiagaraDataInterfaceArrayDistributionInt"
+
+namespace NDIArrayDistributionIntPrivate
+{
+	static const FName NAME_GetRandomValue("GetRandomValue");
+
+	BEGIN_SHADER_PARAMETER_STRUCT(FShaderParameters, )
+		SHADER_PARAMETER(int,					Length)
+		SHADER_PARAMETER_SRV(ByteAddressBuffer,	BuiltTableData)
+	END_SHADER_PARAMETER_STRUCT()
+
+	const TCHAR* TemplateShaderFilePath = TEXT("/Plugin/FX/Niagara/Private/NiagaraDataInterfaceArrayDistributionIntTemplate.ush");
+
+	constexpr int32 GetBuiltTableEntrySize()
+	{
+		return sizeof(int32) * 4;
+	}
+
+	void EncodeBuiltTableEntry(TArray<uint8>& TableData, int32 Index, float Prob, int32 Alias, int32 Value, float Weight)
+	{
+		const int32 EntrySize = GetBuiltTableEntrySize();
+		const int32 Offset = Index * EntrySize;
+		check(TableData.Num() >= Offset + EntrySize);
+		uint8* Data = TableData.GetData() + Offset;
+		reinterpret_cast<float*>(Data)[0] = Prob;
+		reinterpret_cast<int32*>(Data)[1] = Alias;
+		reinterpret_cast<int32*>(Data)[2] = Value;
+		reinterpret_cast<float*>(Data)[3] = Weight;
+	}
+
+	void DecodeTableDirect(const TArray<uint8>& TableData, int32 Index, int32& OutValue, float& OutWeight)
+	{
+		const int32 EntrySize = GetBuiltTableEntrySize();
+		const int32 Offset = Index * EntrySize;
+		check(TableData.Num() >= Offset + EntrySize);
+		const uint8* Data = TableData.GetData() + Offset;
+		OutValue = reinterpret_cast<const int32*>(Data)[2];
+		OutWeight = reinterpret_cast<const float*>(Data)[3];
+	}
+
+	void DecodeTableIndirect(const TArray<uint8>& TableData, int32 Index, float Probability, int32& OutValue, float& OutWeight)
+	{
+		const int32 EntrySize = GetBuiltTableEntrySize();
+		const int32 Offset = Index * EntrySize;
+		check(TableData.Num() >= Offset + EntrySize);
+		const uint8* Data = TableData.GetData() + Offset;
+
+		const float EntryProb	= reinterpret_cast<const float*>(Data)[0];
+		const float EntryAlias	= reinterpret_cast<const int32*>(Data)[1];
+		Index = Probability > EntryProb ? EntryAlias : Index;
+
+		DecodeTableDirect(TableData, Index, OutValue, OutWeight);
+	}
+
+	struct FNDIProxy : public FNiagaraDataInterfaceProxy
+	{
+		virtual int32 PerInstanceDataPassedToRenderThreadSize() const override { return 0; }
+		virtual void ConsumePerInstanceDataFromGameThread(void* PerInstanceData, const FNiagaraSystemInstanceID& InstanceID) override {}
+
+		void UpdateGpuDataIfNeeded(FRDGBuilder& GraphBuilder)
+		{
+			if (bNeedsGpuDataUpdate == false)
+			{
+				return;
+			}
+
+			FRWScopeLock ScopeLock(OwnerDI->BuiltTableDataGuard, SLT_ReadOnly);
+
+			bNeedsGpuDataUpdate = false;
+
+			const bool bDataValue = OwnerDI->BuiltTableData.Num() > 0;
+			const int32 TableSizeBytes = bDataValue ? OwnerDI->BuiltTableData.Num() : GetBuiltTableEntrySize();
+
+			BuiltTableDataLength = TableSizeBytes / GetBuiltTableEntrySize();
+			BuiltTableData.Initialize(GraphBuilder.RHICmdList, TEXT("NiagaraDataInterfaceArrayDistributionInt"), TableSizeBytes);
+
+			void* UploadMemory = GraphBuilder.RHICmdList.LockBuffer(BuiltTableData.Buffer, 0, TableSizeBytes, RLM_WriteOnly);
+			if (OwnerDI->BuiltTableData.Num() > 0)
+			{
+				FMemory::Memcpy(UploadMemory, OwnerDI->BuiltTableData.GetData(), TableSizeBytes);
+			}
+			else
+			{
+				FMemory::Memzero(UploadMemory, TableSizeBytes);
+			}
+			GraphBuilder.RHICmdList.UnlockBuffer(BuiltTableData.Buffer);
+		}
+
+		UNiagaraDataInterfaceArrayDistributionInt*	OwnerDI = nullptr;
+		bool										bNeedsGpuDataUpdate = true;
+		int32										BuiltTableDataLength = 0;
+		FByteAddressBuffer							BuiltTableData;
+	};
+}
+
+UNiagaraDataInterfaceArrayDistributionInt::UNiagaraDataInterfaceArrayDistributionInt(FObjectInitializer const& ObjectInitializer)
+	: Super(ObjectInitializer)
+{
+	using namespace NDIArrayDistributionIntPrivate;
+
+	Proxy.Reset(new FNDIProxy());
+	GetProxyAs<FNDIProxy>()->OwnerDI = this;
+}
+
+void UNiagaraDataInterfaceArrayDistributionInt::PostInitProperties()
+{
+	Super::PostInitProperties();
+
+	//Can we register data interfaces as regular types and fold them into the FNiagaraVariable framework for UI and function calls etc?
+	if (HasAnyFlags(RF_ClassDefaultObject))
+	{
+		ENiagaraTypeRegistryFlags Flags = ENiagaraTypeRegistryFlags::AllowAnyVariable | ENiagaraTypeRegistryFlags::AllowParameter;
+		FNiagaraTypeRegistry::Register(FNiagaraTypeDefinition(GetClass()), Flags);
+	}
+}
+
+void UNiagaraDataInterfaceArrayDistributionInt::PostLoad()
+{
+	Super::PostLoad();
+#if WITH_EDITORONLY_DATA
+	BuildTableData();
+#endif
+}
+
+#if WITH_EDITOR
+void UNiagaraDataInterfaceArrayDistributionInt::PostEditChangeProperty(struct FPropertyChangedEvent& PropertyChangedEvent)
+{
+	Super::PostEditChangeProperty(PropertyChangedEvent);
+	BuildTableData();
+}
+#endif
+
+bool UNiagaraDataInterfaceArrayDistributionInt::Equals(const UNiagaraDataInterface* InOther) const
+{
+	using namespace NDIArrayDistributionIntPrivate;
+
+	const UNiagaraDataInterfaceArrayDistributionInt* Other = CastChecked<const UNiagaraDataInterfaceArrayDistributionInt>(InOther);
+	return
+		Super::Equals(Other) &&
+		Other->ArrayData == ArrayData;
+}
+
+bool UNiagaraDataInterfaceArrayDistributionInt::CopyToInternal(UNiagaraDataInterface* InDestination) const
+{
+	using namespace NDIArrayDistributionIntPrivate;
+
+	if (!Super::CopyToInternal(InDestination))
+	{
+		return false;
+	}
+
+	UNiagaraDataInterfaceArrayDistributionInt* Destination = CastChecked<UNiagaraDataInterfaceArrayDistributionInt>(InDestination);
+	Destination->ArrayData = ArrayData;
+	Destination->BuildTableData();
+	return true;
+}
+
+#if WITH_EDITORONLY_DATA
+void UNiagaraDataInterfaceArrayDistributionInt::GetFunctionsInternal(TArray<FNiagaraFunctionSignature>& OutFunctions) const
+{
+	using namespace NDIArrayDistributionIntPrivate;
+
+	FNiagaraFunctionSignature DefaultImmutableSig;
+	DefaultImmutableSig.bMemberFunction = true;
+	DefaultImmutableSig.bRequiresContext = false;
+	DefaultImmutableSig.bSupportsCPU = true;
+	DefaultImmutableSig.bSupportsGPU = true;
+	DefaultImmutableSig.Inputs.Emplace(FNiagaraTypeDefinition(UNiagaraDataInterfaceArrayDistributionInt::StaticClass()), TEXT("Distribution Int Array"));
+
+	FNiagaraDataInterfaceArrayImplInternal::GetImmutableFunctions(OutFunctions, DefaultImmutableSig, FNiagaraTypeDefinition(FNDIDistributionIntArrayEntry::StaticStruct()));
+	{
+		FNiagaraFunctionSignature& Sig = OutFunctions.Add_GetRef(DefaultImmutableSig);
+		Sig.Name = NAME_GetRandomValue;
+		Sig.Inputs.Emplace(FNiagaraTypeDefinition(FNiagaraRandInfo::StaticStruct()), TEXT("RandomInfo"));
+		Sig.Outputs.Emplace(FNiagaraTypeDefinition::GetIntDef(), TEXT("Value"));
+	}
+}
+#endif
+
+void UNiagaraDataInterfaceArrayDistributionInt::GetVMExternalFunction(const FVMExternalFunctionBindingInfo& BindingInfo, void* InstanceData, FVMExternalFunction& OutFunc)
+{
+	using namespace NDIArrayDistributionIntPrivate;
+
+	if (BindingInfo.Name == FNiagaraDataInterfaceArrayImplInternal::Function_LengthName)
+	{
+		OutFunc = FVMExternalFunction::CreateUObject(this, &UNiagaraDataInterfaceArrayDistributionInt::VMLength);
+	}
+	else if (BindingInfo.Name == FNiagaraDataInterfaceArrayImplInternal::Function_IsValidIndexName)
+	{
+		OutFunc = FVMExternalFunction::CreateUObject(this, &UNiagaraDataInterfaceArrayDistributionInt::VMIsLastIndex);
+	}
+	else if (BindingInfo.Name == FNiagaraDataInterfaceArrayImplInternal::Function_LastIndexName)
+	{
+		OutFunc = FVMExternalFunction::CreateUObject(this, &UNiagaraDataInterfaceArrayDistributionInt::VMGetLastIndex);
+	}
+	else if (BindingInfo.Name == FNiagaraDataInterfaceArrayImplInternal::Function_GetName)
+	{
+		OutFunc = FVMExternalFunction::CreateUObject(this, &UNiagaraDataInterfaceArrayDistributionInt::VMGet);
+	}
+	else if (BindingInfo.Name == NDIArrayDistributionIntPrivate::NAME_GetRandomValue)
+	{
+		OutFunc = FVMExternalFunction::CreateUObject(this, &UNiagaraDataInterfaceArrayDistributionInt::VMGetRandomValue);
+	}
+}
+
+#if WITH_EDITORONLY_DATA
+void UNiagaraDataInterfaceArrayDistributionInt::GetParameterDefinitionHLSL(const FNiagaraDataInterfaceGPUParamInfo& ParamInfo, FString& OutHLSL)
+{
+	using namespace NDIArrayDistributionIntPrivate;
+
+	const TMap<FString, FStringFormatArg> TemplateArgs =
+	{
+		{TEXT("ParameterName"),	ParamInfo.DataInterfaceHLSLSymbol},
+	};
+	AppendTemplateHLSL(OutHLSL, TemplateShaderFilePath, TemplateArgs);
+}
+
+bool UNiagaraDataInterfaceArrayDistributionInt::GetFunctionHLSL(const FNiagaraDataInterfaceHlslGenerationContext& HlslGenContext, FString& OutHLSL)
+{
+	using namespace NDIArrayDistributionIntPrivate;
+
+	return true;
+}
+
+bool UNiagaraDataInterfaceArrayDistributionInt::AppendCompileHash(FNiagaraCompileHashVisitor* InVisitor) const
+{
+	using namespace NDIArrayDistributionIntPrivate;
+
+	bool bSuccess = Super::AppendCompileHash(InVisitor);
+	InVisitor->UpdateShaderFile(TemplateShaderFilePath);
+	return bSuccess;
+}
+#endif
+
+void UNiagaraDataInterfaceArrayDistributionInt::BuildShaderParameters(FNiagaraShaderParametersBuilder& ShaderParametersBuilder) const
+{
+	using namespace NDIArrayDistributionIntPrivate;
+
+	ShaderParametersBuilder.AddNestedStruct<FShaderParameters>();
+}
+
+void UNiagaraDataInterfaceArrayDistributionInt::SetShaderParameters(const FNiagaraDataInterfaceSetShaderParametersContext& Context) const
+{
+	using namespace NDIArrayDistributionIntPrivate;
+
+	FNDIProxy& DIProxy = Context.GetProxy<FNDIProxy>();
+	DIProxy.UpdateGpuDataIfNeeded(Context.GetGraphBuilder());
+
+	FShaderParameters* ShaderParameters = Context.GetParameterNestedStruct<FShaderParameters>();
+	ShaderParameters->Length			= DIProxy.BuiltTableDataLength;
+	ShaderParameters->BuiltTableData	= DIProxy.BuiltTableData.SRV;
+}
+
+void UNiagaraDataInterfaceArrayDistributionInt::BuildTableData()
+{
+	using namespace NDIArrayDistributionIntPrivate;
+
+	struct FArrayWeightedSampler : public FWeightedRandomSampler
+	{
+		FArrayWeightedSampler(TConstArrayView<FNDIDistributionIntArrayEntry> InArrayData)
+			: ArrayData(InArrayData)
+		{
+		}
+
+		virtual float GetWeights(TArray<float>& OutWeights) override
+		{
+			float SumWeights = 0.0f;
+
+			OutWeights.SetNumUninitialized(ArrayData.Num());
+			for (int32 i = 0; i < ArrayData.Num(); ++i)
+			{
+				const FNDIDistributionIntArrayEntry& Entry = ArrayData[i];
+				OutWeights[i] = FMath::Max(Entry.Weight, UE_SMALL_NUMBER);
+				SumWeights += OutWeights[i];
+			}
+
+			return SumWeights;
+		}
+
+		TConstArrayView<FNDIDistributionIntArrayEntry> ArrayData;
+	};
+
+	FArrayWeightedSampler ArraySampler(ArrayData);
+	ArraySampler.Initialize();
+
+	TArray<uint8> NewBuiltTableData;
+	NewBuiltTableData.SetNumUninitialized(FMath::Max(ArrayData.Num(), 1) * GetBuiltTableEntrySize());
+
+	if (ArrayData.Num() > 0)
+	{
+		NewBuiltTableData.SetNumUninitialized(ArrayData.Num() * GetBuiltTableEntrySize());
+
+		for (int32 i = 0; i < ArrayData.Num(); ++i)
+		{
+			const float Prob = ArraySampler.GetProb()[i];
+			const int32 Alias = ArraySampler.GetAlias()[i];
+			const int32 Value = ArrayData[i].Value;
+			const float Weight = ArrayData[i].Weight;
+			EncodeBuiltTableEntry(NewBuiltTableData, i, Prob, Alias, Value, Weight);
+		}
+	}
+	else
+	{
+		EncodeBuiltTableEntry(NewBuiltTableData, 0, 0.0f, 0, 0, 0.0f);
+	}
+
+	FRWScopeLock ScopeLock(BuiltTableDataGuard, SLT_Write);
+	Swap(BuiltTableData, NewBuiltTableData);
+
+	GetProxyAs<FNDIProxy>()->bNeedsGpuDataUpdate = true;
+}
+
+int32 UNiagaraDataInterfaceArrayDistributionInt::GetBuiltTableLength() const
+{
+	using namespace NDIArrayDistributionIntPrivate;
+
+	FRWScopeLock ScopeLock(BuiltTableDataGuard, SLT_ReadOnly);
+	return BuiltTableData.Num() / GetBuiltTableEntrySize();
+}
+
+void UNiagaraDataInterfaceArrayDistributionInt::VMLength(FVectorVMExternalFunctionContext& Context) const
+{
+	FNDIOutputParam<int32> OutLength(Context);
+
+	const int32 Length = GetBuiltTableLength();
+	for (int32 i = 0; i < Context.GetNumInstances(); ++i)
+	{
+		OutLength.SetAndAdvance(Length);
+	}
+}
+
+void UNiagaraDataInterfaceArrayDistributionInt::VMIsLastIndex(FVectorVMExternalFunctionContext& Context) const
+{
+	FNDIInputParam<int32> InIndex(Context);
+	FNDIOutputParam<bool> OutValid(Context);
+
+	const int32 Length = GetBuiltTableLength();
+	for (int32 i = 0; i < Context.GetNumInstances(); ++i)
+	{
+		const int32 Index = InIndex.GetAndAdvance();
+		const bool bIsValid = Index >= 0 && Index < Length;
+		OutValid.SetAndAdvance(bIsValid);
+	}
+}
+
+void UNiagaraDataInterfaceArrayDistributionInt::VMGetLastIndex(FVectorVMExternalFunctionContext& Context) const
+{
+	FNDIOutputParam<int32> OutLastIndex(Context);
+
+	const int32 LastIndex = GetBuiltTableLength() - 1;
+	for (int32 i = 0; i < Context.GetNumInstances(); ++i)
+	{
+		OutLastIndex.SetAndAdvance(LastIndex);
+	}
+}
+
+void UNiagaraDataInterfaceArrayDistributionInt::VMGet(FVectorVMExternalFunctionContext& Context) const
+{
+	using namespace NDIArrayDistributionIntPrivate;
+
+	FNDIInputParam<int32>	InIndex(Context);
+	FNDIOutputParam<int32>	OutIndex(Context);
+	FNDIOutputParam<float>	OutWeight(Context);
+
+	FRWScopeLock ScopeLock(BuiltTableDataGuard, SLT_ReadOnly);
+	const int32 LengthMinusOne = GetBuiltTableLength() - 1;
+	for (int32 i = 0; i < Context.GetNumInstances(); ++i)
+	{
+		const int32 Index = InIndex.GetAndAdvance();
+		const int32 ClampedIndex = FMath::Clamp(Index, 0, LengthMinusOne);
+		
+		int32 EntryValue = 0;
+		float EntryWeight = 0.0f;
+		DecodeTableDirect(BuiltTableData, ClampedIndex, EntryValue, EntryWeight);
+
+		OutIndex.SetAndAdvance(EntryValue);
+		OutWeight.SetAndAdvance(EntryWeight);
+	}
+}
+
+void UNiagaraDataInterfaceArrayDistributionInt::VMGetRandomValue(FVectorVMExternalFunctionContext& Context) const
+{
+	using namespace NDIArrayDistributionIntPrivate;
+
+	FNDIRandomHelper		RandHelper(Context);
+	FNDIOutputParam<int32>	OutValue(Context);
+
+	FRWScopeLock ScopeLock(BuiltTableDataGuard, SLT_ReadOnly);
+	const int32 LengthMinusOne = GetBuiltTableLength() - 1;
+	for (int32 i = 0; i < Context.GetNumInstances(); ++i)
+	{
+		RandHelper.GetAndAdvance();
+		int32 Index = RandHelper.RandRange(i, 0, LengthMinusOne);
+
+		int32 EntryValue = 0;
+		float EntryWeight = 0.0f;
+		DecodeTableIndirect(BuiltTableData, Index, RandHelper.RandRange(i, 0.0f, 1.0f), EntryValue, EntryWeight);
+
+		OutValue.SetAndAdvance(EntryValue);
+	}
+}
+
+void UNiagaraDataInterfaceArrayDistributionInt::SetNiagaraArrayDistributionInt(UNiagaraComponent* NiagaraComponent, FName OverrideName, const TArray<FNDIDistributionIntArrayEntry>& InArrayData)
+{
+	UNiagaraDataInterfaceArrayDistributionInt* ArrayDI = UNiagaraFunctionLibrary::GetDataInterface<UNiagaraDataInterfaceArrayDistributionInt>(NiagaraComponent, OverrideName);
+	if (ArrayDI == nullptr)
+	{
+		return;
+	}
+
+	ArrayDI->ArrayData = InArrayData;
+	ArrayDI->BuildTableData();
+}
+
+#undef LOCTEXT_NAMESPACE
diff --git a/Engine/Plugins/FX/Niagara/Source/Niagara/Private/NiagaraDataInterfaceArrayImpl.cpp b/Engine/Plugins/FX/Niagara/Source/Niagara/Private/NiagaraDataInterfaceArrayImpl.cpp
--- a/Engine/Plugins/FX/Niagara/Source/Niagara/Private/NiagaraDataInterfaceArrayImpl.cpp
+++ b/Engine/Plugins/FX/Niagara/Source/Niagara/Private/NiagaraDataInterfaceArrayImpl.cpp
@@ -35,6 +35,23 @@ void FNiagaraDataInterfaceArrayImplInternal::GetFunctions(TArray<FNiagaraFunctio
 	DefaultImmutableSig.bSupportsGPU = bSupportsGPU;
 	DefaultImmutableSig.Inputs.Emplace(FNiagaraTypeDefinition(DIClass), TEXT("Array interface"));
 	DefaultImmutableSig.FunctionVersion = FNiagaraDataInterfaceArrayImplInternal::FFunctionVersion::LatestVersion;
+
+	GetImmutableFunctions(OutFunctions, DefaultImmutableSig, ValueTypeDef);
+
+	// Mutable functions
+	FNiagaraFunctionSignature DefaultMutableSig = DefaultImmutableSig;
+	DefaultMutableSig.bSupportsGPU = bSupportsGPU;
+	DefaultMutableSig.bRequiresExecPin = true;
+
+	GetMutableFunctions(OutFunctions, DefaultMutableSig, ValueTypeDef);
+	if (bSupportsAtomicOps)
+	{
+		GetAtomicOpFunctions(OutFunctions, DefaultMutableSig, ValueTypeDef);
+	}
+}
+
+void FNiagaraDataInterfaceArrayImplInternal::GetImmutableFunctions(TArray<FNiagaraFunctionSignature>& OutFunctions, const FNiagaraFunctionSignature& DefaultImmutableSig, FNiagaraTypeDefinition ValueTypeDef)
+{
 	{
 		FNiagaraFunctionSignature& Sig = OutFunctions.Add_GetRef(DefaultImmutableSig);
 		Sig.Name = FNiagaraDataInterfaceArrayImplInternal::Function_LengthName;
@@ -67,11 +84,10 @@ void FNiagaraDataInterfaceArrayImplInternal::GetFunctions(TArray<FNiagaraFunctio
 		Sig.Outputs.Emplace(ValueTypeDef, TEXT("Value"));
 		Sig.Description = NSLOCTEXT("Niagara", "Array_GetDesc", "Gets the value from the array at the given zero based index.");
 	}
+}
 
-	// Mutable functions
-	FNiagaraFunctionSignature DefaultMutableSig = DefaultImmutableSig;
-	DefaultMutableSig.bSupportsGPU = bSupportsGPU;
-	DefaultMutableSig.bRequiresExecPin = true;
+void FNiagaraDataInterfaceArrayImplInternal::GetMutableFunctions(TArray<FNiagaraFunctionSignature>& OutFunctions, const FNiagaraFunctionSignature& DefaultMutableSig, FNiagaraTypeDefinition ValueTypeDef)
+{
 	{
 		FNiagaraFunctionSignature& Sig = OutFunctions.Add_GetRef(DefaultMutableSig);
 		Sig.Name = FNiagaraDataInterfaceArrayImplInternal::Function_ClearName;
@@ -116,48 +132,48 @@ void FNiagaraDataInterfaceArrayImplInternal::GetFunctions(TArray<FNiagaraFunctio
 		Sig.InputDescriptions.Add(Sig.Inputs[1], NSLOCTEXT("Niagara", "Array_RemoveLastElemDesc_SkipRemove", "When enabled will not remove a value from the array, the return value will therefore be invalid."));
 		Sig.OutputDescriptions.Add(Sig.Outputs[1], NSLOCTEXT("Niagara", "Array_RemoveLastElemDesc_IsValid", "True if we removed a value from the array, False if no entries or we skipped the remove."));
 	}
+}
 
-	if (bSupportsAtomicOps)
+void FNiagaraDataInterfaceArrayImplInternal::GetAtomicOpFunctions(TArray<FNiagaraFunctionSignature>& OutFunctions, const FNiagaraFunctionSignature& DefaultMutableSig, FNiagaraTypeDefinition ValueTypeDef)
+{
 	{
-		{
-			FNiagaraFunctionSignature& Sig = OutFunctions.Add_GetRef(DefaultMutableSig);
-			Sig.Name = FNiagaraDataInterfaceArrayImplInternal::Function_AtomicAddName;
-			Sig.Inputs.Emplace(FNiagaraTypeDefinition::GetBoolDef(), TEXT("SkipAdd"));
-			Sig.Inputs.Emplace(FNiagaraTypeDefinition::GetIntDef(), TEXT("Index"));
-			Sig.Inputs.Emplace(ValueTypeDef, TEXT("Value"));
-			Sig.Outputs.Emplace(ValueTypeDef, TEXT("PreviousValue"));
-			Sig.Outputs.Emplace(ValueTypeDef, TEXT("CurrentValue"));
-			Sig.SetDescription(NSLOCTEXT("Niagara", "Array_AtomicAddDesc", "Optionally perform an atomic add on the array element."));
-			Sig.SetInputDescription(Sig.Inputs[1], NSLOCTEXT("Niagara", "Array_AtomicAdd_SkipAdd", "When enabled will not perform the add operation, the return values will therefore be invalid."));
-			Sig.SetOutputDescription(Sig.Inputs[0], NSLOCTEXT("Niagara", "Array_AtomicAdd_PrevValue", "The value before the operation was performed."));
-			Sig.SetOutputDescription(Sig.Inputs[1], NSLOCTEXT("Niagara", "Array_AtomicAdd_CurrValue", "The value after the operation was performed."));
-		}
-		{
-			FNiagaraFunctionSignature& Sig = OutFunctions.Add_GetRef(DefaultMutableSig);
-			Sig.Name = FNiagaraDataInterfaceArrayImplInternal::Function_AtomicMinName;
-			Sig.Inputs.Emplace(FNiagaraTypeDefinition::GetBoolDef(), TEXT("SkipMin"));
-			Sig.Inputs.Emplace(FNiagaraTypeDefinition::GetIntDef(), TEXT("Index"));
-			Sig.Inputs.Emplace(ValueTypeDef, TEXT("Value"));
-			Sig.Outputs.Emplace(ValueTypeDef, TEXT("PreviousValue"));
-			Sig.Outputs.Emplace(ValueTypeDef, TEXT("CurrentValue"));
-			Sig.SetDescription(NSLOCTEXT("Niagara", "Array_AtomicMinDesc", "Optionally perform an atomic min on the array element."));
-			Sig.SetInputDescription(Sig.Inputs[1], NSLOCTEXT("Niagara", "Array_AtomicMin_SkipMin", "When enabled will not perform the min operation, the return values will therefore be invalid."));
-			Sig.SetOutputDescription(Sig.Inputs[0], NSLOCTEXT("Niagara", "Array_AtomicMin_PrevValue", "The value before the operation was performed."));
-			Sig.SetOutputDescription(Sig.Inputs[1], NSLOCTEXT("Niagara", "Array_AtomicMin_CurrValue", "The value after the operation was performed."));
-		}
-		{
-			FNiagaraFunctionSignature& Sig = OutFunctions.Add_GetRef(DefaultMutableSig);
-			Sig.Name = FNiagaraDataInterfaceArrayImplInternal::Function_AtomicMaxName;
-			Sig.Inputs.Emplace(FNiagaraTypeDefinition::GetBoolDef(), TEXT("SkipMax"));
-			Sig.Inputs.Emplace(FNiagaraTypeDefinition::GetIntDef(), TEXT("Index"));
-			Sig.Inputs.Emplace(ValueTypeDef, TEXT("Value"));
-			Sig.Outputs.Emplace(ValueTypeDef, TEXT("PreviousValue"));
-			Sig.Outputs.Emplace(ValueTypeDef, TEXT("CurrentValue"));
-			Sig.SetDescription(NSLOCTEXT("Niagara", "Array_AtomicMaxDesc", "Optionally perform an atomic max on the array element."));
-			Sig.SetInputDescription(Sig.Inputs[1], NSLOCTEXT("Niagara", "Array_AtomicMax_SkipMax", "When enabled will not perform the max operation, the return values will therefore be invalid."));
-			Sig.SetOutputDescription(Sig.Inputs[0], NSLOCTEXT("Niagara", "Array_AtomicMax_PrevValue", "The value before the operation was performed."));
-			Sig.SetOutputDescription(Sig.Inputs[1], NSLOCTEXT("Niagara", "Array_AtomicMax_CurrValue", "The value after the operation was performed."));
-		}
+		FNiagaraFunctionSignature& Sig = OutFunctions.Add_GetRef(DefaultMutableSig);
+		Sig.Name = FNiagaraDataInterfaceArrayImplInternal::Function_AtomicAddName;
+		Sig.Inputs.Emplace(FNiagaraTypeDefinition::GetBoolDef(), TEXT("SkipAdd"));
+		Sig.Inputs.Emplace(FNiagaraTypeDefinition::GetIntDef(), TEXT("Index"));
+		Sig.Inputs.Emplace(ValueTypeDef, TEXT("Value"));
+		Sig.Outputs.Emplace(ValueTypeDef, TEXT("PreviousValue"));
+		Sig.Outputs.Emplace(ValueTypeDef, TEXT("CurrentValue"));
+		Sig.SetDescription(NSLOCTEXT("Niagara", "Array_AtomicAddDesc", "Optionally perform an atomic add on the array element."));
+		Sig.SetInputDescription(Sig.Inputs[1], NSLOCTEXT("Niagara", "Array_AtomicAdd_SkipAdd", "When enabled will not perform the add operation, the return values will therefore be invalid."));
+		Sig.SetOutputDescription(Sig.Inputs[0], NSLOCTEXT("Niagara", "Array_AtomicAdd_PrevValue", "The value before the operation was performed."));
+		Sig.SetOutputDescription(Sig.Inputs[1], NSLOCTEXT("Niagara", "Array_AtomicAdd_CurrValue", "The value after the operation was performed."));
+	}
+	{
+		FNiagaraFunctionSignature& Sig = OutFunctions.Add_GetRef(DefaultMutableSig);
+		Sig.Name = FNiagaraDataInterfaceArrayImplInternal::Function_AtomicMinName;
+		Sig.Inputs.Emplace(FNiagaraTypeDefinition::GetBoolDef(), TEXT("SkipMin"));
+		Sig.Inputs.Emplace(FNiagaraTypeDefinition::GetIntDef(), TEXT("Index"));
+		Sig.Inputs.Emplace(ValueTypeDef, TEXT("Value"));
+		Sig.Outputs.Emplace(ValueTypeDef, TEXT("PreviousValue"));
+		Sig.Outputs.Emplace(ValueTypeDef, TEXT("CurrentValue"));
+		Sig.SetDescription(NSLOCTEXT("Niagara", "Array_AtomicMinDesc", "Optionally perform an atomic min on the array element."));
+		Sig.SetInputDescription(Sig.Inputs[1], NSLOCTEXT("Niagara", "Array_AtomicMin_SkipMin", "When enabled will not perform the min operation, the return values will therefore be invalid."));
+		Sig.SetOutputDescription(Sig.Inputs[0], NSLOCTEXT("Niagara", "Array_AtomicMin_PrevValue", "The value before the operation was performed."));
+		Sig.SetOutputDescription(Sig.Inputs[1], NSLOCTEXT("Niagara", "Array_AtomicMin_CurrValue", "The value after the operation was performed."));
+	}
+	{
+		FNiagaraFunctionSignature& Sig = OutFunctions.Add_GetRef(DefaultMutableSig);
+		Sig.Name = FNiagaraDataInterfaceArrayImplInternal::Function_AtomicMaxName;
+		Sig.Inputs.Emplace(FNiagaraTypeDefinition::GetBoolDef(), TEXT("SkipMax"));
+		Sig.Inputs.Emplace(FNiagaraTypeDefinition::GetIntDef(), TEXT("Index"));
+		Sig.Inputs.Emplace(ValueTypeDef, TEXT("Value"));
+		Sig.Outputs.Emplace(ValueTypeDef, TEXT("PreviousValue"));
+		Sig.Outputs.Emplace(ValueTypeDef, TEXT("CurrentValue"));
+		Sig.SetDescription(NSLOCTEXT("Niagara", "Array_AtomicMaxDesc", "Optionally perform an atomic max on the array element."));
+		Sig.SetInputDescription(Sig.Inputs[1], NSLOCTEXT("Niagara", "Array_AtomicMax_SkipMax", "When enabled will not perform the max operation, the return values will therefore be invalid."));
+		Sig.SetOutputDescription(Sig.Inputs[0], NSLOCTEXT("Niagara", "Array_AtomicMax_PrevValue", "The value before the operation was performed."));
+		Sig.SetOutputDescription(Sig.Inputs[1], NSLOCTEXT("Niagara", "Array_AtomicMax_CurrValue", "The value after the operation was performed."));
 	}
 }
 
