Replace some usages of FORCEINLINE with inline in Niagara modules.

For context about this change, see this doc: https://docs.google.com/document/d/1LjpT7NpzTecfXdFNyv5KM6sCu67827BRljNi1iaWVSE/edit?tab=t.0#heading=h.fvtz8gjukz7x

Includes modules: Niagara, NiagaraCore, NiagaraEditor, NiagaraShader, NiagaraVertexFactories

Reviewers were automatically determined. Please comment if anyone else should be added.

#rnx

diff --git a/Engine/Plugins/FX/Niagara/Source/Niagara/Classes/NDISkeletalMeshCommon.h b/Engine/Plugins/FX/Niagara/Source/Niagara/Classes/NDISkeletalMeshCommon.h
--- a/Engine/Plugins/FX/Niagara/Source/Niagara/Classes/NDISkeletalMeshCommon.h
+++ b/Engine/Plugins/FX/Niagara/Source/Niagara/Classes/NDISkeletalMeshCommon.h
@@ -39,7 +39,7 @@ struct FSkeletalMeshAccessorHelper
 	FSkeletalMeshAccessorHelper operator=(FSkeletalMeshAccessorHelper&&) = delete;
 
 	template<typename FilterMode, typename AreaWeightingMode>
-	FORCEINLINE void Init(const FNDISkeletalMesh_InstanceData* InstData)
+	inline void Init(const FNDISkeletalMesh_InstanceData* InstData)
 	{
 		Comp = Cast<USkeletalMeshComponent>(InstData->SceneComponent.Get());
 		Mesh = InstData->SkeletalMesh.Get();
@@ -55,12 +55,12 @@ struct FSkeletalMeshAccessorHelper
 		}
 	}
 
-	FORCEINLINE bool AreBonesAccessible() const
+	inline bool AreBonesAccessible() const
 	{
 		return Mesh != nullptr;
 	}
 
-	FORCEINLINE bool IsLODAccessible() const
+	inline bool IsLODAccessible() const
 	{
 		return LODData != nullptr;
 	}
@@ -88,26 +88,26 @@ void FSkeletalMeshAccessorHelper::Init<TNDISkelMesh_FilterModeSingle, TNDISkelMe
 template<typename SkinningMode>
 struct FSkinnedPositionAccessorHelper
 {
-	FORCEINLINE int32 GetBoneCount(FSkeletalMeshAccessorHelper& Accessor, bool RequiresPrevious) const = delete;
-	FORCEINLINE void GetTriangleIndices(int32 Tri, int32& Idx0, int32& Idx1, int32& Idx2) = delete;
-	FORCEINLINE void GetSkinnedTrianglePositions(FSkeletalMeshAccessorHelper& Accessor, int32 Tri, FVector3f& OutPos0, FVector3f& OutPos1, FVector3f& OutPos2, FVector3f& OutPrev0, FVector3f& OutPrev1, FVector3f& OutPrev2, int32& Idx0, int32& Idx1, int32& Idx2) = delete;
-	FORCEINLINE void GetSkinnedTrianglePositions(FSkeletalMeshAccessorHelper& Accessor, int32 Tri, FVector3f& OutPos0, FVector3f& OutPos1, FVector3f& OutPos2, int32& Idx0, int32& Idx1, int32& Idx2) = delete;
-	FORCEINLINE FVector3f GetSkinnedVertexPosition(FSkeletalMeshAccessorHelper& Accessor, int32 VertexIndex) = delete;
-	FORCEINLINE FVector3f GetSkinnedVertexPreviousPosition(FSkeletalMeshAccessorHelper& Accessor, int32 VertexIndex) = delete;
-	FORCEINLINE void GetSkinnedTangentBasis(FSkeletalMeshAccessorHelper& Accessor, int32 VertexIndex, FVector3f& OutTangentX, FVector3f& OutTangentY, FVector3f& OutTangentZ) = delete;
-	FORCEINLINE void GetSkinnedPreviousTangentBasis(FSkeletalMeshAccessorHelper& Accessor, int32 VertexIndex, FVector3f& OutTangentX, FVector3f& OutTangentY, FVector3f& OutTangentZ) = delete;
-	FORCEINLINE FVector3f GetSkinnedBonePosition(FSkeletalMeshAccessorHelper& Accessor, int32 BoneIndex) = delete;
-	FORCEINLINE FVector3f GetSkinnedBonePreviousPosition(FSkeletalMeshAccessorHelper& Accessor, int32 BoneIndex) = delete;
-	FORCEINLINE FQuat4f GetSkinnedBoneRotation(FSkeletalMeshAccessorHelper& Accessor, int32 BoneIndex) = delete;
-	FORCEINLINE FQuat4f GetSkinnedBonePreviousRotation(FSkeletalMeshAccessorHelper& Accessor, int32 BoneIndex) = delete;
-	FORCEINLINE FVector3f GetSkinnedBoneScale(FSkeletalMeshAccessorHelper& Accessor, int32 BoneIndex) = delete;
-	FORCEINLINE FVector3f GetSkinnedBonePreviousScale(FSkeletalMeshAccessorHelper& Accessor, int32 BoneIndex) = delete;
+	inline int32 GetBoneCount(FSkeletalMeshAccessorHelper& Accessor, bool RequiresPrevious) const = delete;
+	inline void GetTriangleIndices(int32 Tri, int32& Idx0, int32& Idx1, int32& Idx2) = delete;
+	inline void GetSkinnedTrianglePositions(FSkeletalMeshAccessorHelper& Accessor, int32 Tri, FVector3f& OutPos0, FVector3f& OutPos1, FVector3f& OutPos2, FVector3f& OutPrev0, FVector3f& OutPrev1, FVector3f& OutPrev2, int32& Idx0, int32& Idx1, int32& Idx2) = delete;
+	inline void GetSkinnedTrianglePositions(FSkeletalMeshAccessorHelper& Accessor, int32 Tri, FVector3f& OutPos0, FVector3f& OutPos1, FVector3f& OutPos2, int32& Idx0, int32& Idx1, int32& Idx2) = delete;
+	inline FVector3f GetSkinnedVertexPosition(FSkeletalMeshAccessorHelper& Accessor, int32 VertexIndex) = delete;
+	inline FVector3f GetSkinnedVertexPreviousPosition(FSkeletalMeshAccessorHelper& Accessor, int32 VertexIndex) = delete;
+	inline void GetSkinnedTangentBasis(FSkeletalMeshAccessorHelper& Accessor, int32 VertexIndex, FVector3f& OutTangentX, FVector3f& OutTangentY, FVector3f& OutTangentZ) = delete;
+	inline void GetSkinnedPreviousTangentBasis(FSkeletalMeshAccessorHelper& Accessor, int32 VertexIndex, FVector3f& OutTangentX, FVector3f& OutTangentY, FVector3f& OutTangentZ) = delete;
+	inline FVector3f GetSkinnedBonePosition(FSkeletalMeshAccessorHelper& Accessor, int32 BoneIndex) = delete;
+	inline FVector3f GetSkinnedBonePreviousPosition(FSkeletalMeshAccessorHelper& Accessor, int32 BoneIndex) = delete;
+	inline FQuat4f GetSkinnedBoneRotation(FSkeletalMeshAccessorHelper& Accessor, int32 BoneIndex) = delete;
+	inline FQuat4f GetSkinnedBonePreviousRotation(FSkeletalMeshAccessorHelper& Accessor, int32 BoneIndex) = delete;
+	inline FVector3f GetSkinnedBoneScale(FSkeletalMeshAccessorHelper& Accessor, int32 BoneIndex) = delete;
+	inline FVector3f GetSkinnedBonePreviousScale(FSkeletalMeshAccessorHelper& Accessor, int32 BoneIndex) = delete;
 };
 
 template<>
 struct FSkinnedPositionAccessorHelper<TNDISkelMesh_SkinningModeInvalid>
 {
-	FORCEINLINE int32 GetBoneCount(FSkeletalMeshAccessorHelper& Accessor, bool RequiresPrevious) const
+	inline int32 GetBoneCount(FSkeletalMeshAccessorHelper& Accessor, bool RequiresPrevious) const
 	{
 		if (const USkeletalMesh* Mesh = Accessor.Mesh)
 		{
@@ -117,39 +117,39 @@ struct FSkinnedPositionAccessorHelper<TNDISkelMesh_SkinningModeInvalid>
 		return 0;
 	}
 
-	FORCEINLINE void GetTriangleIndices(FSkeletalMeshAccessorHelper& Accessor, int32 Tri, int32& Idx0, int32& Idx1, int32& Idx2)
+	inline void GetTriangleIndices(FSkeletalMeshAccessorHelper& Accessor, int32 Tri, int32& Idx0, int32& Idx1, int32& Idx2)
 	{
 		Idx0 = Idx1 = Idx2 = -1;
 	}
 
-	FORCEINLINE void GetSkinnedTrianglePositions(FSkeletalMeshAccessorHelper& Accessor, int32 Idx0, int32 Idx1, int32 Idx2, FVector3f& OutPos0, FVector3f& OutPos1, FVector3f& OutPos2)
+	inline void GetSkinnedTrianglePositions(FSkeletalMeshAccessorHelper& Accessor, int32 Idx0, int32 Idx1, int32 Idx2, FVector3f& OutPos0, FVector3f& OutPos1, FVector3f& OutPos2)
 	{
 		OutPos0 = OutPos1 = OutPos2 = FVector3f::ZeroVector;
 	}
 
-	FORCEINLINE void GetSkinnedTrianglePreviousPositions(FSkeletalMeshAccessorHelper& Accessor, int32 Idx0, int32 Idx1, int32 Idx2, FVector3f& OutPos0, FVector3f& OutPos1, FVector3f& OutPos2)
+	inline void GetSkinnedTrianglePreviousPositions(FSkeletalMeshAccessorHelper& Accessor, int32 Idx0, int32 Idx1, int32 Idx2, FVector3f& OutPos0, FVector3f& OutPos1, FVector3f& OutPos2)
 	{
 		OutPos0 = OutPos1 = OutPos2 = FVector3f::ZeroVector;
 	}
 
-	FORCEINLINE FVector3f GetSkinnedVertexPosition(FSkeletalMeshAccessorHelper& Accessor, int32 VertexIndex)
+	inline FVector3f GetSkinnedVertexPosition(FSkeletalMeshAccessorHelper& Accessor, int32 VertexIndex)
 	{
 		return FVector3f::ZeroVector;
 	}
 
-	FORCEINLINE FVector3f GetSkinnedVertexPreviousPosition(FSkeletalMeshAccessorHelper& Accessor, int32 VertexIndex)
+	inline FVector3f GetSkinnedVertexPreviousPosition(FSkeletalMeshAccessorHelper& Accessor, int32 VertexIndex)
 	{
 		return FVector3f::ZeroVector;
 	}
 
-	FORCEINLINE void GetSkinnedTangentBasis(FSkeletalMeshAccessorHelper& Accessor, int32 VertexIndex, FVector3f& OutTangentX, FVector3f& OutTangentY, FVector3f& OutTangentZ)
+	inline void GetSkinnedTangentBasis(FSkeletalMeshAccessorHelper& Accessor, int32 VertexIndex, FVector3f& OutTangentX, FVector3f& OutTangentY, FVector3f& OutTangentZ)
 	{
 		OutTangentX = FVector3f(1.0f, 0.0f, 0.0f);
 		OutTangentY = FVector3f(0.0f, 1.0f, 0.0f);
 		OutTangentZ = FVector3f(0.0f, 0.0f, 1.0f);
 	}
 
-	FORCEINLINE void GetSkinnedPreviousTangentBasis(FSkeletalMeshAccessorHelper& Accessor, int32 VertexIndex, FVector3f& OutTangentX, FVector3f& OutTangentY, FVector3f& OutTangentZ)
+	inline void GetSkinnedPreviousTangentBasis(FSkeletalMeshAccessorHelper& Accessor, int32 VertexIndex, FVector3f& OutTangentX, FVector3f& OutTangentY, FVector3f& OutTangentZ)
 	{
 		OutTangentX = FVector3f(1.0f, 0.0f, 0.0f);
 		OutTangentY = FVector3f(0.0f, 1.0f, 0.0f);
@@ -157,7 +157,7 @@ struct FSkinnedPositionAccessorHelper<TNDISkelMesh_SkinningModeInvalid>
 	}
 
 	// The bone accessor functions are valid if a mesh is present, so don't stub them entirely
-	FORCEINLINE FVector3f GetSkinnedBonePosition(FSkeletalMeshAccessorHelper& Accessor, int32 BoneIndex)
+	inline FVector3f GetSkinnedBonePosition(FSkeletalMeshAccessorHelper& Accessor, int32 BoneIndex)
 	{
 		const int32 NumRealBones = Accessor.Mesh->GetRefSkeleton().GetRawBoneNum();
 		if (BoneIndex < NumRealBones)
@@ -169,12 +169,12 @@ struct FSkinnedPositionAccessorHelper<TNDISkelMesh_SkinningModeInvalid>
 		return (FVector3f)RefTransform.GetLocation();
 	}
 
-	FORCEINLINE FVector3f GetSkinnedBonePreviousPosition(FSkeletalMeshAccessorHelper& Accessor, int32 BoneIndex)
+	inline FVector3f GetSkinnedBonePreviousPosition(FSkeletalMeshAccessorHelper& Accessor, int32 BoneIndex)
 	{
 		return GetSkinnedBonePosition(Accessor, BoneIndex);
 	}
 
-	FORCEINLINE FQuat4f GetSkinnedBoneRotation(FSkeletalMeshAccessorHelper& Accessor, int32 BoneIndex)
+	inline FQuat4f GetSkinnedBoneRotation(FSkeletalMeshAccessorHelper& Accessor, int32 BoneIndex)
 	{
 		const int32 NumRealBones = Accessor.Mesh->GetRefSkeleton().GetRawBoneNum();
 		if (BoneIndex < NumRealBones)
@@ -186,12 +186,12 @@ struct FSkinnedPositionAccessorHelper<TNDISkelMesh_SkinningModeInvalid>
 		return (FQuat4f)RefTransform.GetRotation();
 	}
 
-	FORCEINLINE FQuat4f GetSkinnedBonePreviousRotation(FSkeletalMeshAccessorHelper& Accessor, int32 BoneIndex)
+	inline FQuat4f GetSkinnedBonePreviousRotation(FSkeletalMeshAccessorHelper& Accessor, int32 BoneIndex)
 	{
 		return GetSkinnedBoneRotation(Accessor, BoneIndex);
 	}
 
-	FORCEINLINE FVector3f GetSkinnedBoneScale(FSkeletalMeshAccessorHelper& Accessor, int32 BoneIndex)
+	inline FVector3f GetSkinnedBoneScale(FSkeletalMeshAccessorHelper& Accessor, int32 BoneIndex)
 	{
 		const int32 NumRealBones = Accessor.Mesh->GetRefSkeleton().GetRawBoneNum();
 		if (BoneIndex < NumRealBones)
@@ -203,7 +203,7 @@ struct FSkinnedPositionAccessorHelper<TNDISkelMesh_SkinningModeInvalid>
 		return (FVector3f)RefTransform.GetScale3D();
 	}
 
-	FORCEINLINE FVector3f GetSkinnedBonePreviousScale(FSkeletalMeshAccessorHelper& Accessor, int32 BoneIndex)
+	inline FVector3f GetSkinnedBonePreviousScale(FSkeletalMeshAccessorHelper& Accessor, int32 BoneIndex)
 	{
 		return GetSkinnedBoneScale(Accessor, BoneIndex);
 	}
@@ -212,7 +212,7 @@ struct FSkinnedPositionAccessorHelper<TNDISkelMesh_SkinningModeInvalid>
 template<>
 struct FSkinnedPositionAccessorHelper<TNDISkelMesh_SkinningModeNone>
 {
-	FORCEINLINE int32 GetBoneCount(FSkeletalMeshAccessorHelper& Accessor, bool RequiresPrevious) const
+	inline int32 GetBoneCount(FSkeletalMeshAccessorHelper& Accessor, bool RequiresPrevious) const
 	{
 		if (const USkeletalMesh* Mesh = Accessor.Mesh)
 		{
@@ -222,7 +222,7 @@ struct FSkinnedPositionAccessorHelper<TNDISkelMesh_SkinningModeNone>
 		return 0;
 	}
 
-	FORCEINLINE void GetTriangleIndices(FSkeletalMeshAccessorHelper& Accessor, int32 Tri, int32& Idx0, int32& Idx1, int32& Idx2)
+	inline void GetTriangleIndices(FSkeletalMeshAccessorHelper& Accessor, int32 Tri, int32& Idx0, int32& Idx1, int32& Idx2)
 	{
 		const int32 BaseIndex = Tri * 3;
 		check(BaseIndex + 2 < Accessor.IndexBuffer->Num());
@@ -231,41 +231,41 @@ struct FSkinnedPositionAccessorHelper<TNDISkelMesh_SkinningModeNone>
 		Idx2 = Accessor.IndexBuffer->Get(BaseIndex + 2);
 	}
 
-	FORCEINLINE void GetSkinnedTrianglePositions(FSkeletalMeshAccessorHelper& Accessor, int32 Idx0, int32 Idx1, int32 Idx2, FVector3f& OutPos0, FVector3f& OutPos1, FVector3f& OutPos2)
+	inline void GetSkinnedTrianglePositions(FSkeletalMeshAccessorHelper& Accessor, int32 Idx0, int32 Idx1, int32 Idx2, FVector3f& OutPos0, FVector3f& OutPos1, FVector3f& OutPos2)
 	{
 		OutPos0 = (FVector3f)GetSkeletalMeshRefVertLocation(Accessor.Mesh, *Accessor.LODData, *Accessor.SkinWeightBuffer, Idx0);
 		OutPos1 = (FVector3f)GetSkeletalMeshRefVertLocation(Accessor.Mesh, *Accessor.LODData, *Accessor.SkinWeightBuffer, Idx1);
 		OutPos2 = (FVector3f)GetSkeletalMeshRefVertLocation(Accessor.Mesh, *Accessor.LODData, *Accessor.SkinWeightBuffer, Idx2);
 	}
 
-	FORCEINLINE void GetSkinnedTrianglePreviousPositions(FSkeletalMeshAccessorHelper& Accessor, int32 Idx0, int32 Idx1, int32 Idx2, FVector3f& OutPos0, FVector3f& OutPos1, FVector3f& OutPos2)
+	inline void GetSkinnedTrianglePreviousPositions(FSkeletalMeshAccessorHelper& Accessor, int32 Idx0, int32 Idx1, int32 Idx2, FVector3f& OutPos0, FVector3f& OutPos1, FVector3f& OutPos2)
 	{
 		OutPos0 = (FVector3f)GetSkeletalMeshRefVertLocation(Accessor.Mesh, *Accessor.LODData, *Accessor.SkinWeightBuffer, Idx0);
 		OutPos1 = (FVector3f)GetSkeletalMeshRefVertLocation(Accessor.Mesh, *Accessor.LODData, *Accessor.SkinWeightBuffer, Idx1);
 		OutPos2 = (FVector3f)GetSkeletalMeshRefVertLocation(Accessor.Mesh, *Accessor.LODData, *Accessor.SkinWeightBuffer, Idx2);
 	}
 
-	FORCEINLINE FVector3f GetSkinnedVertexPosition(FSkeletalMeshAccessorHelper& Accessor, int32 VertexIndex)
+	inline FVector3f GetSkinnedVertexPosition(FSkeletalMeshAccessorHelper& Accessor, int32 VertexIndex)
 	{
 		return (FVector3f)GetSkeletalMeshRefVertLocation(Accessor.Mesh, *Accessor.LODData, *Accessor.SkinWeightBuffer, VertexIndex);
 	}
 
-	FORCEINLINE FVector3f GetSkinnedVertexPreviousPosition(FSkeletalMeshAccessorHelper& Accessor, int32 VertexIndex)
+	inline FVector3f GetSkinnedVertexPreviousPosition(FSkeletalMeshAccessorHelper& Accessor, int32 VertexIndex)
 	{
 		return (FVector3f)GetSkeletalMeshRefVertLocation(Accessor.Mesh, *Accessor.LODData, *Accessor.SkinWeightBuffer, VertexIndex);
 	}
 
-	FORCEINLINE void GetSkinnedTangentBasis(FSkeletalMeshAccessorHelper& Accessor, int32 VertexIndex, FVector3f& OutTangentX, FVector3f& OutTangentY, FVector3f& OutTangentZ)
+	inline void GetSkinnedTangentBasis(FSkeletalMeshAccessorHelper& Accessor, int32 VertexIndex, FVector3f& OutTangentX, FVector3f& OutTangentY, FVector3f& OutTangentZ)
 	{
 		GetSkeletalMeshRefTangentBasis(Accessor.Mesh, *Accessor.LODData, *Accessor.SkinWeightBuffer, VertexIndex, OutTangentX, OutTangentY, OutTangentZ);
 	}
 
-	FORCEINLINE void GetSkinnedPreviousTangentBasis(FSkeletalMeshAccessorHelper& Accessor, int32 VertexIndex, FVector3f& OutTangentX, FVector3f& OutTangentY, FVector3f& OutTangentZ)
+	inline void GetSkinnedPreviousTangentBasis(FSkeletalMeshAccessorHelper& Accessor, int32 VertexIndex, FVector3f& OutTangentX, FVector3f& OutTangentY, FVector3f& OutTangentZ)
 	{
 		GetSkeletalMeshRefTangentBasis(Accessor.Mesh, *Accessor.LODData, *Accessor.SkinWeightBuffer, VertexIndex, OutTangentX, OutTangentY, OutTangentZ);
 	}
 
-	FORCEINLINE FVector3f GetSkinnedBonePosition(FSkeletalMeshAccessorHelper& Accessor, int32 BoneIndex)
+	inline FVector3f GetSkinnedBonePosition(FSkeletalMeshAccessorHelper& Accessor, int32 BoneIndex)
 	{
 		const int32 NumRealBones = Accessor.Mesh->GetRefSkeleton().GetRawBoneNum();
 		if (BoneIndex < NumRealBones)
@@ -277,12 +277,12 @@ struct FSkinnedPositionAccessorHelper<TNDISkelMesh_SkinningModeNone>
 		return (FVector3f)RefTransform.GetLocation();
 	}
 
-	FORCEINLINE FVector3f GetSkinnedBonePreviousPosition(FSkeletalMeshAccessorHelper& Accessor, int32 BoneIndex)
+	inline FVector3f GetSkinnedBonePreviousPosition(FSkeletalMeshAccessorHelper& Accessor, int32 BoneIndex)
 	{
 		return GetSkinnedBonePosition(Accessor, BoneIndex);
 	}
 
-	FORCEINLINE FQuat4f GetSkinnedBoneRotation(FSkeletalMeshAccessorHelper& Accessor, int32 BoneIndex)
+	inline FQuat4f GetSkinnedBoneRotation(FSkeletalMeshAccessorHelper& Accessor, int32 BoneIndex)
 	{
 		const int32 NumRealBones = Accessor.Mesh->GetRefSkeleton().GetRawBoneNum();
 		if (BoneIndex < NumRealBones)
@@ -294,12 +294,12 @@ struct FSkinnedPositionAccessorHelper<TNDISkelMesh_SkinningModeNone>
 		return (FQuat4f)RefTransform.GetRotation();
 	}
 
-	FORCEINLINE FQuat4f GetSkinnedBonePreviousRotation(FSkeletalMeshAccessorHelper& Accessor, int32 BoneIndex)
+	inline FQuat4f GetSkinnedBonePreviousRotation(FSkeletalMeshAccessorHelper& Accessor, int32 BoneIndex)
 	{
 		return GetSkinnedBoneRotation(Accessor, BoneIndex);
 	}
 
-	FORCEINLINE FVector3f GetSkinnedBoneScale(FSkeletalMeshAccessorHelper& Accessor, int32 BoneIndex)
+	inline FVector3f GetSkinnedBoneScale(FSkeletalMeshAccessorHelper& Accessor, int32 BoneIndex)
 	{
 		const int32 NumRealBones = Accessor.Mesh->GetRefSkeleton().GetRawBoneNum();
 		if (BoneIndex < NumRealBones)
@@ -311,7 +311,7 @@ struct FSkinnedPositionAccessorHelper<TNDISkelMesh_SkinningModeNone>
 		return (FVector3f)RefTransform.GetScale3D();
 	}
 
-	FORCEINLINE FVector3f GetSkinnedBonePreviousScale(FSkeletalMeshAccessorHelper& Accessor, int32 BoneIndex)
+	inline FVector3f GetSkinnedBonePreviousScale(FSkeletalMeshAccessorHelper& Accessor, int32 BoneIndex)
 	{
 		return GetSkinnedBoneScale(Accessor, BoneIndex);
 	}
@@ -320,7 +320,7 @@ struct FSkinnedPositionAccessorHelper<TNDISkelMesh_SkinningModeNone>
 template<>
 struct FSkinnedPositionAccessorHelper<TNDISkelMesh_SkinningModeOnTheFly>
 {
-	FORCEINLINE int32 GetBoneCount(FSkeletalMeshAccessorHelper& Accessor, bool RequiresPrevious) const
+	inline int32 GetBoneCount(FSkeletalMeshAccessorHelper& Accessor, bool RequiresPrevious) const
 	{
 		if (const FSkeletalMeshSkinningData* SkinningData = Accessor.SkinningData)
 		{
@@ -330,7 +330,7 @@ struct FSkinnedPositionAccessorHelper<TNDISkelMesh_SkinningModeOnTheFly>
 		return 0;
 	}
 
-	FORCEINLINE void GetTriangleIndices(FSkeletalMeshAccessorHelper& Accessor, int32 Tri, int32& Idx0, int32& Idx1, int32& Idx2)
+	inline void GetTriangleIndices(FSkeletalMeshAccessorHelper& Accessor, int32 Tri, int32& Idx0, int32& Idx1, int32& Idx2)
 	{
 		const int32 BaseIndex = Tri * 3;
 		check(BaseIndex + 2 < Accessor.IndexBuffer->Num());
@@ -339,66 +339,66 @@ struct FSkinnedPositionAccessorHelper<TNDISkelMesh_SkinningModeOnTheFly>
 		Idx2 = Accessor.IndexBuffer->Get(BaseIndex + 2);
 	}
 
-	FORCEINLINE void GetSkinnedTrianglePositions(FSkeletalMeshAccessorHelper& Accessor, int32 Idx0, int32 Idx1, int32 Idx2, FVector3f& OutPos0, FVector3f& OutPos1, FVector3f& OutPos2)
+	inline void GetSkinnedTrianglePositions(FSkeletalMeshAccessorHelper& Accessor, int32 Idx0, int32 Idx1, int32 Idx2, FVector3f& OutPos0, FVector3f& OutPos1, FVector3f& OutPos2)
 	{
 		OutPos0 = USkeletalMeshComponent::GetSkinnedVertexPosition(Accessor.Comp, Idx0, *Accessor.LODData, *Accessor.SkinWeightBuffer, Accessor.SkinningData->CurrBoneRefToLocals());
 		OutPos1 = USkeletalMeshComponent::GetSkinnedVertexPosition(Accessor.Comp, Idx1, *Accessor.LODData, *Accessor.SkinWeightBuffer, Accessor.SkinningData->CurrBoneRefToLocals());
 		OutPos2 = USkeletalMeshComponent::GetSkinnedVertexPosition(Accessor.Comp, Idx2, *Accessor.LODData, *Accessor.SkinWeightBuffer, Accessor.SkinningData->CurrBoneRefToLocals());
 	}
 
-	FORCEINLINE void GetSkinnedTrianglePreviousPositions(FSkeletalMeshAccessorHelper& Accessor, int32 Idx0, int32 Idx1, int32 Idx2, FVector3f& OutPos0, FVector3f& OutPos1, FVector3f& OutPos2)
+	inline void GetSkinnedTrianglePreviousPositions(FSkeletalMeshAccessorHelper& Accessor, int32 Idx0, int32 Idx1, int32 Idx2, FVector3f& OutPos0, FVector3f& OutPos1, FVector3f& OutPos2)
 	{
 		OutPos0 = USkeletalMeshComponent::GetSkinnedVertexPosition(Accessor.Comp, Idx0, *Accessor.LODData, *Accessor.SkinWeightBuffer, Accessor.SkinningData->PrevBoneRefToLocals());
 		OutPos1 = USkeletalMeshComponent::GetSkinnedVertexPosition(Accessor.Comp, Idx1, *Accessor.LODData, *Accessor.SkinWeightBuffer, Accessor.SkinningData->PrevBoneRefToLocals());
 		OutPos2 = USkeletalMeshComponent::GetSkinnedVertexPosition(Accessor.Comp, Idx2, *Accessor.LODData, *Accessor.SkinWeightBuffer, Accessor.SkinningData->PrevBoneRefToLocals());
 	}
 
-	FORCEINLINE FVector3f GetSkinnedVertexPosition(FSkeletalMeshAccessorHelper& Accessor, int32 VertexIndex)
+	inline FVector3f GetSkinnedVertexPosition(FSkeletalMeshAccessorHelper& Accessor, int32 VertexIndex)
 	{
 		return USkeletalMeshComponent::GetSkinnedVertexPosition(Accessor.Comp, VertexIndex, *Accessor.LODData, *Accessor.SkinWeightBuffer, Accessor.SkinningData->CurrBoneRefToLocals());
 	}
 
-	FORCEINLINE FVector3f GetSkinnedVertexPreviousPosition(FSkeletalMeshAccessorHelper& Accessor, int32 VertexIndex)
+	inline FVector3f GetSkinnedVertexPreviousPosition(FSkeletalMeshAccessorHelper& Accessor, int32 VertexIndex)
 	{
 		return USkeletalMeshComponent::GetSkinnedVertexPosition(Accessor.Comp, VertexIndex, *Accessor.LODData, *Accessor.SkinWeightBuffer, Accessor.SkinningData->PrevBoneRefToLocals());
 	}
 
-	FORCEINLINE void GetSkinnedTangentBasis(FSkeletalMeshAccessorHelper& Accessor, int32 VertexIndex, FVector3f& OutTangentX, FVector3f& OutTangentY, FVector3f& OutTangentZ)
+	inline void GetSkinnedTangentBasis(FSkeletalMeshAccessorHelper& Accessor, int32 VertexIndex, FVector3f& OutTangentX, FVector3f& OutTangentY, FVector3f& OutTangentZ)
 	{
 		USkeletalMeshComponent::GetSkinnedTangentBasis(Accessor.Comp, VertexIndex, *Accessor.LODData, *Accessor.SkinWeightBuffer, Accessor.SkinningData->CurrBoneRefToLocals(), OutTangentX, OutTangentY, OutTangentZ);
 	}
 
-	FORCEINLINE void GetSkinnedPreviousTangentBasis(FSkeletalMeshAccessorHelper& Accessor, int32 VertexIndex, FVector3f& OutTangentX, FVector3f& OutTangentY, FVector3f& OutTangentZ)
+	inline void GetSkinnedPreviousTangentBasis(FSkeletalMeshAccessorHelper& Accessor, int32 VertexIndex, FVector3f& OutTangentX, FVector3f& OutTangentY, FVector3f& OutTangentZ)
 	{
 		USkeletalMeshComponent::GetSkinnedTangentBasis(Accessor.Comp, VertexIndex, *Accessor.LODData, *Accessor.SkinWeightBuffer, Accessor.SkinningData->PrevBoneRefToLocals(), OutTangentX, OutTangentY, OutTangentZ);
 	}
 
-	FORCEINLINE FVector3f GetSkinnedBonePosition(FSkeletalMeshAccessorHelper& Accessor, int32 BoneIndex)
+	inline FVector3f GetSkinnedBonePosition(FSkeletalMeshAccessorHelper& Accessor, int32 BoneIndex)
 	{
 		return (FVector3f)Accessor.SkinningData->CurrComponentTransforms()[BoneIndex].GetLocation();
 	}
 
-	FORCEINLINE FVector3f GetSkinnedBonePreviousPosition(FSkeletalMeshAccessorHelper& Accessor, int32 BoneIndex)
+	inline FVector3f GetSkinnedBonePreviousPosition(FSkeletalMeshAccessorHelper& Accessor, int32 BoneIndex)
 	{
 		return (FVector3f)Accessor.SkinningData->PrevComponentTransforms()[BoneIndex].GetLocation();
 	}
 
-	FORCEINLINE FQuat4f GetSkinnedBoneRotation(FSkeletalMeshAccessorHelper& Accessor, int32 BoneIndex)
+	inline FQuat4f GetSkinnedBoneRotation(FSkeletalMeshAccessorHelper& Accessor, int32 BoneIndex)
 	{
 		return (FQuat4f)Accessor.SkinningData->CurrComponentTransforms()[BoneIndex].GetRotation();
 	}
 
-	FORCEINLINE FQuat4f GetSkinnedBonePreviousRotation(FSkeletalMeshAccessorHelper& Accessor, int32 BoneIndex)
+	inline FQuat4f GetSkinnedBonePreviousRotation(FSkeletalMeshAccessorHelper& Accessor, int32 BoneIndex)
 	{
 		return (FQuat4f)Accessor.SkinningData->PrevComponentTransforms()[BoneIndex].GetRotation();
 	}
 
-	FORCEINLINE FVector3f GetSkinnedBoneScale(FSkeletalMeshAccessorHelper& Accessor, int32 BoneIndex)
+	inline FVector3f GetSkinnedBoneScale(FSkeletalMeshAccessorHelper& Accessor, int32 BoneIndex)
 	{
 		return (FVector3f)Accessor.SkinningData->CurrComponentTransforms()[BoneIndex].GetScale3D();
 	}
 
-	FORCEINLINE FVector3f GetSkinnedBonePreviousScale(FSkeletalMeshAccessorHelper& Accessor, int32 BoneIndex)
+	inline FVector3f GetSkinnedBonePreviousScale(FSkeletalMeshAccessorHelper& Accessor, int32 BoneIndex)
 	{
 		return (FVector3f)Accessor.SkinningData->PrevComponentTransforms()[BoneIndex].GetScale3D();
 	}
@@ -407,7 +407,7 @@ struct FSkinnedPositionAccessorHelper<TNDISkelMesh_SkinningModeOnTheFly>
 template<>
 struct FSkinnedPositionAccessorHelper<TNDISkelMesh_SkinningModePreSkin>
 {
-	FORCEINLINE int32 GetBoneCount(FSkeletalMeshAccessorHelper& Accessor, bool RequiresPrevious) const
+	inline int32 GetBoneCount(FSkeletalMeshAccessorHelper& Accessor, bool RequiresPrevious) const
 	{
 		if (const FSkeletalMeshSkinningData* SkinningData = Accessor.SkinningData)
 		{
@@ -417,7 +417,7 @@ struct FSkinnedPositionAccessorHelper<TNDISkelMesh_SkinningModePreSkin>
 		return 0;
 	}
 
-	FORCEINLINE void GetTriangleIndices(FSkeletalMeshAccessorHelper& Accessor, int32 Tri, int32& Idx0, int32& Idx1, int32& Idx2)
+	inline void GetTriangleIndices(FSkeletalMeshAccessorHelper& Accessor, int32 Tri, int32& Idx0, int32& Idx1, int32& Idx2)
 	{
 		const int32 BaseIndex = Tri * 3;
 		check(BaseIndex + 2 < Accessor.IndexBuffer->Num());
@@ -426,66 +426,66 @@ struct FSkinnedPositionAccessorHelper<TNDISkelMesh_SkinningModePreSkin>
 		Idx2 = Accessor.IndexBuffer->Get(BaseIndex + 2);
 	}
 
-	FORCEINLINE void GetSkinnedTrianglePositions(FSkeletalMeshAccessorHelper& Accessor, int32 Idx0, int32 Idx1, int32 Idx2, FVector3f& OutPos0, FVector3f& OutPos1, FVector3f& OutPos2)
+	inline void GetSkinnedTrianglePositions(FSkeletalMeshAccessorHelper& Accessor, int32 Idx0, int32 Idx1, int32 Idx2, FVector3f& OutPos0, FVector3f& OutPos1, FVector3f& OutPos2)
 	{
 		OutPos0 = Accessor.SkinningData->GetPosition(Accessor.Usage.GetLODIndex(), Idx0);
 		OutPos1 = Accessor.SkinningData->GetPosition(Accessor.Usage.GetLODIndex(), Idx1);
 		OutPos2 = Accessor.SkinningData->GetPosition(Accessor.Usage.GetLODIndex(), Idx2);
 	}
 
-	FORCEINLINE void GetSkinnedTrianglePreviousPositions(FSkeletalMeshAccessorHelper& Accessor, int32 Idx0, int32 Idx1, int32 Idx2, FVector3f& OutPos0, FVector3f& OutPos1, FVector3f& OutPos2)
+	inline void GetSkinnedTrianglePreviousPositions(FSkeletalMeshAccessorHelper& Accessor, int32 Idx0, int32 Idx1, int32 Idx2, FVector3f& OutPos0, FVector3f& OutPos1, FVector3f& OutPos2)
 	{
 		OutPos0 = Accessor.SkinningData->GetPreviousPosition(Accessor.Usage.GetLODIndex(), Idx0);
 		OutPos1 = Accessor.SkinningData->GetPreviousPosition(Accessor.Usage.GetLODIndex(), Idx1);
 		OutPos2 = Accessor.SkinningData->GetPreviousPosition(Accessor.Usage.GetLODIndex(), Idx2);
 	}
 
-	FORCEINLINE FVector3f GetSkinnedVertexPosition(FSkeletalMeshAccessorHelper& Accessor, int32 VertexIndex)
+	inline FVector3f GetSkinnedVertexPosition(FSkeletalMeshAccessorHelper& Accessor, int32 VertexIndex)
 	{
 		return Accessor.SkinningData->GetPosition(Accessor.Usage.GetLODIndex(), VertexIndex);
 	}
 
-	FORCEINLINE FVector3f GetSkinnedVertexPreviousPosition(FSkeletalMeshAccessorHelper& Accessor, int32 VertexIndex)
+	inline FVector3f GetSkinnedVertexPreviousPosition(FSkeletalMeshAccessorHelper& Accessor, int32 VertexIndex)
 	{
 		return Accessor.SkinningData->GetPreviousPosition(Accessor.Usage.GetLODIndex(), VertexIndex);
 	}
 
-	FORCEINLINE void GetSkinnedTangentBasis(FSkeletalMeshAccessorHelper& Accessor, int32 VertexIndex, FVector3f& OutTangentX, FVector3f& OutTangentY, FVector3f& OutTangentZ)
+	inline void GetSkinnedTangentBasis(FSkeletalMeshAccessorHelper& Accessor, int32 VertexIndex, FVector3f& OutTangentX, FVector3f& OutTangentY, FVector3f& OutTangentZ)
 	{
 		Accessor.SkinningData->GetTangentBasis(Accessor.Usage.GetLODIndex(), VertexIndex, OutTangentX, OutTangentY, OutTangentZ);
 	}
 
-	FORCEINLINE void GetSkinnedPreviousTangentBasis(FSkeletalMeshAccessorHelper& Accessor, int32 VertexIndex, FVector3f& OutTangentX, FVector3f& OutTangentY, FVector3f& OutTangentZ)
+	inline void GetSkinnedPreviousTangentBasis(FSkeletalMeshAccessorHelper& Accessor, int32 VertexIndex, FVector3f& OutTangentX, FVector3f& OutTangentY, FVector3f& OutTangentZ)
 	{
 		Accessor.SkinningData->GetPreviousTangentBasis(Accessor.Usage.GetLODIndex(), VertexIndex, OutTangentX, OutTangentY, OutTangentZ);
 	}
 
-	FORCEINLINE FVector3f GetSkinnedBonePosition(FSkeletalMeshAccessorHelper& Accessor, int32 BoneIndex)
+	inline FVector3f GetSkinnedBonePosition(FSkeletalMeshAccessorHelper& Accessor, int32 BoneIndex)
 	{
 		return Accessor.SkinningData->CurrComponentTransforms()[BoneIndex].GetLocation();
 	}
 
-	FORCEINLINE FVector3f GetSkinnedBonePreviousPosition(FSkeletalMeshAccessorHelper& Accessor, int32 BoneIndex)
+	inline FVector3f GetSkinnedBonePreviousPosition(FSkeletalMeshAccessorHelper& Accessor, int32 BoneIndex)
 	{
 		return Accessor.SkinningData->PrevComponentTransforms()[BoneIndex].GetLocation();
 	}
 
-	FORCEINLINE FQuat4f GetSkinnedBoneRotation(FSkeletalMeshAccessorHelper& Accessor, int32 BoneIndex)
+	inline FQuat4f GetSkinnedBoneRotation(FSkeletalMeshAccessorHelper& Accessor, int32 BoneIndex)
 	{
 		return (FQuat4f)Accessor.SkinningData->CurrComponentTransforms()[BoneIndex].GetRotation();
 	}
 
-	FORCEINLINE FQuat4f GetSkinnedBonePreviousRotation(FSkeletalMeshAccessorHelper& Accessor, int32 BoneIndex)
+	inline FQuat4f GetSkinnedBonePreviousRotation(FSkeletalMeshAccessorHelper& Accessor, int32 BoneIndex)
 	{
 		return (FQuat4f)Accessor.SkinningData->PrevComponentTransforms()[BoneIndex].GetRotation();
 	}
 
-	FORCEINLINE FVector3f GetSkinnedBoneScale(FSkeletalMeshAccessorHelper& Accessor, int32 BoneIndex)
+	inline FVector3f GetSkinnedBoneScale(FSkeletalMeshAccessorHelper& Accessor, int32 BoneIndex)
 	{
 		return (FVector3f)Accessor.SkinningData->CurrComponentTransforms()[BoneIndex].GetScale3D();
 	}
 
-	FORCEINLINE FVector3f GetSkinnedBonePreviousScale(FSkeletalMeshAccessorHelper& Accessor, int32 BoneIndex)
+	inline FVector3f GetSkinnedBonePreviousScale(FSkeletalMeshAccessorHelper& Accessor, int32 BoneIndex)
 	{
 		return (FVector3f)Accessor.SkinningData->PrevComponentTransforms()[BoneIndex].GetScale3D();
 	}
@@ -496,7 +496,7 @@ struct FSkinnedPositionAccessorHelper<TNDISkelMesh_SkinningModePreSkin>
 template<bool bUseFullPrecisionUVs>
 struct FSkelMeshVertexAccessor
 {
-	FORCEINLINE FVector2D GetVertexUV(const FSkeletalMeshLODRenderData* LODData, int32 VertexIdx, int32 UVChannel)const
+	inline FVector2D GetVertexUV(const FSkeletalMeshLODRenderData* LODData, int32 VertexIdx, int32 UVChannel)const
 	{
 		check(LODData);
 		if (bUseFullPrecisionUVs)
@@ -509,7 +509,7 @@ struct FSkelMeshVertexAccessor
 		}
 	}
 
-	FORCEINLINE FLinearColor GetVertexColor(FSkeletalMeshLODRenderData* LODData, int32 VertexIdx)const
+	inline FLinearColor GetVertexColor(FSkeletalMeshLODRenderData* LODData, int32 VertexIdx)const
 	{
 		check(LODData);
 		return LODData->StaticVertexBuffers.ColorVertexBuffer.VertexColor(VertexIdx);
@@ -518,12 +518,12 @@ struct FSkelMeshVertexAccessor
 
 struct FSkelMeshVertexAccessorNoop
 {
-	FORCEINLINE FVector2D GetVertexUV(const FSkeletalMeshLODRenderData* LODData, int32 VertexIdx, int32 UVChannel)const
+	inline FVector2D GetVertexUV(const FSkeletalMeshLODRenderData* LODData, int32 VertexIdx, int32 UVChannel)const
 	{
 		return FVector2D(0.0f, 0.0f);
 	}
 
-	FORCEINLINE FLinearColor GetVertexColor(const FSkeletalMeshLODRenderData* LODData, int32 VertexIdx)const
+	inline FLinearColor GetVertexColor(const FSkeletalMeshLODRenderData* LODData, int32 VertexIdx)const
 	{
 		return FLinearColor::White;
 	}
diff --git a/Engine/Plugins/FX/Niagara/Source/Niagara/Classes/NiagaraDataInterface.h b/Engine/Plugins/FX/Niagara/Source/Niagara/Classes/NiagaraDataInterface.h
--- a/Engine/Plugins/FX/Niagara/Source/Niagara/Classes/NiagaraDataInterface.h
+++ b/Engine/Plugins/FX/Niagara/Source/Niagara/Classes/NiagaraDataInterface.h
@@ -32,36 +32,36 @@ namespace ERHIFeatureLevel { enum Type : int; }
 
 struct FNDITransformHandlerNoop
 {
-	FORCEINLINE void TransformPosition(FVector3f& V, const FMatrix44f& M) const { }
-	FORCEINLINE void TransformPosition(FVector3d& V, const FMatrix44d& M) const { }
-	FORCEINLINE void TransformUnitVector(FVector3f& V, const FMatrix44f& M) const { }
-	FORCEINLINE void TransformUnitVector(FVector3d& V, const FMatrix44d& M) const { }
-	FORCEINLINE void TransformNotUnitVector(FVector3f& V, const FMatrix44f& M) const { }
-	FORCEINLINE void TransformNotUnitVector(FVector3d& V, const FMatrix44d& M) const { }
-	FORCEINLINE void TransformRotation(FQuat4f& Q1, const FQuat4f& Q2) const { }
-	FORCEINLINE void TransformRotation(FQuat4d& Q1, const FQuat4d& Q2) const { }
+	inline void TransformPosition(FVector3f& V, const FMatrix44f& M) const { }
+	inline void TransformPosition(FVector3d& V, const FMatrix44d& M) const { }
+	inline void TransformUnitVector(FVector3f& V, const FMatrix44f& M) const { }
+	inline void TransformUnitVector(FVector3d& V, const FMatrix44d& M) const { }
+	inline void TransformNotUnitVector(FVector3f& V, const FMatrix44f& M) const { }
+	inline void TransformNotUnitVector(FVector3d& V, const FMatrix44d& M) const { }
+	inline void TransformRotation(FQuat4f& Q1, const FQuat4f& Q2) const { }
+	inline void TransformRotation(FQuat4d& Q1, const FQuat4d& Q2) const { }
 
 	UE_DEPRECATED(5.4, "Please update your code to use TransformUnitVector")
-	FORCEINLINE void TransformVector(FVector3f& V, const FMatrix44f& M) const { }
+	inline void TransformVector(FVector3f& V, const FMatrix44f& M) const { }
 	UE_DEPRECATED(5.4, "Please update your code to use TransformUnitVector")
-	FORCEINLINE void TransformVector(FVector3d& V, const FMatrix44d& M) const { }
+	inline void TransformVector(FVector3d& V, const FMatrix44d& M) const { }
 };
 
 struct FNDITransformHandler
 {
-	FORCEINLINE void TransformPosition(FVector3f& P, const FMatrix44f& M) const { P = M.TransformPosition(P); }
-	FORCEINLINE void TransformPosition(FVector3d& P, const FMatrix44d& M) const { P = M.TransformPosition(P); }
-	FORCEINLINE void TransformUnitVector(FVector3f& V, const FMatrix44f& M) const { V = M.TransformVector(V).GetUnsafeNormal3(); }
-	FORCEINLINE void TransformUnitVector(FVector3d& V, const FMatrix44d& M) const { V = M.TransformVector(V).GetUnsafeNormal3(); }
-	FORCEINLINE void TransformNotUnitVector(FVector3f& V, const FMatrix44f& M) const { V = M.TransformVector(V); }
-	FORCEINLINE void TransformNotUnitVector(FVector3d& V, const FMatrix44d& M) const { V = M.TransformVector(V); }
-	FORCEINLINE void TransformRotation(FQuat4f& Q1, const FQuat4f& Q2) const { Q1 = Q2 * Q1; }
-	FORCEINLINE void TransformRotation(FQuat4d& Q1, const FQuat4d& Q2) const { Q1 = Q2 * Q1; }
+	inline void TransformPosition(FVector3f& P, const FMatrix44f& M) const { P = M.TransformPosition(P); }
+	inline void TransformPosition(FVector3d& P, const FMatrix44d& M) const { P = M.TransformPosition(P); }
+	inline void TransformUnitVector(FVector3f& V, const FMatrix44f& M) const { V = M.TransformVector(V).GetUnsafeNormal3(); }
+	inline void TransformUnitVector(FVector3d& V, const FMatrix44d& M) const { V = M.TransformVector(V).GetUnsafeNormal3(); }
+	inline void TransformNotUnitVector(FVector3f& V, const FMatrix44f& M) const { V = M.TransformVector(V); }
+	inline void TransformNotUnitVector(FVector3d& V, const FMatrix44d& M) const { V = M.TransformVector(V); }
+	inline void TransformRotation(FQuat4f& Q1, const FQuat4f& Q2) const { Q1 = Q2 * Q1; }
+	inline void TransformRotation(FQuat4d& Q1, const FQuat4d& Q2) const { Q1 = Q2 * Q1; }
 
 	UE_DEPRECATED(5.4, "Please update your code to use TransformUnitVector")
-	FORCEINLINE void TransformVector(FVector3f& V, const FMatrix44f& M) const { V = M.TransformVector(V).GetUnsafeNormal3(); }
+	inline void TransformVector(FVector3f& V, const FMatrix44f& M) const { V = M.TransformVector(V).GetUnsafeNormal3(); }
 	UE_DEPRECATED(5.4, "Please update your code to use TransformUnitVector")
-	FORCEINLINE void TransformVector(FVector3d& V, const FMatrix44d& M) const { V = M.TransformVector(V).GetUnsafeNormal3(); }
+	inline void TransformVector(FVector3d& V, const FMatrix44d& M) const { V = M.TransformVector(V).GetUnsafeNormal3(); }
 };
 
 // FNiagaraDataInterfaceProxy should always outlive any ticks, etc, that are on the rendering thread.
@@ -187,12 +187,12 @@ public:
 		return ErrorSummaryText;
 	}
 
-	FORCEINLINE bool operator !=(const FNiagaraDataInterfaceError& Other) const
+	inline bool operator !=(const FNiagaraDataInterfaceError& Other) const
 	{
 		return !(*this == Other);
 	}
 
-	FORCEINLINE bool operator == (const FNiagaraDataInterfaceError& Other) const
+	inline bool operator == (const FNiagaraDataInterfaceError& Other) const
 	{
 		return ErrorText.EqualTo(Other.ErrorText) && ErrorSummaryText.EqualTo(Other.ErrorSummaryText);
 	}
@@ -244,12 +244,12 @@ public:
 		return FeedbackSummaryText;
 	}
 
-	FORCEINLINE bool operator !=(const FNiagaraDataInterfaceFeedback& Other) const
+	inline bool operator !=(const FNiagaraDataInterfaceFeedback& Other) const
 	{
 		return !(*this == Other);
 	}
 
-	FORCEINLINE bool operator == (const FNiagaraDataInterfaceFeedback& Other) const
+	inline bool operator == (const FNiagaraDataInterfaceFeedback& Other) const
 	{
 		return FeedbackText.EqualTo(Other.FeedbackText) && FeedbackSummaryText.EqualTo(Other.FeedbackSummaryText);
 	}
@@ -953,8 +953,8 @@ template<typename T>
 struct FNDIParameter
 {
 	FNDIParameter(FVectorVMExternalFunctionContext& Context) = delete;
-	FORCEINLINE void GetAndAdvance(T& OutValue) = delete;
-	FORCEINLINE bool IsConstant() const = delete;
+	inline void GetAndAdvance(T& OutValue) = delete;
+	inline bool IsConstant() const = delete;
 };
 
 template<>
@@ -974,7 +974,7 @@ struct FNDIParameter<FNiagaraRandInfo>
 		Seed3Param.Init(Context);
 	}
 
-	FORCEINLINE void GetAndAdvance(FNiagaraRandInfo& OutValue)
+	inline void GetAndAdvance(FNiagaraRandInfo& OutValue)
 	{
 		OutValue.Seed1 = Seed1Param.GetAndAdvance();
 		OutValue.Seed2 = Seed2Param.GetAndAdvance();
@@ -982,7 +982,7 @@ struct FNDIParameter<FNiagaraRandInfo>
 	}
 
 
-	FORCEINLINE bool IsConstant()const
+	inline bool IsConstant()const
 	{
 		return Seed1Param.IsConstant() && Seed2Param.IsConstant() && Seed3Param.IsConstant();
 	}
@@ -993,12 +993,12 @@ struct FNDIRandomStreamPolicy
 {
 	FNDIRandomStreamPolicy(FVectorVMExternalFunctionContext& InContext) : Context(InContext) {}
 
-	FORCEINLINE void GetAndAdvance() {}
-	FORCEINLINE bool IsDeterministic() const { return false; }
-	FORCEINLINE_DEBUGGABLE FVector4 Rand4(int32 InstanceIndex) const { return FVector4(Context.GetRandStream().GetFraction(), Context.GetRandStream().GetFraction(), Context.GetRandStream().GetFraction(), Context.GetRandStream().GetFraction()); }
-	FORCEINLINE_DEBUGGABLE FVector Rand3(int32 InstanceIndex) const { return FVector(Context.GetRandStream().GetFraction(), Context.GetRandStream().GetFraction(), Context.GetRandStream().GetFraction()); }
-	FORCEINLINE_DEBUGGABLE FVector2D Rand2(int32 InstanceIndex) const { return FVector2D(Context.GetRandStream().GetFraction(), Context.GetRandStream().GetFraction()); }
-	FORCEINLINE_DEBUGGABLE float Rand(int32 InstanceIndex) const { return Context.GetRandStream().GetFraction(); }
+	inline void GetAndAdvance() {}
+	inline bool IsDeterministic() const { return false; }
+	inline FVector4 Rand4(int32 InstanceIndex) const { return FVector4(Context.GetRandStream().GetFraction(), Context.GetRandStream().GetFraction(), Context.GetRandStream().GetFraction(), Context.GetRandStream().GetFraction()); }
+	inline FVector Rand3(int32 InstanceIndex) const { return FVector(Context.GetRandStream().GetFraction(), Context.GetRandStream().GetFraction(), Context.GetRandStream().GetFraction()); }
+	inline FVector2D Rand2(int32 InstanceIndex) const { return FVector2D(Context.GetRandStream().GetFraction(), Context.GetRandStream().GetFraction()); }
+	inline float Rand(int32 InstanceIndex) const { return Context.GetRandStream().GetFraction(); }
 
 	FVectorVMExternalFunctionContext& Context;
 };
@@ -1011,11 +1011,11 @@ struct FNDIRandomInfoPolicy
 		, RandParam(InContext)
 	{}
 
-	FORCEINLINE void GetAndAdvance() { RandParam.GetAndAdvance(RandInfo); }
-	FORCEINLINE bool IsDeterministic() const { return RandInfo.Seed3 != INDEX_NONE; }
+	inline void GetAndAdvance() { RandParam.GetAndAdvance(RandInfo); }
+	inline bool IsDeterministic() const { return RandInfo.Seed3 != INDEX_NONE; }
 
 	//////////////////////////////////////////////////////////////////////////
-	FORCEINLINE_DEBUGGABLE FVector4f Rand4(int32 InstanceIndex) const
+	inline FVector4f Rand4(int32 InstanceIndex) const
 	{
 		if (IsDeterministic())
 		{
@@ -1051,7 +1051,7 @@ struct FNDIRandomInfoPolicy
 		}
 	}
 
-	FORCEINLINE_DEBUGGABLE FVector3f Rand3(int32 InstanceIndex) const
+	inline FVector3f Rand3(int32 InstanceIndex) const
 	{
 		if (IsDeterministic())
 		{
@@ -1078,7 +1078,7 @@ struct FNDIRandomInfoPolicy
 		}
 	}
 
-	FORCEINLINE_DEBUGGABLE FVector2f Rand2(int32 InstanceIndex) const
+	inline FVector2f Rand2(int32 InstanceIndex) const
 	{
 		if (IsDeterministic())
 		{
@@ -1091,7 +1091,7 @@ struct FNDIRandomInfoPolicy
 		}
 	}
 
-	FORCEINLINE_DEBUGGABLE float Rand(int32 InstanceIndex) const
+	inline float Rand(int32 InstanceIndex) const
 	{
 		if (IsDeterministic())
 		{
@@ -1116,31 +1116,31 @@ struct TNDIRandomHelper : public TRandomPolicy
 	{
 	}
 
-	FORCEINLINE_DEBUGGABLE FVector4 RandRange(int32 InstanceIndex, FVector4 Min, FVector4 Max) const
+	inline FVector4 RandRange(int32 InstanceIndex, FVector4 Min, FVector4 Max) const
 	{
 		FVector4 Range = Max - Min;
 		return Min + (TRandomPolicy::Rand(InstanceIndex) * Range);
 	}
 
-	FORCEINLINE_DEBUGGABLE FVector RandRange(int32 InstanceIndex, FVector Min, FVector Max) const
+	inline FVector RandRange(int32 InstanceIndex, FVector Min, FVector Max) const
 	{
 		FVector Range = Max - Min;
 		return Min + (TRandomPolicy::Rand(InstanceIndex) * Range);
 	}
 
-	FORCEINLINE_DEBUGGABLE FVector2D RandRange(int32 InstanceIndex, FVector2D Min, FVector2D Max) const
+	inline FVector2D RandRange(int32 InstanceIndex, FVector2D Min, FVector2D Max) const
 	{
 		FVector2D Range = Max - Min;
 		return Min + (TRandomPolicy::Rand(InstanceIndex) * Range);
 	}
 
-	FORCEINLINE_DEBUGGABLE float RandRange(int32 InstanceIndex, float Min, float Max) const
+	inline float RandRange(int32 InstanceIndex, float Min, float Max) const
 	{
 		float Range = Max - Min;
 		return Min + (TRandomPolicy::Rand(InstanceIndex) * Range);
 	}
 
-	FORCEINLINE_DEBUGGABLE int32 RandRange(int32 InstanceIndex, int32 Min, int32 Max) const
+	inline int32 RandRange(int32 InstanceIndex, int32 Min, int32 Max) const
 	{
 		// NOTE: Scaling a uniform float range provides better distribution of 
 		//       numbers than using %.
@@ -1149,7 +1149,7 @@ struct TNDIRandomHelper : public TRandomPolicy
 		return Min + (int(TRandomPolicy::Rand(InstanceIndex) * (Range + 1)));
 	}
 
-	FORCEINLINE_DEBUGGABLE FVector3f RandomBarycentricCoord(int32 InstanceIndex) const
+	inline FVector3f RandomBarycentricCoord(int32 InstanceIndex) const
 	{
 		//TODO: This is gonna be slooooow. Move to an LUT possibly or find faster method.
 		//Can probably handle lower quality randoms / uniformity for a decent speed win.
@@ -1168,56 +1168,56 @@ struct FNDIInputParam
 {
 	static_assert(sizeof(T) == sizeof(float), "Generic template assumes 4 bytes per element");
 	VectorVM::FExternalFuncInputHandler<T> Data;	
-	FORCEINLINE FNDIInputParam(FVectorVMExternalFunctionContext& Context) : Data(Context) {}
-	FORCEINLINE FNDIInputParam(){ }
-	FORCEINLINE void Init(FVectorVMExternalFunctionContext& Context){ new(this)FNDIInputParam<T>(Context); }
-	FORCEINLINE T GetAndAdvance() { return Data.GetAndAdvance(); }
-	FORCEINLINE T Get() { return Data.Get(); }
-	FORCEINLINE void Advance(int32 Count=1) { return Data.Advance(Count); }
-	FORCEINLINE bool IsConstant() const { return Data.IsConstant(); }
-	FORCEINLINE void Reset() { Data.Reset(); }
+	inline FNDIInputParam(FVectorVMExternalFunctionContext& Context) : Data(Context) {}
+	inline FNDIInputParam(){ }
+	inline void Init(FVectorVMExternalFunctionContext& Context){ new(this)FNDIInputParam<T>(Context); }
+	inline T GetAndAdvance() { return Data.GetAndAdvance(); }
+	inline T Get() { return Data.Get(); }
+	inline void Advance(int32 Count=1) { return Data.Advance(Count); }
+	inline bool IsConstant() const { return Data.IsConstant(); }
+	inline void Reset() { Data.Reset(); }
 };
 
 template<>
 struct FNDIInputParam<FFloat16>
 {
 	VectorVM::FExternalFuncInputHandler<FFloat16> Data;
-	FORCEINLINE FNDIInputParam(FVectorVMExternalFunctionContext& Context) : Data(Context) {}
-	FORCEINLINE FNDIInputParam() { }
-	FORCEINLINE void Init(FVectorVMExternalFunctionContext& Context) { new(this)FNDIInputParam<FFloat16>(Context); }
-	FORCEINLINE FFloat16 GetAndAdvance() { return Data.GetAndAdvance(); }
-	FORCEINLINE FFloat16 Get() { return Data.Get(); }
-	FORCEINLINE void Advance(int32 Count=1) { return Data.Advance(Count); }
-	FORCEINLINE bool IsConstant() const { return Data.IsConstant(); }
-	FORCEINLINE void Reset() { Data.Reset(); }
+	inline FNDIInputParam(FVectorVMExternalFunctionContext& Context) : Data(Context) {}
+	inline FNDIInputParam() { }
+	inline void Init(FVectorVMExternalFunctionContext& Context) { new(this)FNDIInputParam<FFloat16>(Context); }
+	inline FFloat16 GetAndAdvance() { return Data.GetAndAdvance(); }
+	inline FFloat16 Get() { return Data.Get(); }
+	inline void Advance(int32 Count=1) { return Data.Advance(Count); }
+	inline bool IsConstant() const { return Data.IsConstant(); }
+	inline void Reset() { Data.Reset(); }
 };
 
 template<>
 struct FNDIInputParam<FNiagaraBool>
 {
 	VectorVM::FExternalFuncInputHandler<FNiagaraBool> Data;
-	FORCEINLINE FNDIInputParam(FVectorVMExternalFunctionContext& Context) : Data(Context) {}
-	FORCEINLINE FNDIInputParam() { }
-	FORCEINLINE void Init(FVectorVMExternalFunctionContext& Context) { new(this)FNDIInputParam<FNiagaraBool>(Context); }
-	FORCEINLINE bool GetAndAdvance() { return Data.GetAndAdvance().GetValue(); }
-	FORCEINLINE bool Get() { return Data.Get().GetValue(); }
-	FORCEINLINE void Advance(int32 Count=1) { return Data.Advance(Count); }
-	FORCEINLINE bool IsConstant() const { return Data.IsConstant(); }
-	FORCEINLINE void Reset() { Data.Reset(); }
+	inline FNDIInputParam(FVectorVMExternalFunctionContext& Context) : Data(Context) {}
+	inline FNDIInputParam() { }
+	inline void Init(FVectorVMExternalFunctionContext& Context) { new(this)FNDIInputParam<FNiagaraBool>(Context); }
+	inline bool GetAndAdvance() { return Data.GetAndAdvance().GetValue(); }
+	inline bool Get() { return Data.Get().GetValue(); }
+	inline void Advance(int32 Count=1) { return Data.Advance(Count); }
+	inline bool IsConstant() const { return Data.IsConstant(); }
+	inline void Reset() { Data.Reset(); }
 };
 
 template<>
 struct FNDIInputParam<bool>
 {
 	VectorVM::FExternalFuncInputHandler<FNiagaraBool> Data;
-	FORCEINLINE FNDIInputParam(FVectorVMExternalFunctionContext& Context) : Data(Context) {}
-	FORCEINLINE FNDIInputParam() { }
-	FORCEINLINE void Init(FVectorVMExternalFunctionContext& Context) { new(this)FNDIInputParam<bool>(Context); }
-	FORCEINLINE bool GetAndAdvance() { return Data.GetAndAdvance().GetValue(); }
-	FORCEINLINE bool Get() { return Data.Get().GetValue(); }
-	FORCEINLINE void Advance(int32 Count=1) { return Data.Advance(Count); }
-	FORCEINLINE bool IsConstant() const { return Data.IsConstant(); }
-	FORCEINLINE void Reset() { Data.Reset(); }
+	inline FNDIInputParam(FVectorVMExternalFunctionContext& Context) : Data(Context) {}
+	inline FNDIInputParam() { }
+	inline void Init(FVectorVMExternalFunctionContext& Context) { new(this)FNDIInputParam<bool>(Context); }
+	inline bool GetAndAdvance() { return Data.GetAndAdvance().GetValue(); }
+	inline bool Get() { return Data.Get().GetValue(); }
+	inline void Advance(int32 Count=1) { return Data.Advance(Count); }
+	inline bool IsConstant() const { return Data.IsConstant(); }
+	inline void Reset() { Data.Reset(); }
 };
 
 template<>
@@ -1225,14 +1225,14 @@ struct FNDIInputParam<FVector2f>
 {
 	VectorVM::FExternalFuncInputHandler<float> X;
 	VectorVM::FExternalFuncInputHandler<float> Y;
-	FORCEINLINE FNDIInputParam(FVectorVMExternalFunctionContext& Context) : X(Context), Y(Context) {}
-	FORCEINLINE FNDIInputParam() { }
-	FORCEINLINE void Init(FVectorVMExternalFunctionContext& Context) { new(this)FNDIInputParam<FVector2f>(Context); }
-	FORCEINLINE FVector2f GetAndAdvance() { return FVector2f(X.GetAndAdvance(), Y.GetAndAdvance()); }
-	FORCEINLINE FVector2f Get() { return FVector2f(X.Get(), Y.Get()); }
-	FORCEINLINE void Advance(int32 Count=1) { X.Advance(Count), Y.Advance(Count); }
-	FORCEINLINE bool IsConstant() const { return X.IsConstant() && Y.IsConstant(); }
-	FORCEINLINE void Reset() { X.Reset(); Y.Reset(); }
+	inline FNDIInputParam(FVectorVMExternalFunctionContext& Context) : X(Context), Y(Context) {}
+	inline FNDIInputParam() { }
+	inline void Init(FVectorVMExternalFunctionContext& Context) { new(this)FNDIInputParam<FVector2f>(Context); }
+	inline FVector2f GetAndAdvance() { return FVector2f(X.GetAndAdvance(), Y.GetAndAdvance()); }
+	inline FVector2f Get() { return FVector2f(X.Get(), Y.Get()); }
+	inline void Advance(int32 Count=1) { X.Advance(Count), Y.Advance(Count); }
+	inline bool IsConstant() const { return X.IsConstant() && Y.IsConstant(); }
+	inline void Reset() { X.Reset(); Y.Reset(); }
 };
 
 template<>
@@ -1242,13 +1242,13 @@ struct FNDIInputParam<FVector3f>
 	VectorVM::FExternalFuncInputHandler<float> Y;
 	VectorVM::FExternalFuncInputHandler<float> Z;
 	FNDIInputParam(FVectorVMExternalFunctionContext& Context) : X(Context), Y(Context), Z(Context) {}
-	FORCEINLINE FNDIInputParam() { }
-	FORCEINLINE void Init(FVectorVMExternalFunctionContext& Context) { new(this)FNDIInputParam<FVector3f>(Context); }
-	FORCEINLINE FVector3f GetAndAdvance() { return FVector3f(X.GetAndAdvance(), Y.GetAndAdvance(), Z.GetAndAdvance()); }
-	FORCEINLINE FVector3f Get() { return FVector3f(X.Get(), Y.Get(), Z.Get()); }
-	FORCEINLINE void Advance(int32 Count=1) { X.Advance(Count); Y.Advance(Count); Z.Advance(Count); }
-	FORCEINLINE bool IsConstant() const { return X.IsConstant() && Y.IsConstant() && Z.IsConstant(); }
-	FORCEINLINE void Reset() { X.Reset(); Y.Reset(); Z.Reset(); }
+	inline FNDIInputParam() { }
+	inline void Init(FVectorVMExternalFunctionContext& Context) { new(this)FNDIInputParam<FVector3f>(Context); }
+	inline FVector3f GetAndAdvance() { return FVector3f(X.GetAndAdvance(), Y.GetAndAdvance(), Z.GetAndAdvance()); }
+	inline FVector3f Get() { return FVector3f(X.Get(), Y.Get(), Z.Get()); }
+	inline void Advance(int32 Count=1) { X.Advance(Count); Y.Advance(Count); Z.Advance(Count); }
+	inline bool IsConstant() const { return X.IsConstant() && Y.IsConstant() && Z.IsConstant(); }
+	inline void Reset() { X.Reset(); Y.Reset(); Z.Reset(); }
 };
 
 template<>
@@ -1258,12 +1258,12 @@ struct FNDIInputParam<FNiagaraPosition>
 	VectorVM::FExternalFuncInputHandler<float> Y;
 	VectorVM::FExternalFuncInputHandler<float> Z;
 	FNDIInputParam(FVectorVMExternalFunctionContext& Context) : X(Context), Y(Context), Z(Context) {}
-	FORCEINLINE FNDIInputParam() { }
-	FORCEINLINE void Init(FVectorVMExternalFunctionContext& Context) { new(this)FNDIInputParam<FNiagaraPosition>(Context); }
-	FORCEINLINE FNiagaraPosition GetAndAdvance() { return FNiagaraPosition(X.GetAndAdvance(), Y.GetAndAdvance(), Z.GetAndAdvance()); }
-	FORCEINLINE FNiagaraPosition Get() { return FNiagaraPosition(X.Get(), Y.Get(), Z.Get()); }
-	FORCEINLINE void Advance(int32 Count=1) { X.Advance(Count); Y.Advance(Count); Z.Advance(Count); }
-	FORCEINLINE void Reset() { X.Reset(); Y.Reset(); Z.Reset(); }
+	inline FNDIInputParam() { }
+	inline void Init(FVectorVMExternalFunctionContext& Context) { new(this)FNDIInputParam<FNiagaraPosition>(Context); }
+	inline FNiagaraPosition GetAndAdvance() { return FNiagaraPosition(X.GetAndAdvance(), Y.GetAndAdvance(), Z.GetAndAdvance()); }
+	inline FNiagaraPosition Get() { return FNiagaraPosition(X.Get(), Y.Get(), Z.Get()); }
+	inline void Advance(int32 Count=1) { X.Advance(Count); Y.Advance(Count); Z.Advance(Count); }
+	inline void Reset() { X.Reset(); Y.Reset(); Z.Reset(); }
 };
 template<>
 struct FNDIInputParam<FVector4f>
@@ -1272,14 +1272,14 @@ struct FNDIInputParam<FVector4f>
 	VectorVM::FExternalFuncInputHandler<float> Y;
 	VectorVM::FExternalFuncInputHandler<float> Z;
 	VectorVM::FExternalFuncInputHandler<float> W;
-	FORCEINLINE FNDIInputParam(FVectorVMExternalFunctionContext& Context) : X(Context), Y(Context), Z(Context), W(Context) {}
-	FORCEINLINE FNDIInputParam() { }
-	FORCEINLINE void Init(FVectorVMExternalFunctionContext& Context) { new(this)FNDIInputParam<FVector4f>(Context); }
-	FORCEINLINE FVector4f GetAndAdvance() { return FVector4f(X.GetAndAdvance(), Y.GetAndAdvance(), Z.GetAndAdvance(), W.GetAndAdvance()); }
-	FORCEINLINE FVector4f Get() { return FVector4f(X.Get(), Y.Get(), Z.Get(), W.Get()); }
-	FORCEINLINE void Advance(int32 Count=1) { return X.Advance(Count); Y.Advance(Count); Z.Advance(Count); W.Advance(Count); }
-	FORCEINLINE bool IsConstant() const { return X.IsConstant() && Y.IsConstant() && Z.IsConstant() && W.IsConstant(); }
-	FORCEINLINE void Reset() { X.Reset(); Y.Reset(); Z.Reset(); W.Reset(); }
+	inline FNDIInputParam(FVectorVMExternalFunctionContext& Context) : X(Context), Y(Context), Z(Context), W(Context) {}
+	inline FNDIInputParam() { }
+	inline void Init(FVectorVMExternalFunctionContext& Context) { new(this)FNDIInputParam<FVector4f>(Context); }
+	inline FVector4f GetAndAdvance() { return FVector4f(X.GetAndAdvance(), Y.GetAndAdvance(), Z.GetAndAdvance(), W.GetAndAdvance()); }
+	inline FVector4f Get() { return FVector4f(X.Get(), Y.Get(), Z.Get(), W.Get()); }
+	inline void Advance(int32 Count=1) { return X.Advance(Count); Y.Advance(Count); Z.Advance(Count); W.Advance(Count); }
+	inline bool IsConstant() const { return X.IsConstant() && Y.IsConstant() && Z.IsConstant() && W.IsConstant(); }
+	inline void Reset() { X.Reset(); Y.Reset(); Z.Reset(); W.Reset(); }
 };
 
 template<>
@@ -1289,14 +1289,14 @@ struct FNDIInputParam<FQuat4f>
 	VectorVM::FExternalFuncInputHandler<float> Y;
 	VectorVM::FExternalFuncInputHandler<float> Z;
 	VectorVM::FExternalFuncInputHandler<float> W;
-	FORCEINLINE FNDIInputParam(FVectorVMExternalFunctionContext& Context) : X(Context), Y(Context), Z(Context), W(Context) {}
-	FORCEINLINE FNDIInputParam() { }
-	FORCEINLINE void Init(FVectorVMExternalFunctionContext& Context) { new(this)FNDIInputParam<FQuat4f>(Context); }
-	FORCEINLINE FQuat4f GetAndAdvance() { return FQuat4f(X.GetAndAdvance(), Y.GetAndAdvance(), Z.GetAndAdvance(), W.GetAndAdvance()); }
-	FORCEINLINE FQuat4f Get() { return FQuat4f(X.Get(), Y.Get(), Z.Get(), W.Get()); }
-	FORCEINLINE void Advance(int32 Count=1) { X.Advance(Count); Y.Advance(Count); Z.Advance(Count); W.Advance(Count); }
-	FORCEINLINE bool IsConstant() const { return X.IsConstant() && Y.IsConstant() && Z.IsConstant() && W.IsConstant(); }
-	FORCEINLINE void Reset() { X.Reset(); Y.Reset(); Z.Reset(); W.Reset(); }
+	inline FNDIInputParam(FVectorVMExternalFunctionContext& Context) : X(Context), Y(Context), Z(Context), W(Context) {}
+	inline FNDIInputParam() { }
+	inline void Init(FVectorVMExternalFunctionContext& Context) { new(this)FNDIInputParam<FQuat4f>(Context); }
+	inline FQuat4f GetAndAdvance() { return FQuat4f(X.GetAndAdvance(), Y.GetAndAdvance(), Z.GetAndAdvance(), W.GetAndAdvance()); }
+	inline FQuat4f Get() { return FQuat4f(X.Get(), Y.Get(), Z.Get(), W.Get()); }
+	inline void Advance(int32 Count=1) { X.Advance(Count); Y.Advance(Count); Z.Advance(Count); W.Advance(Count); }
+	inline bool IsConstant() const { return X.IsConstant() && Y.IsConstant() && Z.IsConstant() && W.IsConstant(); }
+	inline void Reset() { X.Reset(); Y.Reset(); Z.Reset(); W.Reset(); }
 };
 
 template<>
@@ -1304,16 +1304,16 @@ struct FNDIInputParam<FMatrix44f>
 {
 	VectorVM::FExternalFuncInputHandler<float> M[16];
 
-	FORCEINLINE FNDIInputParam() { }
-	FORCEINLINE FNDIInputParam(FVectorVMExternalFunctionContext& Context) { Init(Context); }
-	FORCEINLINE void Init(FVectorVMExternalFunctionContext& Context)
+	inline FNDIInputParam() { }
+	inline FNDIInputParam(FVectorVMExternalFunctionContext& Context) { Init(Context); }
+	inline void Init(FVectorVMExternalFunctionContext& Context)
 	{
 		for (int32 i=0; i < 16; ++i)
 		{
 			M[i].Init(Context);
 		}
 	}
-	FORCEINLINE FMatrix44f GetAndAdvance()
+	inline FMatrix44f GetAndAdvance()
 	{
 		FMatrix44f Matrix;
 		for (int32 i=0; i < 16; ++i)
@@ -1322,7 +1322,7 @@ struct FNDIInputParam<FMatrix44f>
 		}
 		return Matrix;
 	}
-	FORCEINLINE FMatrix44f Get()
+	inline FMatrix44f Get()
 	{
 		FMatrix44f Matrix;
 		for (int32 i = 0; i < 16; ++i)
@@ -1331,14 +1331,14 @@ struct FNDIInputParam<FMatrix44f>
 		}
 		return Matrix;
 	}
-	FORCEINLINE void Advance(int32 Count = 1)
+	inline void Advance(int32 Count = 1)
 	{
 		for (int32 i = 0; i < 16; ++i)
 		{
 			M[i].Advance(Count);
 		}
 	}
-	FORCEINLINE bool IsConstant() const
+	inline bool IsConstant() const
 	{
 		bool bConstant = true;
 		for (int32 i = 0; i < 16; ++i)
@@ -1347,7 +1347,7 @@ struct FNDIInputParam<FMatrix44f>
 		}
 		return bConstant;
 	}
-	FORCEINLINE void Reset()
+	inline void Reset()
 	{
 		for (int32 i = 0; i < 16; ++i)
 		{
@@ -1363,14 +1363,14 @@ struct FNDIInputParam<FLinearColor>
 	VectorVM::FExternalFuncInputHandler<float> G;
 	VectorVM::FExternalFuncInputHandler<float> B;
 	VectorVM::FExternalFuncInputHandler<float> A;
-	FORCEINLINE FNDIInputParam(FVectorVMExternalFunctionContext& Context) : R(Context), G(Context), B(Context), A(Context) {}
-	FORCEINLINE FNDIInputParam() { }
-	FORCEINLINE void Init(FVectorVMExternalFunctionContext& Context) { new(this)FNDIInputParam<FLinearColor>(Context); }
-	FORCEINLINE FLinearColor GetAndAdvance() { return FLinearColor(R.GetAndAdvance(), G.GetAndAdvance(), B.GetAndAdvance(), A.GetAndAdvance()); }
-	FORCEINLINE FLinearColor Get() { return FLinearColor(R.Get(), G.Get(), B.Get(), A.Get()); }
-	FORCEINLINE void Advance(int32 Count=1) { R.Advance(Count); G.Advance(Count); B.Advance(Count); A.Advance(Count); }
-	FORCEINLINE bool IsConstant() const { return R.IsConstant() && G.IsConstant() && B.IsConstant() && A.IsConstant(); }
-	FORCEINLINE void Reset() { R.Reset(); G.Reset(); B.Reset(); A.Reset(); }
+	inline FNDIInputParam(FVectorVMExternalFunctionContext& Context) : R(Context), G(Context), B(Context), A(Context) {}
+	inline FNDIInputParam() { }
+	inline void Init(FVectorVMExternalFunctionContext& Context) { new(this)FNDIInputParam<FLinearColor>(Context); }
+	inline FLinearColor GetAndAdvance() { return FLinearColor(R.GetAndAdvance(), G.GetAndAdvance(), B.GetAndAdvance(), A.GetAndAdvance()); }
+	inline FLinearColor Get() { return FLinearColor(R.Get(), G.Get(), B.Get(), A.Get()); }
+	inline void Advance(int32 Count=1) { R.Advance(Count); G.Advance(Count); B.Advance(Count); A.Advance(Count); }
+	inline bool IsConstant() const { return R.IsConstant() && G.IsConstant() && B.IsConstant() && A.IsConstant(); }
+	inline void Reset() { R.Reset(); G.Reset(); B.Reset(); A.Reset(); }
 };
 
 template<>
@@ -1378,14 +1378,14 @@ struct FNDIInputParam<FNiagaraID>
 {
 	VectorVM::FExternalFuncInputHandler<int32> Index;
 	VectorVM::FExternalFuncInputHandler<int32> AcquireTag;
-	FORCEINLINE FNDIInputParam(FVectorVMExternalFunctionContext& Context) : Index(Context), AcquireTag(Context) {}
-	FORCEINLINE FNDIInputParam() { }
-	FORCEINLINE void Init(FVectorVMExternalFunctionContext& Context) { new(this)FNDIInputParam<FNiagaraID>(Context); }
-	FORCEINLINE FNiagaraID GetAndAdvance() { return FNiagaraID(Index.GetAndAdvance(), AcquireTag.GetAndAdvance()); }
-	FORCEINLINE FNiagaraID Get() { return FNiagaraID(Index.Get(), AcquireTag.Get()); }
-	FORCEINLINE void Advance(int32 Count=1) { return Index.Advance(Count); AcquireTag.Advance(Count); }
-	FORCEINLINE bool IsConstant() const { return Index.IsConstant() && AcquireTag.IsConstant(); }
-	FORCEINLINE void Reset() { Index.Reset(); AcquireTag.Reset(); }
+	inline FNDIInputParam(FVectorVMExternalFunctionContext& Context) : Index(Context), AcquireTag(Context) {}
+	inline FNDIInputParam() { }
+	inline void Init(FVectorVMExternalFunctionContext& Context) { new(this)FNDIInputParam<FNiagaraID>(Context); }
+	inline FNiagaraID GetAndAdvance() { return FNiagaraID(Index.GetAndAdvance(), AcquireTag.GetAndAdvance()); }
+	inline FNiagaraID Get() { return FNiagaraID(Index.Get(), AcquireTag.Get()); }
+	inline void Advance(int32 Count=1) { return Index.Advance(Count); AcquireTag.Advance(Count); }
+	inline bool IsConstant() const { return Index.IsConstant() && AcquireTag.IsConstant(); }
+	inline void Reset() { Index.Reset(); AcquireTag.Reset(); }
 };
 
 template<>
@@ -1395,28 +1395,28 @@ struct FNDIInputParam<FNiagaraSpawnInfo>
 	VectorVM::FExternalFuncInputHandler<float> InterpStartDt;
 	VectorVM::FExternalFuncInputHandler<float> IntervalDt;
 	VectorVM::FExternalFuncInputHandler<int32> SpawnGroup;
-	FORCEINLINE FNDIInputParam(FVectorVMExternalFunctionContext& Context) : Count(Context), InterpStartDt(Context), IntervalDt(Context), SpawnGroup(Context) {}
-	FORCEINLINE FNDIInputParam() { }
-	FORCEINLINE void Init(FVectorVMExternalFunctionContext& Context) { new(this)FNDIInputParam<FNiagaraSpawnInfo>(Context); }
-	FORCEINLINE FNiagaraSpawnInfo GetAndAdvance() { return FNiagaraSpawnInfo(Count.GetAndAdvance(), InterpStartDt.GetAndAdvance(),IntervalDt.GetAndAdvance(),SpawnGroup.GetAndAdvance()); }
-	FORCEINLINE FNiagaraSpawnInfo Get() { return FNiagaraSpawnInfo(Count.Get(), InterpStartDt.Get(),IntervalDt.Get(),SpawnGroup.Get()); }
-	FORCEINLINE void Advance(int32 InCount=1) { return Count.Advance(InCount); InterpStartDt.Advance(InCount);  IntervalDt.Advance(InCount);  SpawnGroup.Advance(InCount); }
-	FORCEINLINE bool IsConstant() const { return Count.IsConstant() && InterpStartDt.IsConstant() && IntervalDt.IsConstant() && SpawnGroup.IsConstant(); }
-	FORCEINLINE void Reset() { Count.Reset(); InterpStartDt.Reset();  IntervalDt.Reset(); SpawnGroup.Reset();}
+	inline FNDIInputParam(FVectorVMExternalFunctionContext& Context) : Count(Context), InterpStartDt(Context), IntervalDt(Context), SpawnGroup(Context) {}
+	inline FNDIInputParam() { }
+	inline void Init(FVectorVMExternalFunctionContext& Context) { new(this)FNDIInputParam<FNiagaraSpawnInfo>(Context); }
+	inline FNiagaraSpawnInfo GetAndAdvance() { return FNiagaraSpawnInfo(Count.GetAndAdvance(), InterpStartDt.GetAndAdvance(),IntervalDt.GetAndAdvance(),SpawnGroup.GetAndAdvance()); }
+	inline FNiagaraSpawnInfo Get() { return FNiagaraSpawnInfo(Count.Get(), InterpStartDt.Get(),IntervalDt.Get(),SpawnGroup.Get()); }
+	inline void Advance(int32 InCount=1) { return Count.Advance(InCount); InterpStartDt.Advance(InCount);  IntervalDt.Advance(InCount);  SpawnGroup.Advance(InCount); }
+	inline bool IsConstant() const { return Count.IsConstant() && InterpStartDt.IsConstant() && IntervalDt.IsConstant() && SpawnGroup.IsConstant(); }
+	inline void Reset() { Count.Reset(); InterpStartDt.Reset();  IntervalDt.Reset(); SpawnGroup.Reset();}
 };
 
 template<>
 struct FNDIInputParam<FNiagaraEmitterID>
 {
 	VectorVM::FExternalFuncInputHandler<int32> Index;
-	FORCEINLINE FNDIInputParam(FVectorVMExternalFunctionContext& Context) : Index(Context) {}
-	FORCEINLINE FNDIInputParam() { }
-	FORCEINLINE void Init(FVectorVMExternalFunctionContext& Context) { new(this)FNDIInputParam<FNiagaraEmitterID>(Context); }
-	FORCEINLINE FNiagaraEmitterID GetAndAdvance() { return FNiagaraEmitterID(Index.GetAndAdvance()); }
-	FORCEINLINE FNiagaraEmitterID Get() { return FNiagaraEmitterID(Index.Get()); }
-	FORCEINLINE void Advance(int32 Count = 1) { return Index.Advance(Count); }
-	FORCEINLINE bool IsConstant() const { return Index.IsConstant(); }
-	FORCEINLINE void Reset() { Index.Reset(); }
+	inline FNDIInputParam(FVectorVMExternalFunctionContext& Context) : Index(Context) {}
+	inline FNDIInputParam() { }
+	inline void Init(FVectorVMExternalFunctionContext& Context) { new(this)FNDIInputParam<FNiagaraEmitterID>(Context); }
+	inline FNiagaraEmitterID GetAndAdvance() { return FNiagaraEmitterID(Index.GetAndAdvance()); }
+	inline FNiagaraEmitterID Get() { return FNiagaraEmitterID(Index.Get()); }
+	inline void Advance(int32 Count = 1) { return Index.Advance(Count); }
+	inline bool IsConstant() const { return Index.IsConstant(); }
+	inline void Reset() { Index.Reset(); }
 };
 
 template<>
@@ -1424,9 +1424,9 @@ struct FNDIInputParam<FIntVector2>
 {
 	VectorVM::FExternalFuncInputHandler<int32> X;
 	VectorVM::FExternalFuncInputHandler<int32> Y;
-	FORCEINLINE FNDIInputParam(FVectorVMExternalFunctionContext& Context) : X(Context), Y(Context) {}
-	FORCEINLINE FIntVector2 GetAndAdvance() { return FIntVector2(X.GetAndAdvance(), Y.GetAndAdvance()); }
-	FORCEINLINE bool IsConstant() const { return X.IsConstant() && Y.IsConstant(); }
+	inline FNDIInputParam(FVectorVMExternalFunctionContext& Context) : X(Context), Y(Context) {}
+	inline FIntVector2 GetAndAdvance() { return FIntVector2(X.GetAndAdvance(), Y.GetAndAdvance()); }
+	inline bool IsConstant() const { return X.IsConstant() && Y.IsConstant(); }
 };
 
 template<>
@@ -1435,9 +1435,9 @@ struct FNDIInputParam<FIntVector3>
 	VectorVM::FExternalFuncInputHandler<int32> X;
 	VectorVM::FExternalFuncInputHandler<int32> Y;
 	VectorVM::FExternalFuncInputHandler<int32> Z;
-	FORCEINLINE FNDIInputParam(FVectorVMExternalFunctionContext& Context) : X(Context), Y(Context), Z(Context) {}
-	FORCEINLINE FIntVector3 GetAndAdvance() { return FIntVector3(X.GetAndAdvance(), Y.GetAndAdvance(), Z.GetAndAdvance()); }
-	FORCEINLINE bool IsConstant() const { return X.IsConstant() && Y.IsConstant() && Z.IsConstant(); }
+	inline FNDIInputParam(FVectorVMExternalFunctionContext& Context) : X(Context), Y(Context), Z(Context) {}
+	inline FIntVector3 GetAndAdvance() { return FIntVector3(X.GetAndAdvance(), Y.GetAndAdvance(), Z.GetAndAdvance()); }
+	inline bool IsConstant() const { return X.IsConstant() && Y.IsConstant() && Z.IsConstant(); }
 };
 
 template<>
@@ -1447,9 +1447,9 @@ struct FNDIInputParam<FIntVector4>
 	VectorVM::FExternalFuncInputHandler<int32> Y;
 	VectorVM::FExternalFuncInputHandler<int32> Z;
 	VectorVM::FExternalFuncInputHandler<int32> W;
-	FORCEINLINE FNDIInputParam(FVectorVMExternalFunctionContext& Context) : X(Context), Y(Context), Z(Context), W(Context) {}
-	FORCEINLINE FIntVector4 GetAndAdvance() { return FIntVector4(X.GetAndAdvance(), Y.GetAndAdvance(), Z.GetAndAdvance(), W.GetAndAdvance()); }
-	FORCEINLINE bool IsConstant() const { return X.IsConstant() && Y.IsConstant() && Z.IsConstant() && W.IsConstant(); }
+	inline FNDIInputParam(FVectorVMExternalFunctionContext& Context) : X(Context), Y(Context), Z(Context), W(Context) {}
+	inline FIntVector4 GetAndAdvance() { return FIntVector4(X.GetAndAdvance(), Y.GetAndAdvance(), Z.GetAndAdvance(), W.GetAndAdvance()); }
+	inline bool IsConstant() const { return X.IsConstant() && Y.IsConstant() && Z.IsConstant() && W.IsConstant(); }
 };
 
 //Helper to deal with types with potentially several output registers.
@@ -1458,28 +1458,28 @@ struct FNDIOutputParam
 {
 	static_assert(sizeof(T) == sizeof(float), "Generic template assumes 4 bytes per element");
 	VectorVM::FExternalFuncRegisterHandler<T> Data;
-	FORCEINLINE FNDIOutputParam(FVectorVMExternalFunctionContext& Context) : Data(Context) {}
-	FORCEINLINE bool IsValid() const { return Data.IsValid();  }
-	FORCEINLINE void SetAndAdvance(T Val) { *Data.GetDestAndAdvance() = Val; }
+	inline FNDIOutputParam(FVectorVMExternalFunctionContext& Context) : Data(Context) {}
+	inline bool IsValid() const { return Data.IsValid();  }
+	inline void SetAndAdvance(T Val) { *Data.GetDestAndAdvance() = Val; }
 };
 
 template<>
 struct FNDIOutputParam<FNiagaraBool>
 {
 	VectorVM::FExternalFuncRegisterHandler<FNiagaraBool> Data;
-	FORCEINLINE FNDIOutputParam(FVectorVMExternalFunctionContext& Context) : Data(Context) {}
-	FORCEINLINE bool IsValid() const { return Data.IsValid(); }
-	FORCEINLINE void SetAndAdvance(bool Val) { Data.GetDestAndAdvance()->SetValue(Val); }
-	FORCEINLINE void SetAndAdvance(FNiagaraBool Val) { *Data.GetDestAndAdvance() = Val; }
+	inline FNDIOutputParam(FVectorVMExternalFunctionContext& Context) : Data(Context) {}
+	inline bool IsValid() const { return Data.IsValid(); }
+	inline void SetAndAdvance(bool Val) { Data.GetDestAndAdvance()->SetValue(Val); }
+	inline void SetAndAdvance(FNiagaraBool Val) { *Data.GetDestAndAdvance() = Val; }
 };
 
 template<>
 struct FNDIOutputParam<bool>
 {
 	VectorVM::FExternalFuncRegisterHandler<FNiagaraBool> Data;
-	FORCEINLINE FNDIOutputParam(FVectorVMExternalFunctionContext& Context) : Data(Context) {}
-	FORCEINLINE bool IsValid() const { return Data.IsValid(); }
-	FORCEINLINE void SetAndAdvance(bool Val) { Data.GetDestAndAdvance()->SetValue(Val); }
+	inline FNDIOutputParam(FVectorVMExternalFunctionContext& Context) : Data(Context) {}
+	inline bool IsValid() const { return Data.IsValid(); }
+	inline void SetAndAdvance(bool Val) { Data.GetDestAndAdvance()->SetValue(Val); }
 };
 
 template<>
@@ -1487,9 +1487,9 @@ struct FNDIOutputParam<FVector2f>
 {
 	VectorVM::FExternalFuncRegisterHandler<float> X;
 	VectorVM::FExternalFuncRegisterHandler<float> Y;
-	FORCEINLINE FNDIOutputParam(FVectorVMExternalFunctionContext& Context) : X(Context), Y(Context) {}
-	FORCEINLINE bool IsValid() const { return X.IsValid() || Y.IsValid(); }
-	FORCEINLINE void SetAndAdvance(FVector2f Val)
+	inline FNDIOutputParam(FVectorVMExternalFunctionContext& Context) : X(Context), Y(Context) {}
+	inline bool IsValid() const { return X.IsValid() || Y.IsValid(); }
+	inline void SetAndAdvance(FVector2f Val)
 	{
 		*X.GetDestAndAdvance() = Val.X;
 		*Y.GetDestAndAdvance() = Val.Y;
@@ -1503,8 +1503,8 @@ struct FNDIOutputParam<FVector3f>
 	VectorVM::FExternalFuncRegisterHandler<float> Y;
 	VectorVM::FExternalFuncRegisterHandler<float> Z;
 	FNDIOutputParam(FVectorVMExternalFunctionContext& Context) : X(Context), Y(Context), Z(Context) {}
-	FORCEINLINE bool IsValid() const { return X.IsValid() || Y.IsValid() || Z.IsValid(); }
-	FORCEINLINE void SetAndAdvance(FVector3f Val)
+	inline bool IsValid() const { return X.IsValid() || Y.IsValid() || Z.IsValid(); }
+	inline void SetAndAdvance(FVector3f Val)
 	{
 		*X.GetDestAndAdvance() = Val.X;
 		*Y.GetDestAndAdvance() = Val.Y;
@@ -1519,8 +1519,8 @@ struct FNDIOutputParam<FNiagaraPosition>
 	VectorVM::FExternalFuncRegisterHandler<float> Y;
 	VectorVM::FExternalFuncRegisterHandler<float> Z;
 	FNDIOutputParam(FVectorVMExternalFunctionContext& Context) : X(Context), Y(Context), Z(Context) {}
-	FORCEINLINE bool IsValid() const { return X.IsValid() || Y.IsValid() || Z.IsValid(); }
-	FORCEINLINE void SetAndAdvance(FNiagaraPosition Val)
+	inline bool IsValid() const { return X.IsValid() || Y.IsValid() || Z.IsValid(); }
+	inline void SetAndAdvance(FNiagaraPosition Val)
 	{
 		*X.GetDestAndAdvance() = Val.X;
 		*Y.GetDestAndAdvance() = Val.Y;
@@ -1535,9 +1535,9 @@ struct FNDIOutputParam<FVector4f>
 	VectorVM::FExternalFuncRegisterHandler<float> Y;
 	VectorVM::FExternalFuncRegisterHandler<float> Z;
 	VectorVM::FExternalFuncRegisterHandler<float> W;
-	FORCEINLINE FNDIOutputParam(FVectorVMExternalFunctionContext& Context) : X(Context), Y(Context), Z(Context), W(Context) {}
-	FORCEINLINE bool IsValid() const { return X.IsValid() || Y.IsValid() || Z.IsValid() || W.IsValid(); }
-	FORCEINLINE void SetAndAdvance(FVector4f Val)
+	inline FNDIOutputParam(FVectorVMExternalFunctionContext& Context) : X(Context), Y(Context), Z(Context), W(Context) {}
+	inline bool IsValid() const { return X.IsValid() || Y.IsValid() || Z.IsValid() || W.IsValid(); }
+	inline void SetAndAdvance(FVector4f Val)
 	{
 		*X.GetDestAndAdvance() = Val.X;
 		*Y.GetDestAndAdvance() = Val.Y;
@@ -1553,9 +1553,9 @@ struct FNDIOutputParam<FQuat4f>
 	VectorVM::FExternalFuncRegisterHandler<float> Y;
 	VectorVM::FExternalFuncRegisterHandler<float> Z;
 	VectorVM::FExternalFuncRegisterHandler<float> W;
-	FORCEINLINE FNDIOutputParam(FVectorVMExternalFunctionContext& Context) : X(Context), Y(Context), Z(Context), W(Context) {}
-	FORCEINLINE bool IsValid() const { return X.IsValid() || Y.IsValid() || Z.IsValid() || W.IsValid(); }
-	FORCEINLINE void SetAndAdvance(FQuat4f Val)
+	inline FNDIOutputParam(FVectorVMExternalFunctionContext& Context) : X(Context), Y(Context), Z(Context), W(Context) {}
+	inline bool IsValid() const { return X.IsValid() || Y.IsValid() || Z.IsValid() || W.IsValid(); }
+	inline void SetAndAdvance(FQuat4f Val)
 	{
 		*X.GetDestAndAdvance() = Val.X;
 		*Y.GetDestAndAdvance() = Val.Y;
@@ -1584,10 +1584,10 @@ struct FNDIOutputParam<FMatrix44f>
 	VectorVM::FExternalFuncRegisterHandler<float> Out14;
 	VectorVM::FExternalFuncRegisterHandler<float> Out15;
 
-	FORCEINLINE FNDIOutputParam(FVectorVMExternalFunctionContext& Context) : Out00(Context), Out01(Context), Out02(Context), Out03(Context), Out04(Context), Out05(Context),
+	inline FNDIOutputParam(FVectorVMExternalFunctionContext& Context) : Out00(Context), Out01(Context), Out02(Context), Out03(Context), Out04(Context), Out05(Context),
 		Out06(Context), Out07(Context), Out08(Context), Out09(Context), Out10(Context), Out11(Context), Out12(Context), Out13(Context), Out14(Context), Out15(Context)	{}
-	FORCEINLINE bool IsValid() const { return Out00.IsValid(); }
-	FORCEINLINE void SetAndAdvance(const FMatrix44f& Val)
+	inline bool IsValid() const { return Out00.IsValid(); }
+	inline void SetAndAdvance(const FMatrix44f& Val)
 	{
 		*Out00.GetDestAndAdvance() = Val.M[0][0];
 		*Out01.GetDestAndAdvance() = Val.M[0][1];
@@ -1615,9 +1615,9 @@ struct FNDIOutputParam<FLinearColor>
 	VectorVM::FExternalFuncRegisterHandler<float> G;
 	VectorVM::FExternalFuncRegisterHandler<float> B;
 	VectorVM::FExternalFuncRegisterHandler<float> A;
-	FORCEINLINE FNDIOutputParam(FVectorVMExternalFunctionContext& Context) : R(Context), G(Context), B(Context), A(Context) {}
-	FORCEINLINE bool IsValid() const { return R.IsValid() || G.IsValid() || B.IsValid() || A.IsValid(); }
-	FORCEINLINE void SetAndAdvance(FLinearColor Val)
+	inline FNDIOutputParam(FVectorVMExternalFunctionContext& Context) : R(Context), G(Context), B(Context), A(Context) {}
+	inline bool IsValid() const { return R.IsValid() || G.IsValid() || B.IsValid() || A.IsValid(); }
+	inline void SetAndAdvance(FLinearColor Val)
 	{
 		*R.GetDestAndAdvance() = Val.R;
 		*G.GetDestAndAdvance() = Val.G;
@@ -1631,9 +1631,9 @@ struct FNDIOutputParam<FNiagaraID>
 {
 	VectorVM::FExternalFuncRegisterHandler<int32> Index;
 	VectorVM::FExternalFuncRegisterHandler<int32> AcquireTag;
-	FORCEINLINE FNDIOutputParam(FVectorVMExternalFunctionContext& Context) : Index(Context), AcquireTag(Context) {}
-	FORCEINLINE bool IsValid() const { return Index.IsValid() || AcquireTag.IsValid(); }
-	FORCEINLINE void SetAndAdvance(FNiagaraID Val)
+	inline FNDIOutputParam(FVectorVMExternalFunctionContext& Context) : Index(Context), AcquireTag(Context) {}
+	inline bool IsValid() const { return Index.IsValid() || AcquireTag.IsValid(); }
+	inline void SetAndAdvance(FNiagaraID Val)
 	{
 		*Index.GetDestAndAdvance() = Val.Index;
 		*AcquireTag.GetDestAndAdvance() = Val.AcquireTag;
@@ -1645,9 +1645,9 @@ struct FNDIOutputParam<FIntVector2>
 {
 	VectorVM::FExternalFuncRegisterHandler<int32> X;
 	VectorVM::FExternalFuncRegisterHandler<int32> Y;
-	FORCEINLINE FNDIOutputParam(FVectorVMExternalFunctionContext& Context) : X(Context), Y(Context) {}
-	FORCEINLINE bool IsValid() const { return X.IsValid() || Y.IsValid(); }
-	FORCEINLINE void SetAndAdvance(FIntVector2 Val)
+	inline FNDIOutputParam(FVectorVMExternalFunctionContext& Context) : X(Context), Y(Context) {}
+	inline bool IsValid() const { return X.IsValid() || Y.IsValid(); }
+	inline void SetAndAdvance(FIntVector2 Val)
 	{
 		*X.GetDestAndAdvance() = Val.X;
 		*Y.GetDestAndAdvance() = Val.Y;
@@ -1660,9 +1660,9 @@ struct FNDIOutputParam<FIntVector3>
 	VectorVM::FExternalFuncRegisterHandler<int32> X;
 	VectorVM::FExternalFuncRegisterHandler<int32> Y;
 	VectorVM::FExternalFuncRegisterHandler<int32> Z;
-	FORCEINLINE FNDIOutputParam(FVectorVMExternalFunctionContext& Context) : X(Context), Y(Context), Z(Context) {}
-	FORCEINLINE bool IsValid() const { return X.IsValid() || Y.IsValid() || Z.IsValid(); }
-	FORCEINLINE void SetAndAdvance(FIntVector3 Val)
+	inline FNDIOutputParam(FVectorVMExternalFunctionContext& Context) : X(Context), Y(Context), Z(Context) {}
+	inline bool IsValid() const { return X.IsValid() || Y.IsValid() || Z.IsValid(); }
+	inline void SetAndAdvance(FIntVector3 Val)
 	{
 		*X.GetDestAndAdvance() = Val.X;
 		*Y.GetDestAndAdvance() = Val.Y;
@@ -1677,9 +1677,9 @@ struct FNDIOutputParam<FIntVector4>
 	VectorVM::FExternalFuncRegisterHandler<int32> Y;
 	VectorVM::FExternalFuncRegisterHandler<int32> Z;
 	VectorVM::FExternalFuncRegisterHandler<int32> W;
-	FORCEINLINE FNDIOutputParam(FVectorVMExternalFunctionContext& Context) : X(Context), Y(Context), Z(Context), W(Context) {}
-	FORCEINLINE bool IsValid() const { return X.IsValid() || Y.IsValid() || Z.IsValid() || W.IsValid(); }
-	FORCEINLINE void SetAndAdvance(FIntVector4 Val)
+	inline FNDIOutputParam(FVectorVMExternalFunctionContext& Context) : X(Context), Y(Context), Z(Context), W(Context) {}
+	inline bool IsValid() const { return X.IsValid() || Y.IsValid() || Z.IsValid() || W.IsValid(); }
+	inline void SetAndAdvance(FIntVector4 Val)
 	{
 		*X.GetDestAndAdvance() = Val.X;
 		*Y.GetDestAndAdvance() = Val.Y;
diff --git a/Engine/Plugins/FX/Niagara/Source/Niagara/Classes/NiagaraDataInterfaceColorCurve.h b/Engine/Plugins/FX/Niagara/Source/Niagara/Classes/NiagaraDataInterfaceColorCurve.h
--- a/Engine/Plugins/FX/Niagara/Source/Niagara/Classes/NiagaraDataInterfaceColorCurve.h
+++ b/Engine/Plugins/FX/Niagara/Source/Niagara/Classes/NiagaraDataInterfaceColorCurve.h
@@ -90,7 +90,7 @@ protected:
 
 private:
 	template<typename UseLUT>
-	FORCEINLINE_DEBUGGABLE FLinearColor SampleCurveInternal(float X);
+	inline FLinearColor SampleCurveInternal(float X);
 
 	static NIAGARA_API const FName SampleCurveName;
 };
diff --git a/Engine/Plugins/FX/Niagara/Source/Niagara/Classes/NiagaraDataInterfaceCurve.h b/Engine/Plugins/FX/Niagara/Source/Niagara/Classes/NiagaraDataInterfaceCurve.h
--- a/Engine/Plugins/FX/Niagara/Source/Niagara/Classes/NiagaraDataInterfaceCurve.h
+++ b/Engine/Plugins/FX/Niagara/Source/Niagara/Classes/NiagaraDataInterfaceCurve.h
@@ -68,7 +68,7 @@ protected:
 
 private:
 	template<typename UseLUT>
-	FORCEINLINE_DEBUGGABLE float SampleCurveInternal(float X);
+	inline float SampleCurveInternal(float X);
 
 	static NIAGARA_API const FName SampleCurveName;
 };
diff --git a/Engine/Plugins/FX/Niagara/Source/Niagara/Classes/NiagaraDataInterfaceCurveBase.h b/Engine/Plugins/FX/Niagara/Source/Niagara/Classes/NiagaraDataInterfaceCurveBase.h
--- a/Engine/Plugins/FX/Niagara/Source/Niagara/Classes/NiagaraDataInterfaceCurveBase.h
+++ b/Engine/Plugins/FX/Niagara/Source/Niagara/Classes/NiagaraDataInterfaceCurveBase.h
@@ -37,13 +37,13 @@ protected:
 	uint32 LUTOffset = INDEX_NONE;
 
 	/** Remap a sample time for this curve to 0 to 1 between first and last keys for LUT access.*/
-	FORCEINLINE float NormalizeTime(float T) const
+	inline float NormalizeTime(float T) const
 	{
 		return (T - LUTMinTime) * LUTInvTimeRange;
 	}
 
 	/** Remap a 0 to 1 value between the first and last keys to a real sample time for this curve. */
-	FORCEINLINE float UnnormalizeTime(float T) const
+	inline float UnnormalizeTime(float T) const
 	{
 		return (T / LUTInvTimeRange) + LUTMinTime;
 	}
@@ -160,9 +160,9 @@ public:
 	virtual void UpdateTimeRanges() PURE_VIRTUAL(UNiagaraDataInterfaceCurveBase::UpdateTimeRanges(), )
 	virtual TArray<float> BuildLUT(int32 NumEntries) const PURE_VIRTUAL(UNiagaraDataInterfaceCurveBase::BuildLUT(), return TArray<float>(); )
 
-	FORCEINLINE float GetMinTime()const { return LUTMinTime; }
-	FORCEINLINE float GetMaxTime()const { return LUTMaxTime; }
-	FORCEINLINE float GetInvTimeRange()const { return LUTInvTimeRange; }
+	inline float GetMinTime()const { return LUTMinTime; }
+	inline float GetMaxTime()const { return LUTMaxTime; }
+	inline float GetInvTimeRange()const { return LUTInvTimeRange; }
 
 protected:
 	NIAGARA_API virtual void PushToRenderThreadImpl() override;
diff --git a/Engine/Plugins/FX/Niagara/Source/Niagara/Classes/NiagaraDataInterfaceMeshCommon.h b/Engine/Plugins/FX/Niagara/Source/Niagara/Classes/NiagaraDataInterfaceMeshCommon.h
--- a/Engine/Plugins/FX/Niagara/Source/Niagara/Classes/NiagaraDataInterfaceMeshCommon.h
+++ b/Engine/Plugins/FX/Niagara/Source/Niagara/Classes/NiagaraDataInterfaceMeshCommon.h
@@ -30,7 +30,7 @@ struct FMeshTriCoordinate
 	{}
 };
 
-FORCEINLINE FVector3f RandomBarycentricCoord(FRandomStream& RandStream)
+inline FVector3f RandomBarycentricCoord(FRandomStream& RandStream)
 {
 	//TODO: This is gonna be slooooow. Move to an LUT possibly or find faster method.
 	//Can probably handle lower quality randoms / uniformity for a decent speed win.
@@ -41,25 +41,25 @@ FORCEINLINE FVector3f RandomBarycentricCoord(FRandomStream& RandStream)
 }
 
 template<typename T>
-FORCEINLINE T BarycentricInterpolate(float BaryX, float BaryY, float BaryZ, T V0, T V1, T V2)
+inline T BarycentricInterpolate(float BaryX, float BaryY, float BaryZ, T V0, T V1, T V2)
 {
 	return V0 * BaryX + V1 * BaryY + V2 * BaryZ;
 }
 
 // Overload for FVector4 to work around C2719: (formal parameter with requested alignment of 16 won't be aligned)
-FORCEINLINE FVector4 BarycentricInterpolate(float BaryX, float BaryY, float BaryZ, const FVector4& V0, const FVector4& V1, const FVector4& V2)
+inline FVector4 BarycentricInterpolate(float BaryX, float BaryY, float BaryZ, const FVector4& V0, const FVector4& V1, const FVector4& V2)
 {
 	return V0 * BaryX + V1 * BaryY + V2 * BaryZ;
 }
 
 template<typename T>
-FORCEINLINE T BarycentricInterpolate(FVector3f BaryCoord, T V0, T V1, T V2)
+inline T BarycentricInterpolate(FVector3f BaryCoord, T V0, T V1, T V2)
 {
 	return V0 * BaryCoord.X + V1 * BaryCoord.Y + V2 * BaryCoord.Z;
 }
 
 // Overload for FVector4 to work around C2719: (formal parameter with requested alignment of 16 won't be aligned)
-FORCEINLINE FVector4 BarycentricInterpolate(FVector3f BaryCoord, const FVector4& V0, const FVector4& V1, const FVector4& V2)
+inline FVector4 BarycentricInterpolate(FVector3f BaryCoord, const FVector4& V0, const FVector4& V1, const FVector4& V2)
 {
 	return V0 * BaryCoord.X + V1 * BaryCoord.Y + V2 * BaryCoord.Z;
 }
diff --git a/Engine/Plugins/FX/Niagara/Source/Niagara/Classes/NiagaraDataInterfaceSkeletalMesh.h b/Engine/Plugins/FX/Niagara/Source/Niagara/Classes/NiagaraDataInterfaceSkeletalMesh.h
--- a/Engine/Plugins/FX/Niagara/Source/Niagara/Classes/NiagaraDataInterfaceSkeletalMesh.h
+++ b/Engine/Plugins/FX/Niagara/Source/Niagara/Classes/NiagaraDataInterfaceSkeletalMesh.h
@@ -41,9 +41,9 @@ struct FSkeletalMeshSkinningDataUsage
 		, bUsesPreSkinnedVerts(bInUsesPreSkinnedVerts)
 	{}
 
-	FORCEINLINE bool NeedBoneMatrices()const { return bUsesBoneMatrices || bUsesPreSkinnedVerts; }
-	FORCEINLINE bool NeedPreSkinnedVerts()const { return bUsesPreSkinnedVerts; }
-	FORCEINLINE int32 GetLODIndex()const { return LODIndex; }
+	inline bool NeedBoneMatrices()const { return bUsesBoneMatrices || bUsesPreSkinnedVerts; }
+	inline bool NeedPreSkinnedVerts()const { return bUsesPreSkinnedVerts; }
+	inline int32 GetLODIndex()const { return LODIndex; }
 private:
 	int32 LODIndex;
 	uint32 bUsesBoneMatrices : 1;
@@ -83,17 +83,17 @@ struct FSkeletalMeshSkinningData
 
 	bool Tick(float InDeltaSeconds, bool bRequirePreskin = true);
 
-	FORCEINLINE void EnterRead()
+	inline void EnterRead()
 	{
 		RWGuard.ReadLock();
 	}
 
-	FORCEINLINE void ExitRead()
+	inline void ExitRead()
 	{
 		RWGuard.ReadUnlock();
 	}
 
-	FORCEINLINE int32 GetBoneCount(bool RequiresPrevious) const
+	inline int32 GetBoneCount(bool RequiresPrevious) const
 	{
 		int32 BoneCount = CurrComponentTransforms().Num();
 		if (RequiresPrevious)
@@ -104,17 +104,17 @@ struct FSkeletalMeshSkinningData
 		return BoneCount;
 	}
 
-	FORCEINLINE FVector3f GetPosition(int32 LODIndex, int32 VertexIndex) const
+	inline FVector3f GetPosition(int32 LODIndex, int32 VertexIndex) const
 	{
 		return LODData.IsValidIndex(LODIndex) ? LODData[LODIndex].SkinnedCPUPositions[CurrIndex][VertexIndex] : FVector3f::ZeroVector;
 	}
 
-	FORCEINLINE FVector3f GetPreviousPosition(int32 LODIndex, int32 VertexIndex) const
+	inline FVector3f GetPreviousPosition(int32 LODIndex, int32 VertexIndex) const
 	{
 		return LODData.IsValidIndex(LODIndex) ? LODData[LODIndex].SkinnedCPUPositions[CurrIndex ^ 1][VertexIndex] : FVector3f::ZeroVector;
 	}
 
-	FORCEINLINE void GetTangentBasis(int32 LODIndex, int32 VertexIndex, FVector3f& OutTangentX, FVector3f& OutTangentY, FVector3f& OutTangentZ)
+	inline void GetTangentBasis(int32 LODIndex, int32 VertexIndex, FVector3f& OutTangentX, FVector3f& OutTangentY, FVector3f& OutTangentZ)
 	{
 		const bool bValidLOD = LODData.IsValidIndex(LODIndex);
 		OutTangentX = bValidLOD ? LODData[LODIndex].SkinnedTangentBasis[CurrIndex][(VertexIndex * 3) + 0] : FVector3f(1.0f, 0.0f, 0.0f);
@@ -122,7 +122,7 @@ struct FSkeletalMeshSkinningData
 		OutTangentZ = bValidLOD ? LODData[LODIndex].SkinnedTangentBasis[CurrIndex][(VertexIndex * 3) + 2] : FVector3f(0.0f, 0.0f, 1.0f);
 	}
 
-	FORCEINLINE void GetPreviousTangentBasis(int32 LODIndex, int32 VertexIndex, FVector3f& OutTangentX, FVector3f& OutTangentY, FVector3f& OutTangentZ)
+	inline void GetPreviousTangentBasis(int32 LODIndex, int32 VertexIndex, FVector3f& OutTangentX, FVector3f& OutTangentY, FVector3f& OutTangentZ)
 	{
 		const bool bValidLOD = LODData.IsValidIndex(LODIndex);
 		OutTangentX = bValidLOD ? LODData[LODIndex].SkinnedTangentBasis[CurrIndex ^ 1][(VertexIndex * 3) + 0] : FVector3f(1.0f, 0.0f, 0.0f);
@@ -131,68 +131,68 @@ struct FSkeletalMeshSkinningData
 	}
 
 private:
-	FORCEINLINE TArray<FVector3f>& CurrSkinnedPositions(int32 LODIndex)
+	inline TArray<FVector3f>& CurrSkinnedPositions(int32 LODIndex)
 	{
 		return LODData[LODIndex].SkinnedCPUPositions[CurrIndex];
 	}
 
-	FORCEINLINE TArray<FVector3f>& PrevSkinnedPositions(int32 LODIndex)
+	inline TArray<FVector3f>& PrevSkinnedPositions(int32 LODIndex)
 	{
 		return LODData[LODIndex].SkinnedCPUPositions[CurrIndex ^ 1];
 	}
 
-	FORCEINLINE TArray<FVector3f>& CurrSkinnedTangentBasis(int32 LODIndex)
+	inline TArray<FVector3f>& CurrSkinnedTangentBasis(int32 LODIndex)
 	{
 		return LODData[LODIndex].SkinnedTangentBasis[CurrIndex];
 	}
 
-	FORCEINLINE TArray<FVector3f>& PrevSkinnedTangentBasis(int32 LODIndex)
+	inline TArray<FVector3f>& PrevSkinnedTangentBasis(int32 LODIndex)
 	{
 		return LODData[LODIndex].SkinnedTangentBasis[CurrIndex ^ 1];
 	}
 
 public:
-	FORCEINLINE TArray<FMatrix44f>& CurrBoneRefToLocals()
+	inline TArray<FMatrix44f>& CurrBoneRefToLocals()
 	{
 		return BoneRefToLocals[CurrIndex];
 	}
 
-	FORCEINLINE const TArray<FMatrix44f>& CurrBoneRefToLocals() const
+	inline const TArray<FMatrix44f>& CurrBoneRefToLocals() const
 	{
 		return BoneRefToLocals[CurrIndex];
 	}
 
-	FORCEINLINE TArray<FMatrix44f>& PrevBoneRefToLocals()
+	inline TArray<FMatrix44f>& PrevBoneRefToLocals()
 	{
 		return BoneRefToLocals[CurrIndex ^ 1];
 	}
 
-	FORCEINLINE const TArray<FMatrix44f>& PrevBoneRefToLocals() const
+	inline const TArray<FMatrix44f>& PrevBoneRefToLocals() const
 	{
 		return BoneRefToLocals[CurrIndex ^ 1];
 	}
 
-	FORCEINLINE TArray<FTransform3f>& CurrComponentTransforms()
+	inline TArray<FTransform3f>& CurrComponentTransforms()
 	{
 		return ComponentTransforms[CurrIndex];
 	}
 
-	FORCEINLINE const TArray<FTransform3f>& CurrComponentTransforms() const
+	inline const TArray<FTransform3f>& CurrComponentTransforms() const
 	{
 		return ComponentTransforms[CurrIndex];
 	}
 
-	FORCEINLINE TArray<FTransform3f>& PrevComponentTransforms()
+	inline TArray<FTransform3f>& PrevComponentTransforms()
 	{
 		return ComponentTransforms[CurrIndex ^ 1];
 	}
 
-	FORCEINLINE const TArray<FTransform3f>& PrevComponentTransforms() const
+	inline const TArray<FTransform3f>& PrevComponentTransforms() const
 	{
 		return ComponentTransforms[CurrIndex ^ 1];
 	}
 
-	FORCEINLINE bool NeedPreSkinnedVerts() const
+	inline bool NeedPreSkinnedVerts() const
 	{
 		return TotalPreSkinnedVertsUsers > 0;
 	}
@@ -376,7 +376,7 @@ struct FSkeletalMeshSamplingRegionAreaWeightedSampler : FWeightedRandomSampler
 	void Init(FNDISkeletalMesh_InstanceData* InOwner);
 	virtual float GetWeights(TArray<float>& OutWeights)override;
 
-	FORCEINLINE bool IsValid() { return TotalWeight > 0.0f; }
+	inline bool IsValid() { return TotalWeight > 0.0f; }
 
 	int32 GetEntries() const { return Alias.Num(); }
 protected:
@@ -394,7 +394,7 @@ class FSkeletalMeshGpuSpawnStaticBuffers : public FRenderResource
 public:
 	virtual ~FSkeletalMeshGpuSpawnStaticBuffers() override;
 
-	FORCEINLINE_DEBUGGABLE void Initialise(struct FNDISkeletalMesh_InstanceData* InstData, int32 InLODIndex, const FSkeletalMeshLODRenderData& SkeletalMeshLODRenderData, const FSkeletalMeshSamplingLODBuiltData* SkeletalMeshSamplingLODBuiltData, FNiagaraSystemInstance* SystemInstance);
+	inline void Initialise(struct FNDISkeletalMesh_InstanceData* InstData, int32 InLODIndex, const FSkeletalMeshLODRenderData& SkeletalMeshLODRenderData, const FSkeletalMeshSamplingLODBuiltData* SkeletalMeshSamplingLODBuiltData, FNiagaraSystemInstance* SystemInstance);
 
 	virtual void InitRHI(FRHICommandListBase& RHICmdList) override;
 	virtual void ReleaseRHI() override;
@@ -662,9 +662,9 @@ struct FNDISkeletalMesh_InstanceData
 	bool Tick(UNiagaraDataInterfaceSkeletalMesh* Interface, FNiagaraSystemInstance* SystemInstance, float InDeltaSeconds);
 	void Release();
 
-	FORCEINLINE int32 GetLODIndex()const { return CachedLODIdx; }
+	inline int32 GetLODIndex()const { return CachedLODIdx; }
 
-	FORCEINLINE_DEBUGGABLE const FSkinWeightVertexBuffer* GetSkinWeights() const
+	inline const FSkinWeightVertexBuffer* GetSkinWeights() const
 	{
 		USkeletalMeshComponent* SkelComp = Cast<USkeletalMeshComponent>(SceneComponent.Get());
 		if (SkelComp != nullptr && SkelComp->GetSkeletalMeshAsset() != nullptr)
@@ -902,16 +902,16 @@ public:
 
 private:
 	template<typename FilterMode, typename AreaWeightingMode>
-	FORCEINLINE int32 RandomTriIndex(FNDIRandomHelper& RandHelper, FSkeletalMeshAccessorHelper& Accessor, FNDISkeletalMesh_InstanceData* InstData, int32 InstanceIndex);
+	inline int32 RandomTriIndex(FNDIRandomHelper& RandHelper, FSkeletalMeshAccessorHelper& Accessor, FNDISkeletalMesh_InstanceData* InstData, int32 InstanceIndex);
 
 	NIAGARA_API void RandomTriangle(FVectorVMExternalFunctionContext& Context);
 	NIAGARA_API void GetTriangleCount(FVectorVMExternalFunctionContext& Context);
 
 	template<typename FilterMode, typename AreaWeightingMode>
-	FORCEINLINE int32 GetFilteredTriangleCount(FSkeletalMeshAccessorHelper& Accessor, FNDISkeletalMesh_InstanceData* InstData);
+	inline int32 GetFilteredTriangleCount(FSkeletalMeshAccessorHelper& Accessor, FNDISkeletalMesh_InstanceData* InstData);
 
 	template<typename FilterMode, typename AreaWeightingMode>
-	FORCEINLINE int32 GetFilteredTriangleAt(FSkeletalMeshAccessorHelper& Accessor, FNDISkeletalMesh_InstanceData* InstData, int32 FilteredIdx);
+	inline int32 GetFilteredTriangleAt(FSkeletalMeshAccessorHelper& Accessor, FNDISkeletalMesh_InstanceData* InstData, int32 FilteredIdx);
 	//End of Mesh Sampling
 	//////////////////////////////////////////////////////////////////////////
 
@@ -953,13 +953,13 @@ public:
 
 private:
 	template<typename FilterMode>
-	FORCEINLINE int32 RandomFilteredVertIndex(FNDIRandomHelper& RandHelper, int32 Instance, FSkeletalMeshAccessorHelper& Accessor, FNDISkeletalMesh_InstanceData* InstData);
+	inline int32 RandomFilteredVertIndex(FNDIRandomHelper& RandHelper, int32 Instance, FSkeletalMeshAccessorHelper& Accessor, FNDISkeletalMesh_InstanceData* InstData);
 
 	template<typename FilterMode>
-	FORCEINLINE int32 GetFilteredVertexCount(FSkeletalMeshAccessorHelper& Accessor, FNDISkeletalMesh_InstanceData* InstData);
+	inline int32 GetFilteredVertexCount(FSkeletalMeshAccessorHelper& Accessor, FNDISkeletalMesh_InstanceData* InstData);
 
 	template<typename FilterMode>
-	FORCEINLINE int32 GetFilteredVertexAt(FSkeletalMeshAccessorHelper& Accessor, FNDISkeletalMesh_InstanceData* InstData, int32 FilteredIdx);
+	inline int32 GetFilteredVertexAt(FSkeletalMeshAccessorHelper& Accessor, FNDISkeletalMesh_InstanceData* InstData, int32 FilteredIdx);
 
 	//End of Vertex Sampling
 	//////////////////////////////////////////////////////////////////////////
diff --git a/Engine/Plugins/FX/Niagara/Source/Niagara/Classes/NiagaraDataInterfaceSpline.h b/Engine/Plugins/FX/Niagara/Source/Niagara/Classes/NiagaraDataInterfaceSpline.h
--- a/Engine/Plugins/FX/Niagara/Source/Niagara/Classes/NiagaraDataInterfaceSpline.h
+++ b/Engine/Plugins/FX/Niagara/Source/Niagara/Classes/NiagaraDataInterfaceSpline.h
@@ -111,7 +111,7 @@ struct FNDISpline_InstanceData
 	// We cache the version of the current spline curves so that we can reset the curves structure if we're using the LUT
 	uint32 SplineCurvesVersion;
 
-	FORCEINLINE_DEBUGGABLE bool ResetRequired(UNiagaraDataInterfaceSpline* Interface, FNiagaraSystemInstance* SystemInstance) const;
+	inline bool ResetRequired(UNiagaraDataInterfaceSpline* Interface, FNiagaraSystemInstance* SystemInstance) const;
 	
 	template<typename UseLUT>
 	float GetSplineLength() const;
diff --git a/Engine/Plugins/FX/Niagara/Source/Niagara/Classes/NiagaraDataInterfaceVector2DCurve.h b/Engine/Plugins/FX/Niagara/Source/Niagara/Classes/NiagaraDataInterfaceVector2DCurve.h
--- a/Engine/Plugins/FX/Niagara/Source/Niagara/Classes/NiagaraDataInterfaceVector2DCurve.h
+++ b/Engine/Plugins/FX/Niagara/Source/Niagara/Classes/NiagaraDataInterfaceVector2DCurve.h
@@ -72,7 +72,7 @@ protected:
 
 private:
 	template<typename UseLUT>
-	FORCEINLINE_DEBUGGABLE FVector2f SampleCurveInternal(float X);
+	inline FVector2f SampleCurveInternal(float X);
 
 	static NIAGARA_API const FName SampleCurveName;
 };
diff --git a/Engine/Plugins/FX/Niagara/Source/Niagara/Classes/NiagaraDataInterfaceVector4Curve.h b/Engine/Plugins/FX/Niagara/Source/Niagara/Classes/NiagaraDataInterfaceVector4Curve.h
--- a/Engine/Plugins/FX/Niagara/Source/Niagara/Classes/NiagaraDataInterfaceVector4Curve.h
+++ b/Engine/Plugins/FX/Niagara/Source/Niagara/Classes/NiagaraDataInterfaceVector4Curve.h
@@ -86,7 +86,7 @@ protected:
 
 private: 
 	template<typename UseLUT>
-	FORCEINLINE_DEBUGGABLE FVector4f SampleCurveInternal(float X);
+	inline FVector4f SampleCurveInternal(float X);
 
 	static NIAGARA_API const FName SampleCurveName;
 };
diff --git a/Engine/Plugins/FX/Niagara/Source/Niagara/Classes/NiagaraDataInterfaceVectorCurve.h b/Engine/Plugins/FX/Niagara/Source/Niagara/Classes/NiagaraDataInterfaceVectorCurve.h
--- a/Engine/Plugins/FX/Niagara/Source/Niagara/Classes/NiagaraDataInterfaceVectorCurve.h
+++ b/Engine/Plugins/FX/Niagara/Source/Niagara/Classes/NiagaraDataInterfaceVectorCurve.h
@@ -83,7 +83,7 @@ protected:
 
 private:
 	template<typename UseLUT>
-	FORCEINLINE_DEBUGGABLE FVector3f SampleCurveInternal(float X);
+	inline FVector3f SampleCurveInternal(float X);
 
 	static NIAGARA_API const FName SampleCurveName;
 };
