Replace some usages of FORCEINLINE with inline in Niagara modules.

diff --git a/Engine/Plugins/FX/Niagara/Source/Niagara/Classes/NiagaraDataSet.h b/Engine/Plugins/FX/Niagara/Source/Niagara/Classes/NiagaraDataSet.h
--- a/Engine/Plugins/FX/Niagara/Source/Niagara/Classes/NiagaraDataSet.h
+++ b/Engine/Plugins/FX/Niagara/Source/Niagara/Classes/NiagaraDataSet.h
@@ -29,23 +29,23 @@ public:
 	void Destroy();
 	static void FlushDeletionList();
 
-	FORCEINLINE bool IsInUse()const { return ReadRefCount.Load() != 0; }
-	FORCEINLINE bool IsBeingRead()const { return ReadRefCount.Load() > 0; }
-	FORCEINLINE bool IsBeingWritten()const { return ReadRefCount.Load() == INDEX_NONE; }
+	inline bool IsInUse()const { return ReadRefCount.Load() != 0; }
+	inline bool IsBeingRead()const { return ReadRefCount.Load() > 0; }
+	inline bool IsBeingWritten()const { return ReadRefCount.Load() == INDEX_NONE; }
 
-	FORCEINLINE void AddRef()
+	inline void AddRef()
 	{
 		check(!IsBeingWritten());
 		ReadRefCount++;
 	}
 
-	FORCEINLINE void Release()
+	inline void Release()
 	{
 		check(IsBeingRead());
 		ReadRefCount--;
 	}
 
-	FORCEINLINE bool TryLock()
+	inline bool TryLock()
 	{
 		//Only lock if we have no readers.
 		//Using INDEX_NONE as a special case value for write locks.
@@ -54,14 +54,14 @@ public:
 	}
 
 	//Remove the write lock
-	FORCEINLINE void Unlock()
+	inline void Unlock()
 	{
 		int32 Expected = INDEX_NONE;
 		ensureAlwaysMsgf(ReadRefCount.CompareExchange(Expected, 0), TEXT("Trying to release a write lock on a Niagara shared object that is not locked for write."));
 	}
 	
 	//Removes the write lock and moves directly into a read state.
-	FORCEINLINE TRefCountPtr<FNiagaraSharedObject> UnlockForRead()
+	inline TRefCountPtr<FNiagaraSharedObject> UnlockForRead()
 	{
 		int32 Expected = INDEX_NONE;
 		ensureAlwaysMsgf(ReadRefCount.CompareExchange(Expected, 1), TEXT("Trying to release a write lock on a Niagara shared object that is not locked for write."));
@@ -92,7 +92,7 @@ protected:
 	NIAGARA_API virtual ~FNiagaraDataBuffer();
 
 public:
-	FORCEINLINE TRefCountPtr<FNiagaraDataBuffer> UnlockForRead()
+	inline TRefCountPtr<FNiagaraDataBuffer> UnlockForRead()
 	{
 		int32 Expected = INDEX_NONE;
 		ensureAlwaysMsgf(ReadRefCount.CompareExchange(Expected, 1), TEXT("Trying to release a write lock on a Niagara shared object that is not locked for write."));
@@ -116,74 +116,74 @@ public:
 	NIAGARA_API void TransferGPUToCPUImmediate(FRHICommandListImmediate& RHICmdList, FNiagaraGpuComputeDispatchInterface* ComputeInterface, FNiagaraDataBuffer* CPUBuffer) const;
 	NIAGARA_API void Dump(int32 StartIndex, int32 NumInstances, const FString& Label, const FName& SortParameterKey = FName()) const;
 
-	FORCEINLINE TArrayView<uint8 const* RESTRICT const> ReadRegisterTable() const { return TArrayView<uint8 const* RESTRICT const>(RegisterTable); }
-	FORCEINLINE TArrayView<uint8* RESTRICT const> EditRegisterTable() const { return TArrayView<uint8* RESTRICT const>(RegisterTable); }
+	inline TArrayView<uint8 const* RESTRICT const> ReadRegisterTable() const { return TArrayView<uint8 const* RESTRICT const>(RegisterTable); }
+	inline TArrayView<uint8* RESTRICT const> EditRegisterTable() const { return TArrayView<uint8* RESTRICT const>(RegisterTable); }
 	
 	typedef uint32 RegisterTypeOffsetType[3];
-	FORCEINLINE const RegisterTypeOffsetType& GetRegisterTypeOffsets() { return RegisterTypeOffsets; }
+	inline const RegisterTypeOffsetType& GetRegisterTypeOffsets() { return RegisterTypeOffsets; }
 
-	FORCEINLINE const TArray<uint8>& GetFloatBuffer() const { return FloatData; }
-	FORCEINLINE const TArray<uint8>& GetInt32Buffer() const { return Int32Data; }
-	FORCEINLINE const TArray<uint8>& GetHalfBuffer() const { return HalfData; }
+	inline const TArray<uint8>& GetFloatBuffer() const { return FloatData; }
+	inline const TArray<uint8>& GetInt32Buffer() const { return Int32Data; }
+	inline const TArray<uint8>& GetHalfBuffer() const { return HalfData; }
 
-	FORCEINLINE const uint8* GetComponentPtrFloat(uint32 ComponentIdx) const { return FloatData.GetData() + FloatStride * ComponentIdx; }
-	FORCEINLINE const uint8* GetComponentPtrInt32(uint32 ComponentIdx) const { return Int32Data.GetData() + Int32Stride * ComponentIdx; }
-	FORCEINLINE const uint8* GetComponentPtrHalf(uint32 ComponentIdx) const { return HalfData.GetData() + HalfStride * ComponentIdx; }
-	FORCEINLINE uint8* GetComponentPtrFloat(uint32 ComponentIdx) { return FloatData.GetData() + FloatStride * ComponentIdx; }
-	FORCEINLINE uint8* GetComponentPtrInt32(uint32 ComponentIdx) { return Int32Data.GetData() + Int32Stride * ComponentIdx;	}
-	FORCEINLINE uint8* GetComponentPtrHalf(uint32 ComponentIdx) { return HalfData.GetData() + HalfStride * ComponentIdx; }
+	inline const uint8* GetComponentPtrFloat(uint32 ComponentIdx) const { return FloatData.GetData() + FloatStride * ComponentIdx; }
+	inline const uint8* GetComponentPtrInt32(uint32 ComponentIdx) const { return Int32Data.GetData() + Int32Stride * ComponentIdx; }
+	inline const uint8* GetComponentPtrHalf(uint32 ComponentIdx) const { return HalfData.GetData() + HalfStride * ComponentIdx; }
+	inline uint8* GetComponentPtrFloat(uint32 ComponentIdx) { return FloatData.GetData() + FloatStride * ComponentIdx; }
+	inline uint8* GetComponentPtrInt32(uint32 ComponentIdx) { return Int32Data.GetData() + Int32Stride * ComponentIdx;	}
+	inline uint8* GetComponentPtrHalf(uint32 ComponentIdx) { return HalfData.GetData() + HalfStride * ComponentIdx; }
 
-	FORCEINLINE float* GetInstancePtrFloat(uint32 ComponentIdx, uint32 InstanceIdx)	{ return (float*)(GetComponentPtrFloat(ComponentIdx)) + InstanceIdx; }
-	FORCEINLINE int32* GetInstancePtrInt32(uint32 ComponentIdx, uint32 InstanceIdx)	{ return (int32*)(GetComponentPtrInt32(ComponentIdx)) + InstanceIdx; }
-	FORCEINLINE FFloat16* GetInstancePtrHalf(uint32 ComponentIdx, uint32 InstanceIdx) { return (FFloat16*)(GetComponentPtrHalf(ComponentIdx)) + InstanceIdx; }
+	inline float* GetInstancePtrFloat(uint32 ComponentIdx, uint32 InstanceIdx)	{ return (float*)(GetComponentPtrFloat(ComponentIdx)) + InstanceIdx; }
+	inline int32* GetInstancePtrInt32(uint32 ComponentIdx, uint32 InstanceIdx)	{ return (int32*)(GetComponentPtrInt32(ComponentIdx)) + InstanceIdx; }
+	inline FFloat16* GetInstancePtrHalf(uint32 ComponentIdx, uint32 InstanceIdx) { return (FFloat16*)(GetComponentPtrHalf(ComponentIdx)) + InstanceIdx; }
 
-	FORCEINLINE const float* GetInstancePtrFloat(uint32 ComponentIdx, uint32 InstanceIdx) const { return (float*)(GetComponentPtrFloat(ComponentIdx)) + InstanceIdx; }
-	FORCEINLINE const int32* GetInstancePtrInt32(uint32 ComponentIdx, uint32 InstanceIdx) const { return (int32*)(GetComponentPtrInt32(ComponentIdx)) + InstanceIdx; }
-	FORCEINLINE const FFloat16* GetInstancePtrHalf(uint32 ComponentIdx, uint32 InstanceIdx) const { return (FFloat16*)(GetComponentPtrHalf(ComponentIdx)) + InstanceIdx; }
+	inline const float* GetInstancePtrFloat(uint32 ComponentIdx, uint32 InstanceIdx) const { return (float*)(GetComponentPtrFloat(ComponentIdx)) + InstanceIdx; }
+	inline const int32* GetInstancePtrInt32(uint32 ComponentIdx, uint32 InstanceIdx) const { return (int32*)(GetComponentPtrInt32(ComponentIdx)) + InstanceIdx; }
+	inline const FFloat16* GetInstancePtrHalf(uint32 ComponentIdx, uint32 InstanceIdx) const { return (FFloat16*)(GetComponentPtrHalf(ComponentIdx)) + InstanceIdx; }
 
-	FORCEINLINE uint8* GetComponentPtrFloat(float* BasePtr, uint32 ComponentIdx) const { return (uint8*)BasePtr + FloatStride * ComponentIdx; }
-	FORCEINLINE uint8* GetComponentPtrInt32(int* BasePtr, uint32 ComponentIdx) const { return (uint8*)BasePtr + Int32Stride * ComponentIdx; }
-	FORCEINLINE uint8* GetComponentPtrHalf(FFloat16* BasePtr, uint32 ComponentIdx) const { return (uint8*)(BasePtr) + HalfStride * ComponentIdx; }
+	inline uint8* GetComponentPtrFloat(float* BasePtr, uint32 ComponentIdx) const { return (uint8*)BasePtr + FloatStride * ComponentIdx; }
+	inline uint8* GetComponentPtrInt32(int* BasePtr, uint32 ComponentIdx) const { return (uint8*)BasePtr + Int32Stride * ComponentIdx; }
+	inline uint8* GetComponentPtrHalf(FFloat16* BasePtr, uint32 ComponentIdx) const { return (uint8*)(BasePtr) + HalfStride * ComponentIdx; }
 
-	FORCEINLINE float* GetInstancePtrFloat(float* BasePtr, uint32 ComponentIdx, uint32 InstanceIdx)const { return (float*)GetComponentPtrFloat(BasePtr, ComponentIdx) + InstanceIdx; }
-	FORCEINLINE int32* GetInstancePtrInt32(int* BasePtr, uint32 ComponentIdx, uint32 InstanceIdx)const { return (int32*)GetComponentPtrInt32(BasePtr, ComponentIdx) + InstanceIdx; }
-	FORCEINLINE FFloat16* GetInstancePtrHalf(FFloat16 *BasePtr, uint32 ComponentIdx, uint32 InstanceIdx)const { return (FFloat16*)GetComponentPtrHalf(BasePtr, ComponentIdx) + InstanceIdx; }
+	inline float* GetInstancePtrFloat(float* BasePtr, uint32 ComponentIdx, uint32 InstanceIdx)const { return (float*)GetComponentPtrFloat(BasePtr, ComponentIdx) + InstanceIdx; }
+	inline int32* GetInstancePtrInt32(int* BasePtr, uint32 ComponentIdx, uint32 InstanceIdx)const { return (int32*)GetComponentPtrInt32(BasePtr, ComponentIdx) + InstanceIdx; }
+	inline FFloat16* GetInstancePtrHalf(FFloat16 *BasePtr, uint32 ComponentIdx, uint32 InstanceIdx)const { return (FFloat16*)GetComponentPtrHalf(BasePtr, ComponentIdx) + InstanceIdx; }
 
-	FORCEINLINE uint32 GetNumInstances()const { return NumInstances; }
-	FORCEINLINE uint32 GetNumInstancesAllocated()const { return NumInstancesAllocated; }
-	FORCEINLINE uint32 GetNumSpawnedInstances() const { return NumSpawnedInstances; }
+	inline uint32 GetNumInstances()const { return NumInstances; }
+	inline uint32 GetNumInstancesAllocated()const { return NumInstancesAllocated; }
+	inline uint32 GetNumSpawnedInstances() const { return NumSpawnedInstances; }
 
-	FORCEINLINE void SetNumInstances(uint32 InNumInstances) { check(InNumInstances <= NumInstancesAllocated); NumInstances = InNumInstances; }
-	FORCEINLINE void SetNumSpawnedInstances(uint32 InNumSpawnedInstances) { NumSpawnedInstances = InNumSpawnedInstances; }
+	inline void SetNumInstances(uint32 InNumInstances) { check(InNumInstances <= NumInstancesAllocated); NumInstances = InNumInstances; }
+	inline void SetNumSpawnedInstances(uint32 InNumSpawnedInstances) { NumSpawnedInstances = InNumSpawnedInstances; }
 
-	FORCEINLINE void SetGPUDataReadyStage(ENiagaraGpuComputeTickStage::Type InReadyStage) { GPUDataReadyStage = InReadyStage; }
-	FORCEINLINE ENiagaraGpuComputeTickStage::Type GetGPUDataReadyStage() const { return GPUDataReadyStage; }
-	FORCEINLINE FRWBuffer& GetGPUBufferFloat() { return GPUBufferFloat; }
-	FORCEINLINE FRWBuffer& GetGPUBufferInt() { return GPUBufferInt; }
-	FORCEINLINE FRWBuffer& GetGPUBufferHalf() { return GPUBufferHalf; }
-	FORCEINLINE uint32 GetGPUInstanceCountBufferOffset() const { return GPUInstanceCountBufferOffset; }
-	FORCEINLINE FRWBuffer& GetGPUIDToIndexTable() { return GPUIDToIndexTable; }
+	inline void SetGPUDataReadyStage(ENiagaraGpuComputeTickStage::Type InReadyStage) { GPUDataReadyStage = InReadyStage; }
+	inline ENiagaraGpuComputeTickStage::Type GetGPUDataReadyStage() const { return GPUDataReadyStage; }
+	inline FRWBuffer& GetGPUBufferFloat() { return GPUBufferFloat; }
+	inline FRWBuffer& GetGPUBufferInt() { return GPUBufferInt; }
+	inline FRWBuffer& GetGPUBufferHalf() { return GPUBufferHalf; }
+	inline uint32 GetGPUInstanceCountBufferOffset() const { return GPUInstanceCountBufferOffset; }
+	inline FRWBuffer& GetGPUIDToIndexTable() { return GPUIDToIndexTable; }
 
-	FORCEINLINE void SetGPUInstanceCountBufferOffset(uint32 Offset) { GPUInstanceCountBufferOffset = Offset; }
+	inline void SetGPUInstanceCountBufferOffset(uint32 Offset) { GPUInstanceCountBufferOffset = Offset; }
 
-	FORCEINLINE int32 GetSafeComponentBufferSize() const { return GetSafeComponentBufferSize(GetNumInstancesAllocated()); }
-	FORCEINLINE uint32 GetFloatStride() const { return FloatStride; }
-	FORCEINLINE uint32 GetInt32Stride() const { return Int32Stride; }
-	FORCEINLINE uint32 GetHalfStride() const { return HalfStride; }
+	inline int32 GetSafeComponentBufferSize() const { return GetSafeComponentBufferSize(GetNumInstancesAllocated()); }
+	inline uint32 GetFloatStride() const { return FloatStride; }
+	inline uint32 GetInt32Stride() const { return Int32Stride; }
+	inline uint32 GetHalfStride() const { return HalfStride; }
 
-	FORCEINLINE uint32 GetIDAcquireTag() const { return IDAcquireTag; }
-	FORCEINLINE void SetIDAcquireTag(uint32 InTag) { IDAcquireTag = InTag; }
+	inline uint32 GetIDAcquireTag() const { return IDAcquireTag; }
+	inline void SetIDAcquireTag(uint32 InTag) { IDAcquireTag = InTag; }
 
-	FORCEINLINE FNiagaraDataSet* GetOwner()const { return Owner; }
+	inline FNiagaraDataSet* GetOwner()const { return Owner; }
 
 	NIAGARA_API int32 TransferInstance(FNiagaraDataBuffer& SourceBuffer, int32 InstanceIndex, bool bRemoveFromSource=true);
 
 	NIAGARA_API bool CheckForNaNs()const;
 
-	FORCEINLINE TArray<int32>& GetIDTable() { return IDToIndexTable; }
-	FORCEINLINE const TArray<int32>& GetIDTable() const { return IDToIndexTable; }
+	inline TArray<int32>& GetIDTable() { return IDToIndexTable; }
+	inline const TArray<int32>& GetIDTable() const { return IDToIndexTable; }
 
-	FORCEINLINE void ClearGPUInstanceCount() { GPUInstanceCountBufferOffset = INDEX_NONE; }
+	inline void ClearGPUInstanceCount() { GPUInstanceCountBufferOffset = INDEX_NONE; }
 
 	NIAGARA_API void BuildRegisterTable();
 	
@@ -195,9 +195,9 @@ public:
 
 
 private:
-	FORCEINLINE void CheckUsage(bool bReadOnly)const;
+	inline void CheckUsage(bool bReadOnly)const;
 
-	FORCEINLINE int32 GetSafeComponentBufferSize(int32 RequiredSize) const
+	inline int32 GetSafeComponentBufferSize(int32 RequiredSize) const
 	{
 		//Round up to VECTOR_WIDTH_BYTES.
 		//Both aligns the component buffers to the vector width but also ensures their ops cannot stomp over one another.		
@@ -299,29 +299,29 @@ public:
 	/** Returns size in bytes for all data buffers currently allocated by this dataset. */
 	NIAGARA_API int64 GetSizeBytes() const;
 
-	FORCEINLINE bool IsInitialized() const { return bInitialized; }
-	FORCEINLINE ENiagaraSimTarget GetSimTarget() const { return CompiledData->SimTarget; }
-	FORCEINLINE FNiagaraDataSetID GetID() const { return CompiledData->ID; }	
-	FORCEINLINE bool RequiresPersistentIDs() const { return CompiledData->bRequiresPersistentIDs; }
-
-	FORCEINLINE TArray<int32>& GetFreeIDTable() { return FreeIDsTable; }
-	FORCEINLINE TArray<int32>& GetSpawnedIDsTable() { return SpawnedIDsTable; }
-	FORCEINLINE const TArray<int32>& GetFreeIDTable() const { return FreeIDsTable; }
-	FORCEINLINE const TArray<int32>& GetSpawnedIDsTable() const { return SpawnedIDsTable; }
-	FORCEINLINE int32* GetNumFreeIDsPtr() { return &NumFreeIDs; }
-	FORCEINLINE int32* GetMaxUsedIDPtr() { return &MaxUsedID; }
-	FORCEINLINE int32 GetIDAcquireTag() const { return IDAcquireTag; }
-	FORCEINLINE void SetIDAcquireTag(int32 InTag) { IDAcquireTag = InTag; }
-	FORCEINLINE FRWBuffer& GetGPUFreeIDs() { return GPUFreeIDs; }
-	FORCEINLINE uint32 GetGPUNumAllocatedIDs() const { return GPUNumAllocatedIDs; }
-
-	FORCEINLINE const TArray<FNiagaraVariableBase>& GetVariables() const { return CompiledData->Variables; }
-	FORCEINLINE uint32 GetNumVariables() const { return CompiledData->Variables.Num(); }
-	FORCEINLINE bool HasVariable(const FNiagaraVariableBase& Var) const { return CompiledData->Variables.Contains(Var); }
-	FORCEINLINE bool HasVariable(const FName& Name) const { return CompiledData->Variables.FindByPredicate([&](const FNiagaraVariableBase& VarInfo) { return VarInfo.GetName() == Name; }) != nullptr; }
-	FORCEINLINE uint32 GetNumFloatComponents() const { return CompiledData->TotalFloatComponents; }
-	FORCEINLINE uint32 GetNumInt32Components() const { return CompiledData->TotalInt32Components; }
-	FORCEINLINE uint32 GetNumHalfComponents() const { return CompiledData->TotalHalfComponents; }
+	inline bool IsInitialized() const { return bInitialized; }
+	inline ENiagaraSimTarget GetSimTarget() const { return CompiledData->SimTarget; }
+	inline FNiagaraDataSetID GetID() const { return CompiledData->ID; }	
+	inline bool RequiresPersistentIDs() const { return CompiledData->bRequiresPersistentIDs; }
+
+	inline TArray<int32>& GetFreeIDTable() { return FreeIDsTable; }
+	inline TArray<int32>& GetSpawnedIDsTable() { return SpawnedIDsTable; }
+	inline const TArray<int32>& GetFreeIDTable() const { return FreeIDsTable; }
+	inline const TArray<int32>& GetSpawnedIDsTable() const { return SpawnedIDsTable; }
+	inline int32* GetNumFreeIDsPtr() { return &NumFreeIDs; }
+	inline int32* GetMaxUsedIDPtr() { return &MaxUsedID; }
+	inline int32 GetIDAcquireTag() const { return IDAcquireTag; }
+	inline void SetIDAcquireTag(int32 InTag) { IDAcquireTag = InTag; }
+	inline FRWBuffer& GetGPUFreeIDs() { return GPUFreeIDs; }
+	inline uint32 GetGPUNumAllocatedIDs() const { return GPUNumAllocatedIDs; }
+
+	inline const TArray<FNiagaraVariableBase>& GetVariables() const { return CompiledData->Variables; }
+	inline uint32 GetNumVariables() const { return CompiledData->Variables.Num(); }
+	inline bool HasVariable(const FNiagaraVariableBase& Var) const { return CompiledData->Variables.Contains(Var); }
+	inline bool HasVariable(const FName& Name) const { return CompiledData->Variables.FindByPredicate([&](const FNiagaraVariableBase& VarInfo) { return VarInfo.GetName() == Name; }) != nullptr; }
+	inline uint32 GetNumFloatComponents() const { return CompiledData->TotalFloatComponents; }
+	inline uint32 GetNumInt32Components() const { return CompiledData->TotalInt32Components; }
+	inline uint32 GetNumHalfComponents() const { return CompiledData->TotalHalfComponents; }
 
 	const TArray<FNiagaraVariableLayoutInfo>& GetVariableLayouts() const { return CompiledData->VariableLayouts; }
 	NIAGARA_API const FNiagaraVariableLayoutInfo* GetVariableLayout(const FNiagaraVariableBase& Var) const;
@@ -335,17 +335,17 @@ public:
 
 	NIAGARA_API void Dump(int32 StartIndex, int32 NumInstances, const FString& Label, const FName& SortParameterKey = FName()) const;
 
-	FORCEINLINE bool IsCurrentDataValid()const { return CurrentData != nullptr; }
-	FORCEINLINE FNiagaraDataBuffer* GetCurrentData()const {	return CurrentData; }
-	FORCEINLINE FNiagaraDataBuffer* GetDestinationData()const { return DestinationData; }
+	inline bool IsCurrentDataValid()const { return CurrentData != nullptr; }
+	inline FNiagaraDataBuffer* GetCurrentData()const {	return CurrentData; }
+	inline FNiagaraDataBuffer* GetDestinationData()const { return DestinationData; }
 
-	FORCEINLINE FNiagaraDataBuffer& GetCurrentDataChecked() const
+	inline FNiagaraDataBuffer& GetCurrentDataChecked() const
 	{
 		check(CurrentData);
 		return *CurrentData;
 	}
 
-	FORCEINLINE FNiagaraDataBuffer& GetDestinationDataChecked() const
+	inline FNiagaraDataBuffer& GetDestinationDataChecked() const
 	{
 		check(DestinationData);
 		return *DestinationData;
@@ -367,7 +367,7 @@ private:
 
 	NIAGARA_API void ResetBuffersInternal();
 
-	FORCEINLINE void CheckCorrectThread()const
+	inline void CheckCorrectThread()const
 	{
 		// In some rare occasions, the render thread might be null, like when offloading work to Lightmass 
 		// The final GIsThreadedRendering check keeps us from inadvertently failing when that happens.
@@ -501,7 +501,7 @@ class FScopedNiagaraDataSetGPUReadback
 {
 public:
 	FScopedNiagaraDataSetGPUReadback() {}
-	FORCEINLINE ~FScopedNiagaraDataSetGPUReadback()
+	inline ~FScopedNiagaraDataSetGPUReadback()
 	{
 		if (DataBuffer != nullptr)
 		{
@@ -523,7 +523,7 @@ private:
 
 //////////////////////////////////////////////////////////////////////////
 
-FORCEINLINE void FNiagaraDataBuffer::CheckUsage(bool bReadOnly)const
+inline void FNiagaraDataBuffer::CheckUsage(bool bReadOnly)const
 {
 	checkSlow(Owner);
 
diff --git a/Engine/Plugins/FX/Niagara/Source/Niagara/Classes/NiagaraDataSetAccessor.h b/Engine/Plugins/FX/Niagara/Source/Niagara/Classes/NiagaraDataSetAccessor.h
--- a/Engine/Plugins/FX/Niagara/Source/Niagara/Classes/NiagaraDataSetAccessor.h
+++ b/Engine/Plugins/FX/Niagara/Source/Niagara/Classes/NiagaraDataSetAccessor.h
@@ -31,10 +31,10 @@ template<typename TType>
 struct FNiagaraDataSetAccessor : public FNiagaraDataSetAccessorTypeInfo<TType>::TAccessorBaseClass
 {
 	static_assert(!TIsUECoreVariant<TType, double>::Value, "Double core variant. Must be float type!");
-	FORCEINLINE FNiagaraDataSetAccessor<TType>() {}
-	FORCEINLINE explicit FNiagaraDataSetAccessor<TType>(const FNiagaraDataSet& DataSet, const FName VariableName) { FNiagaraDataSetAccessorTypeInfo<TType>::TAccessorBaseClass::Init(NiagaraDataSetPrivate::GetCompiledData(DataSet), VariableName); }
-	FORCEINLINE explicit FNiagaraDataSetAccessor<TType>(const FNiagaraDataBuffer* DataBuffer, const FName VariableName) { FNiagaraDataSetAccessorTypeInfo<TType>::TAccessorBaseClass::Init(NiagaraDataSetPrivate::GetCompiledData(DataBuffer), VariableName); }
-	FORCEINLINE explicit FNiagaraDataSetAccessor<TType>(const FNiagaraDataSetCompiledData& DataSetCompiledData, const FName VariableName) { FNiagaraDataSetAccessorTypeInfo<TType>::TAccessorBaseClass::Init(DataSetCompiledData, VariableName); }
+	inline FNiagaraDataSetAccessor<TType>() {}
+	inline explicit FNiagaraDataSetAccessor<TType>(const FNiagaraDataSet& DataSet, const FName VariableName) { FNiagaraDataSetAccessorTypeInfo<TType>::TAccessorBaseClass::Init(NiagaraDataSetPrivate::GetCompiledData(DataSet), VariableName); }
+	inline explicit FNiagaraDataSetAccessor<TType>(const FNiagaraDataBuffer* DataBuffer, const FName VariableName) { FNiagaraDataSetAccessorTypeInfo<TType>::TAccessorBaseClass::Init(NiagaraDataSetPrivate::GetCompiledData(DataBuffer), VariableName); }
+	inline explicit FNiagaraDataSetAccessor<TType>(const FNiagaraDataSetCompiledData& DataSetCompiledData, const FName VariableName) { FNiagaraDataSetAccessorTypeInfo<TType>::TAccessorBaseClass::Init(DataSetCompiledData, VariableName); }
 };
 
 //////////////////////////////////////////////////////////////////////////
@@ -46,7 +46,7 @@ struct FNiagaraDataSetReaderFloat
 	static_assert(!TIsUECoreVariant<TType, double>::Value, "Double core variant. Must be float type!");
 	static constexpr bool bSupportsHalf = FNiagaraDataSetAccessorTypeInfo<TType>::bSupportsHalf;
 
-	FORCEINLINE FNiagaraDataSetReaderFloat() {}
+	inline FNiagaraDataSetReaderFloat() {}
 
 	explicit FNiagaraDataSetReaderFloat(const FNiagaraDataBuffer* DataBuffer, bool bInIsFloat, int32 ComponentIndex)
 		: bIsFloat(bInIsFloat)
@@ -64,14 +64,14 @@ struct FNiagaraDataSetReaderFloat
 		}
 	}
 
-	FORCEINLINE bool IsValid() const { return bIsValid; }
+	inline bool IsValid() const { return bIsValid; }
 
-	FORCEINLINE TType operator[](int32 Index) const
+	inline TType operator[](int32 Index) const
 	{
 		return Get(Index);
 	}
 
-	FORCEINLINE TType Get(int Index) const
+	inline TType Get(int Index) const
 	{
 		check(bIsValid && Index <= NumInstances);
 		TType Value;
@@ -92,7 +92,7 @@ struct FNiagaraDataSetReaderFloat
 		return Value;
 	}
 
-	FORCEINLINE TType GetSafe(int Index, const TType& Default) const
+	inline TType GetSafe(int Index, const TType& Default) const
 	{
 		if (bIsValid && (Index < NumInstances))
 		{
@@ -101,7 +101,7 @@ struct FNiagaraDataSetReaderFloat
 		return Default;
 	}
 
-	FORCEINLINE float GetMaxElement(int32 ElementIndex) const
+	inline float GetMaxElement(int32 ElementIndex) const
 	{
 		check(IsValid());
 		check(ElementIndex < FNiagaraDataSetAccessorTypeInfo<TType>::NumElements);
@@ -125,7 +125,7 @@ struct FNiagaraDataSetReaderFloat
 		}
 	}
 
-	FORCEINLINE void GetMinMaxElement(int32 ElementIndex, float& MinValueOut, float& MaxValueOut) const
+	inline void GetMinMaxElement(int32 ElementIndex, float& MinValueOut, float& MaxValueOut) const
 	{
 		check(IsValid());
 		check(ElementIndex < FNiagaraDataSetAccessorTypeInfo<TType>::NumElements);
@@ -155,7 +155,7 @@ struct FNiagaraDataSetReaderFloat
 		}
 	}
 
-	FORCEINLINE TType GetMax() const
+	inline TType GetMax() const
 	{
 		TType Value;
 		for (int i = 0; i < FNiagaraDataSetAccessorTypeInfo<TType>::NumElements; ++i)
@@ -165,7 +165,7 @@ struct FNiagaraDataSetReaderFloat
 		return Value;
 	}
 
-	FORCEINLINE void GetMinMax(TType& Min, TType& Max) const
+	inline void GetMinMax(TType& Min, TType& Max) const
 	{
 		for (int i = 0; i < FNiagaraDataSetAccessorTypeInfo<TType>::NumElements; ++i)
 		{
@@ -189,12 +189,12 @@ template<typename TType>
 struct FNiagaraDataSetAccessorFloat
 {
 	static_assert(!TIsUECoreVariant<TType, double>::Value, "Double core variant. Must be float type!");
-	FORCEINLINE FNiagaraDataSetAccessorFloat() {}
-	FORCEINLINE explicit FNiagaraDataSetAccessorFloat(const FNiagaraDataSet& DataSet, const FName VariableName) { Init(NiagaraDataSetPrivate::GetCompiledData(DataSet), VariableName); }
-	FORCEINLINE explicit FNiagaraDataSetAccessorFloat(const FNiagaraDataBuffer* DataBuffer, const FName VariableName) { Init(NiagaraDataSetPrivate::GetCompiledData(DataBuffer), VariableName); }
-	FORCEINLINE explicit FNiagaraDataSetAccessorFloat(const FNiagaraDataSetCompiledData& DataSetCompiledData, const FName VariableName) { Init(DataSetCompiledData, VariableName); }
-	FORCEINLINE explicit FNiagaraDataSetAccessorFloat(const FNiagaraDataSetCompiledData* DataSetCompiledData, const FName VariableName) { Init(DataSetCompiledData, VariableName); }
-	FORCEINLINE void Init(const FNiagaraDataSet& DataSet, const FName VariableName) { Init(NiagaraDataSetPrivate::GetCompiledData(DataSet), VariableName); }
+	inline FNiagaraDataSetAccessorFloat() {}
+	inline explicit FNiagaraDataSetAccessorFloat(const FNiagaraDataSet& DataSet, const FName VariableName) { Init(NiagaraDataSetPrivate::GetCompiledData(DataSet), VariableName); }
+	inline explicit FNiagaraDataSetAccessorFloat(const FNiagaraDataBuffer* DataBuffer, const FName VariableName) { Init(NiagaraDataSetPrivate::GetCompiledData(DataBuffer), VariableName); }
+	inline explicit FNiagaraDataSetAccessorFloat(const FNiagaraDataSetCompiledData& DataSetCompiledData, const FName VariableName) { Init(DataSetCompiledData, VariableName); }
+	inline explicit FNiagaraDataSetAccessorFloat(const FNiagaraDataSetCompiledData* DataSetCompiledData, const FName VariableName) { Init(DataSetCompiledData, VariableName); }
+	inline void Init(const FNiagaraDataSet& DataSet, const FName VariableName) { Init(NiagaraDataSetPrivate::GetCompiledData(DataSet), VariableName); }
 
 	void Init(const FNiagaraDataSetCompiledData* DataSetCompiledData, const FName VariableName)
 	{
@@ -243,14 +243,14 @@ struct FNiagaraDataSetAccessorFloat
 		}
 	}
 
-	FORCEINLINE bool IsValid() const { return ComponentIndex != INDEX_NONE; }
+	inline bool IsValid() const { return ComponentIndex != INDEX_NONE; }
 
-	FORCEINLINE static FNiagaraDataSetReaderFloat<TType> CreateReader(const FNiagaraDataSet& DataSet, const FName VariableName) { return FNiagaraDataSetAccessorFloat<TType>(DataSet, VariableName).GetReader(DataSet); }
-	FORCEINLINE static FNiagaraDataSetReaderFloat<TType> CreateReader(const FNiagaraDataBuffer* DataBuffer, const FName VariableName) { return FNiagaraDataSetAccessorFloat<TType>(DataBuffer, VariableName).GetReader(DataBuffer); }
+	inline static FNiagaraDataSetReaderFloat<TType> CreateReader(const FNiagaraDataSet& DataSet, const FName VariableName) { return FNiagaraDataSetAccessorFloat<TType>(DataSet, VariableName).GetReader(DataSet); }
+	inline static FNiagaraDataSetReaderFloat<TType> CreateReader(const FNiagaraDataBuffer* DataBuffer, const FName VariableName) { return FNiagaraDataSetAccessorFloat<TType>(DataBuffer, VariableName).GetReader(DataBuffer); }
 	//FORCEINLINE static FNiagaraDataSetWriterFloat<TType> CreateWriter(const FNiagaraDataSet& DataSet, const FName VariableName) { return FNiagaraDataSetAccessorFloat<TType>(DataSet, VariableName).GetWriter(DataSet); }
 
-	FORCEINLINE FNiagaraDataSetReaderFloat<TType> GetReader(const FNiagaraDataSet& DataSet) const { return FNiagaraDataSetReaderFloat<TType>(NiagaraDataSetPrivate::GetCurrentData(DataSet), bIsFloat, ComponentIndex); }
-	FORCEINLINE FNiagaraDataSetReaderFloat<TType> GetReader(const FNiagaraDataBuffer* Data) const { return FNiagaraDataSetReaderFloat<TType>(Data, bIsFloat, ComponentIndex); }
+	inline FNiagaraDataSetReaderFloat<TType> GetReader(const FNiagaraDataSet& DataSet) const { return FNiagaraDataSetReaderFloat<TType>(NiagaraDataSetPrivate::GetCurrentData(DataSet), bIsFloat, ComponentIndex); }
+	inline FNiagaraDataSetReaderFloat<TType> GetReader(const FNiagaraDataBuffer* Data) const { return FNiagaraDataSetReaderFloat<TType>(Data, bIsFloat, ComponentIndex); }
 	//FORCEINLINE FNiagaraDataSetWriterFloat<TType> GetWriter(const FNiagaraDataSet& DataSet) const { return FNiagaraDataSetWriterFloat<TType>(NiagaraDataSetPrivate::GetDestinationData(DataSet), bIsFloat, ComponentIndex); }
 
 private:
@@ -355,7 +355,7 @@ struct FNiagaraDataSetAccessorTypeInfo<FQuat4f>
 template<typename TType>
 struct FNiagaraDataSetReaderInt32
 {
-	FORCEINLINE FNiagaraDataSetReaderInt32() { }
+	inline FNiagaraDataSetReaderInt32() { }
 
 	explicit FNiagaraDataSetReaderInt32(const FNiagaraDataBuffer* DataBuffer, int32 ComponentIndex)
 	{
@@ -371,14 +371,14 @@ struct FNiagaraDataSetReaderInt32
 		}
 	}
 
-	FORCEINLINE bool IsValid() const { return bIsValid; }
+	inline bool IsValid() const { return bIsValid; }
 
-	FORCEINLINE TType operator[](int32 Index) const
+	inline TType operator[](int32 Index) const
 	{
 		return Get(Index);
 	}
 
-	FORCEINLINE TType Get(int Index) const
+	inline TType Get(int Index) const
 	{
 		check(bIsValid && Index <= NumInstances);
 		TType Value;
@@ -389,7 +389,7 @@ struct FNiagaraDataSetReaderInt32
 		return Value;
 	}
 
-	FORCEINLINE TType GetSafe(int Index, const TType& Default) const
+	inline TType GetSafe(int Index, const TType& Default) const
 	{
 		if (bIsValid && (Index < NumInstances))
 		{
@@ -407,7 +407,7 @@ private:
 template<typename TType>
 struct FNiagaraDataSetWriterInt32
 {
-	FORCEINLINE FNiagaraDataSetWriterInt32() { }
+	inline FNiagaraDataSetWriterInt32() { }
 
 	explicit FNiagaraDataSetWriterInt32(FNiagaraDataBuffer* DataBuffer, int32 ComponentIndex)
 	{
@@ -423,9 +423,9 @@ struct FNiagaraDataSetWriterInt32
 		}
 	}
 
-	FORCEINLINE bool IsValid() const { return bIsValid; }
+	inline bool IsValid() const { return bIsValid; }
 
-	FORCEINLINE void Set(int Index, const TType& Value)
+	inline void Set(int Index, const TType& Value)
 	{
 		check(bIsValid && Index <= NumInstances);
 		for (int i = 0; i < FNiagaraDataSetAccessorTypeInfo<TType>::NumElements; ++i)
@@ -434,7 +434,7 @@ struct FNiagaraDataSetWriterInt32
 		}
 	}
 
-	FORCEINLINE void SetSafe(int Index, const TType& Value)
+	inline void SetSafe(int Index, const TType& Value)
 	{
 		if (bIsValid && Index <= NumInstances)
 		{
@@ -459,14 +459,14 @@ private:
 template<typename TType>
 struct FNiagaraDataSetAccessorInt32
 {
-	FORCEINLINE FNiagaraDataSetAccessorInt32() {}
-	FORCEINLINE explicit FNiagaraDataSetAccessorInt32(const FNiagaraDataSet& DataSet, const FName VariableName) { Init(NiagaraDataSetPrivate::GetCompiledData(DataSet), VariableName); }
-	FORCEINLINE explicit FNiagaraDataSetAccessorInt32(const FNiagaraDataBuffer* DataBuffer, const FName VariableName) { Init(NiagaraDataSetPrivate::GetCompiledData(DataBuffer), VariableName); }
-	FORCEINLINE explicit FNiagaraDataSetAccessorInt32(const FNiagaraDataSetCompiledData& DataSetCompiledData, const FName VariableName) { Init(DataSetCompiledData, VariableName); }
-	FORCEINLINE explicit FNiagaraDataSetAccessorInt32(const FNiagaraDataSetCompiledData* DataSetCompiledData, const FName VariableName) { Init(DataSetCompiledData, VariableName); }
-	FORCEINLINE void Init(const FNiagaraDataSet& DataSet, const FName VariableName) { Init(NiagaraDataSetPrivate::GetCompiledData(DataSet), VariableName); }
-
-	FORCEINLINE void Init(const FNiagaraDataSetCompiledData* DataSetCompiledData, const FName VariableName)
+	inline FNiagaraDataSetAccessorInt32() {}
+	inline explicit FNiagaraDataSetAccessorInt32(const FNiagaraDataSet& DataSet, const FName VariableName) { Init(NiagaraDataSetPrivate::GetCompiledData(DataSet), VariableName); }
+	inline explicit FNiagaraDataSetAccessorInt32(const FNiagaraDataBuffer* DataBuffer, const FName VariableName) { Init(NiagaraDataSetPrivate::GetCompiledData(DataBuffer), VariableName); }
+	inline explicit FNiagaraDataSetAccessorInt32(const FNiagaraDataSetCompiledData& DataSetCompiledData, const FName VariableName) { Init(DataSetCompiledData, VariableName); }
+	inline explicit FNiagaraDataSetAccessorInt32(const FNiagaraDataSetCompiledData* DataSetCompiledData, const FName VariableName) { Init(DataSetCompiledData, VariableName); }
+	inline void Init(const FNiagaraDataSet& DataSet, const FName VariableName) { Init(NiagaraDataSetPrivate::GetCompiledData(DataSet), VariableName); }
+
+	inline void Init(const FNiagaraDataSetCompiledData* DataSetCompiledData, const FName VariableName)
 	{
 		if (DataSetCompiledData != nullptr)
 		{
@@ -489,15 +489,15 @@ struct FNiagaraDataSetAccessorInt32
 		}
 	}
 
-	FORCEINLINE bool IsValid() const { return ComponentIndex != INDEX_NONE; }
+	inline bool IsValid() const { return ComponentIndex != INDEX_NONE; }
 
-	FORCEINLINE static FNiagaraDataSetReaderInt32<TType> CreateReader(const FNiagaraDataSet& DataSet, const FName VariableName) { return FNiagaraDataSetAccessorInt32<TType>(DataSet, VariableName).GetReader(DataSet); }
-	FORCEINLINE static FNiagaraDataSetReaderInt32<TType> CreateReader(const FNiagaraDataBuffer* Data, const FName VariableName) { return FNiagaraDataSetAccessorInt32<TType>(Data, VariableName).GetReader(Data); }
-	FORCEINLINE static FNiagaraDataSetWriterInt32<TType> CreateWriter(const FNiagaraDataSet& DataSet, const FName VariableName) { return FNiagaraDataSetAccessorInt32<TType>(DataSet, VariableName).GetWriter(DataSet); }
+	inline static FNiagaraDataSetReaderInt32<TType> CreateReader(const FNiagaraDataSet& DataSet, const FName VariableName) { return FNiagaraDataSetAccessorInt32<TType>(DataSet, VariableName).GetReader(DataSet); }
+	inline static FNiagaraDataSetReaderInt32<TType> CreateReader(const FNiagaraDataBuffer* Data, const FName VariableName) { return FNiagaraDataSetAccessorInt32<TType>(Data, VariableName).GetReader(Data); }
+	inline static FNiagaraDataSetWriterInt32<TType> CreateWriter(const FNiagaraDataSet& DataSet, const FName VariableName) { return FNiagaraDataSetAccessorInt32<TType>(DataSet, VariableName).GetWriter(DataSet); }
 
-	FORCEINLINE FNiagaraDataSetReaderInt32<TType> GetReader(const FNiagaraDataSet& DataSet) const { return FNiagaraDataSetReaderInt32<TType>(NiagaraDataSetPrivate::GetCurrentData(DataSet), ComponentIndex); }
-	FORCEINLINE FNiagaraDataSetReaderInt32<TType> GetReader(const FNiagaraDataBuffer* Data) const { return FNiagaraDataSetReaderInt32<TType>(Data, ComponentIndex); }
-	FORCEINLINE FNiagaraDataSetWriterInt32<TType> GetWriter(const FNiagaraDataSet& DataSet) const { return FNiagaraDataSetWriterInt32<TType>(NiagaraDataSetPrivate::GetDestinationData(DataSet), ComponentIndex); }
+	inline FNiagaraDataSetReaderInt32<TType> GetReader(const FNiagaraDataSet& DataSet) const { return FNiagaraDataSetReaderInt32<TType>(NiagaraDataSetPrivate::GetCurrentData(DataSet), ComponentIndex); }
+	inline FNiagaraDataSetReaderInt32<TType> GetReader(const FNiagaraDataBuffer* Data) const { return FNiagaraDataSetReaderInt32<TType>(Data, ComponentIndex); }
+	inline FNiagaraDataSetWriterInt32<TType> GetWriter(const FNiagaraDataSet& DataSet) const { return FNiagaraDataSetWriterInt32<TType>(NiagaraDataSetPrivate::GetDestinationData(DataSet), ComponentIndex); }
 
 private:
 	int32 ComponentIndex = INDEX_NONE;
@@ -537,7 +537,7 @@ struct FNiagaraDataSetAccessorTypeInfo<ENiagaraExecutionState>
 template<typename TType>
 struct FNiagaraDataSetReaderStruct
 {
-	FORCEINLINE FNiagaraDataSetReaderStruct() {}
+	inline FNiagaraDataSetReaderStruct() {}
 
 	explicit FNiagaraDataSetReaderStruct(const FNiagaraDataBuffer* DataBuffer, int32 FloatComponentIndex, int32 Int32ComponentIndex)
 	{
@@ -553,20 +553,20 @@ struct FNiagaraDataSetReaderStruct
 		}
 	}
 
-	FORCEINLINE bool IsValid() const { return bIsValid; }
+	inline bool IsValid() const { return bIsValid; }
 
-	FORCEINLINE TType operator[](int32 Index) const
+	inline TType operator[](int32 Index) const
 	{
 		return Get(Index);
 	}
 
-	FORCEINLINE TType Get(int Index) const
+	inline TType Get(int Index) const
 	{
 		check(bIsValid && Index <= NumInstances);
 		return FNiagaraDataSetAccessorTypeInfo<TType>::ReadStruct(Index, FloatComponents, Int32Components, ComponentStride);
 	}
 
-	FORCEINLINE TType GetSafe(int Index, const TType& Default) const
+	inline TType GetSafe(int Index, const TType& Default) const
 	{
 		if (bIsValid && (Index < NumInstances))
 		{
@@ -575,7 +575,7 @@ struct FNiagaraDataSetReaderStruct
 		return Default;
 	}
 
-	FORCEINLINE int32 GetNumInstances()const { return NumInstances; }
+	inline int32 GetNumInstances()const { return NumInstances; }
 
 private:
 	bool bIsValid = false;
@@ -592,11 +592,11 @@ private:
 template<typename TType>
 struct FNiagaraDataSetAccessorStruct
 {
-	FORCEINLINE FNiagaraDataSetAccessorStruct() {}
-	FORCEINLINE explicit FNiagaraDataSetAccessorStruct(const FNiagaraDataSet& DataSet, const FName VariableName) { Init(NiagaraDataSetPrivate::GetCompiledData(DataSet), VariableName); }
-	FORCEINLINE explicit FNiagaraDataSetAccessorStruct(const FNiagaraDataSetCompiledData& DataSetCompiledData, const FName VariableName) { Init(DataSetCompiledData, VariableName); }
-	FORCEINLINE explicit FNiagaraDataSetAccessorStruct(const FNiagaraDataSetCompiledData* DataSetCompiledData, const FName VariableName) { Init(DataSetCompiledData, VariableName); }
-	FORCEINLINE void Init(const FNiagaraDataSet& DataSet, const FName VariableName) { Init(NiagaraDataSetPrivate::GetCompiledData(DataSet), VariableName); }
+	inline FNiagaraDataSetAccessorStruct() {}
+	inline explicit FNiagaraDataSetAccessorStruct(const FNiagaraDataSet& DataSet, const FName VariableName) { Init(NiagaraDataSetPrivate::GetCompiledData(DataSet), VariableName); }
+	inline explicit FNiagaraDataSetAccessorStruct(const FNiagaraDataSetCompiledData& DataSetCompiledData, const FName VariableName) { Init(DataSetCompiledData, VariableName); }
+	inline explicit FNiagaraDataSetAccessorStruct(const FNiagaraDataSetCompiledData* DataSetCompiledData, const FName VariableName) { Init(DataSetCompiledData, VariableName); }
+	inline void Init(const FNiagaraDataSet& DataSet, const FName VariableName) { Init(NiagaraDataSetPrivate::GetCompiledData(DataSet), VariableName); }
 
 	void Init(const FNiagaraDataSetCompiledData* DataSetCompiledData, const FName VariableName)
 	{
@@ -627,15 +627,15 @@ struct FNiagaraDataSetAccessorStruct
 		}
 	}
 
-	FORCEINLINE bool IsValid() const { return bIsValid; }
+	inline bool IsValid() const { return bIsValid; }
 
-	FORCEINLINE static FNiagaraDataSetReaderStruct<TType> CreateReader(const FNiagaraDataSet& DataSet, const FName VariableName) { return FNiagaraDataSetAccessorStruct<TType>(DataSet, VariableName).GetReader(DataSet); }
+	inline static FNiagaraDataSetReaderStruct<TType> CreateReader(const FNiagaraDataSet& DataSet, const FName VariableName) { return FNiagaraDataSetAccessorStruct<TType>(DataSet, VariableName).GetReader(DataSet); }
 	//FORCEINLINE static FNiagaraDataSetWriterStruct<TType> CreateWriter(const FNiagaraDataSet& DataSet, const FName VariableName) { return FNiagaraDataSetAccessorStruct<TType>(DataSet, VariableName).GetWriter(DataSet); }
 
-	FORCEINLINE FNiagaraDataSetReaderStruct<TType> GetReader(const FNiagaraDataSet& DataSet) const { return FNiagaraDataSetReaderStruct<TType>(NiagaraDataSetPrivate::GetCurrentData(DataSet), FloatComponentIndex, Int32ComponentIndex); }
+	inline FNiagaraDataSetReaderStruct<TType> GetReader(const FNiagaraDataSet& DataSet) const { return FNiagaraDataSetReaderStruct<TType>(NiagaraDataSetPrivate::GetCurrentData(DataSet), FloatComponentIndex, Int32ComponentIndex); }
 	//FNiagaraDataSetReaderStruct<TType> GetWriter(const FNiagaraDataSet& DataSet) { return FNiagaraDataSetReaderStruct<TType>(NiagaraDataSetPrivate::GetDestinationData(DataSet), FloatComponentIndex, Int32ComponentIndex); }
 
-	FORCEINLINE FNiagaraDataSetReaderStruct<TType> GetReader(const FNiagaraDataBuffer* DataBuffer) const { return FNiagaraDataSetReaderStruct<TType>(DataBuffer, FloatComponentIndex, Int32ComponentIndex); }
+	inline FNiagaraDataSetReaderStruct<TType> GetReader(const FNiagaraDataBuffer* DataBuffer) const { return FNiagaraDataSetReaderStruct<TType>(DataBuffer, FloatComponentIndex, Int32ComponentIndex); }
 private:
 	bool bIsValid = false;
 	int32 FloatComponentIndex = INDEX_NONE;
diff --git a/Engine/Plugins/FX/Niagara/Source/Niagara/Classes/NiagaraEffectType.h b/Engine/Plugins/FX/Niagara/Source/Niagara/Classes/NiagaraEffectType.h
--- a/Engine/Plugins/FX/Niagara/Source/Niagara/Classes/NiagaraEffectType.h
+++ b/Engine/Plugins/FX/Niagara/Source/Niagara/Classes/NiagaraEffectType.h
@@ -85,7 +85,7 @@ struct FNiagaraLinearRamp
 	UPROPERTY(EditAnywhere, Category = "Ramp")
 	float EndY = 1.0f;
 
-	FORCEINLINE float Evaluate(float X)const 
+	inline float Evaluate(float X)const 
 	{
 		if (X <= StartX)
 		{
@@ -454,8 +454,8 @@ class UNiagaraEffectType : public UObject
 	TArray<TObjectPtr<UNiagaraValidationRuleSet>> ValidationRuleSets;
 #endif
 
-	FORCEINLINE const FNiagaraSystemScalabilitySettingsArray& GetSystemScalabilitySettings()const { return SystemScalabilitySettings; }
-	FORCEINLINE const FNiagaraEmitterScalabilitySettingsArray& GetEmitterScalabilitySettings()const { return EmitterScalabilitySettings; }
+	inline const FNiagaraSystemScalabilitySettingsArray& GetSystemScalabilitySettings()const { return SystemScalabilitySettings; }
+	inline const FNiagaraEmitterScalabilitySettingsArray& GetEmitterScalabilitySettings()const { return EmitterScalabilitySettings; }
 
 	NIAGARA_API const FNiagaraSystemScalabilitySettings& GetActiveSystemScalabilitySettings()const;
 	NIAGARA_API const FNiagaraEmitterScalabilitySettings& GetActiveEmitterScalabilitySettings()const;
@@ -471,7 +471,7 @@ class UNiagaraEffectType : public UObject
 #if NIAGARA_PERF_BASELINES
 	UNiagaraBaselineController* GetPerfBaselineController() { return PerformanceBaselineController; }
 	FNiagaraPerfBaselineStats& GetPerfBaselineStats()const { return PerfBaselineStats; }
-	FORCEINLINE bool IsPerfBaselineValid()const { return PerfBaselineVersion == CurrentPerfBaselineVersion; }
+	inline bool IsPerfBaselineValid()const { return PerfBaselineVersion == CurrentPerfBaselineVersion; }
 	NIAGARA_API void UpdatePerfBaselineStats(FNiagaraPerfBaselineStats& NewBaselineStats);
 
 	NIAGARA_API void InvalidatePerfBaseline();
diff --git a/Engine/Plugins/FX/Niagara/Source/Niagara/Classes/NiagaraEmitter.h b/Engine/Plugins/FX/Niagara/Source/Niagara/Classes/NiagaraEmitter.h
--- a/Engine/Plugins/FX/Niagara/Source/Niagara/Classes/NiagaraEmitter.h
+++ b/Engine/Plugins/FX/Niagara/Source/Niagara/Classes/NiagaraEmitter.h
@@ -383,7 +383,7 @@ struct FVersionedNiagaraEmitterData
 	NIAGARA_API UNiagaraScript* GetScript(ENiagaraScriptUsage Usage, FGuid UsageId) const;
 	UNiagaraScript* GetGPUComputeScript() { return GPUComputeScript; }
 	const UNiagaraScript* GetGPUComputeScript() const { return GPUComputeScript; }
-	FORCEINLINE const TArray<FNiagaraEventScriptProperties>& GetEventHandlers() const { return EventHandlerScriptProps; }
+	inline const TArray<FNiagaraEventScriptProperties>& GetEventHandlers() const { return EventHandlerScriptProps; }
 	NIAGARA_API void CacheFromCompiledData(const FNiagaraDataSetCompiledData* CompiledData, const UNiagaraEmitter& Emitter);
 	NIAGARA_API void CacheFromShaderCompiled();
 	bool IsAllowedToExecute() const { return bIsAllowedToExecute; }
@@ -400,12 +400,12 @@ struct FVersionedNiagaraEmitterData
 	NIAGARA_API bool RequiresPersistentIDs() const;
 	const TArray<UNiagaraSimulationStageBase*>& GetSimulationStages() const { return SimulationStages; }
 	FNiagaraSimStageExecutionDataPtr GetSimStageExcecutionData() const { return SimStageExecutionData; }
-	FORCEINLINE const FNiagaraEmitterScalabilitySettings& GetScalabilitySettings()const { return CurrentScalabilitySettings; }
+	inline const FNiagaraEmitterScalabilitySettings& GetScalabilitySettings()const { return CurrentScalabilitySettings; }
 	NIAGARA_API const FNiagaraEmitterScalabilityOverride& GetCurrentOverrideSettings() const;
 	NIAGARA_API UNiagaraSimulationStageBase* GetSimulationStageById(FGuid ScriptUsageId) const;
 	NIAGARA_API bool BuildParameterStoreRendererBindings(FNiagaraParameterStore& ParameterStore) const;
 	NIAGARA_API void RebuildRendererBindings(const UNiagaraEmitter& Emitter);
-	FORCEINLINE static FBox GetDefaultFixedBounds() { return FBox(FVector(-100), FVector(100)); }
+	inline static FBox GetDefaultFixedBounds() { return FBox(FVector(-100), FVector(100)); }
 	
 	NIAGARA_API FVersionedNiagaraEmitter GetParent() const;
 	NIAGARA_API FVersionedNiagaraEmitter GetParentAtLastMerge() const;
diff --git a/Engine/Plugins/FX/Niagara/Source/Niagara/Classes/NiagaraEmitterInstance.h b/Engine/Plugins/FX/Niagara/Source/Niagara/Classes/NiagaraEmitterInstance.h
--- a/Engine/Plugins/FX/Niagara/Source/Niagara/Classes/NiagaraEmitterInstance.h
+++ b/Engine/Plugins/FX/Niagara/Source/Niagara/Classes/NiagaraEmitterInstance.h
@@ -167,10 +167,10 @@ public:
 	NIAGARA_API TArrayView<FNiagaraScriptExecutionContext> GetEventExecutionContexts();
 
 	UE_DEPRECATED(5.4, "Please update your code to handle execution path for different emitter types")
-	FORCEINLINE  FVersionedNiagaraEmitter GetCachedEmitter() const { return VersionedEmitter; }
+	inline  FVersionedNiagaraEmitter GetCachedEmitter() const { return VersionedEmitter; }
 
 	UE_DEPRECATED(5.4, "Please update your code to handle execution path for different emitter types")
-	FORCEINLINE FVersionedNiagaraEmitterData* GetCachedEmitterData()const { return VersionedEmitter.GetEmitterData(); }
+	inline FVersionedNiagaraEmitterData* GetCachedEmitterData()const { return VersionedEmitter.GetEmitterData(); }
 
 	UE_DEPRECATED(5.4, "Please update your code to handle execution path for different emitter types")
 	NIAGARA_API TArray<FNiagaraSpawnInfo>& GetSpawnInfo();
diff --git a/Engine/Plugins/FX/Niagara/Source/Niagara/Classes/NiagaraGPUInstanceCountManager.h b/Engine/Plugins/FX/Niagara/Source/Niagara/Classes/NiagaraGPUInstanceCountManager.h
--- a/Engine/Plugins/FX/Niagara/Source/Niagara/Classes/NiagaraGPUInstanceCountManager.h
+++ b/Engine/Plugins/FX/Niagara/Source/Niagara/Classes/NiagaraGPUInstanceCountManager.h
@@ -22,7 +22,7 @@ public:
 	int32 Value = 0;
 };
 
-FORCEINLINE uint32 GetTypeHash(const FNiagaraDrawIndirectArgGenTaskInfo& Info)
+inline uint32 GetTypeHash(const FNiagaraDrawIndirectArgGenTaskInfo& Info)
 {
 	return HashCombine(Info.InstanceCountBufferOffset, HashCombine(Info.NumIndicesPerInstance, HashCombine(Info.StartIndexLocation, Info.Flags)));
 }
@@ -52,7 +52,7 @@ public:
 
 		FIndirectArgSlot() {}
 		FIndirectArgSlot(FBufferRHIRef InBuffer, FShaderResourceViewRHIRef InSRV, uint32 InOffset) : Buffer(InBuffer), SRV(InSRV), Offset(InOffset) {}
-		FORCEINLINE bool IsValid() const { return Offset != INDEX_NONE; }
+		inline bool IsValid() const { return Offset != INDEX_NONE; }
 	};
 
 	FNiagaraGPUInstanceCountManager(ERHIFeatureLevel::Type FeatureLevel);
diff --git a/Engine/Plugins/FX/Niagara/Source/Niagara/Classes/NiagaraGPUSortInfo.h b/Engine/Plugins/FX/Niagara/Source/Niagara/Classes/NiagaraGPUSortInfo.h
--- a/Engine/Plugins/FX/Niagara/Source/Niagara/Classes/NiagaraGPUSortInfo.h
+++ b/Engine/Plugins/FX/Niagara/Source/Niagara/Classes/NiagaraGPUSortInfo.h
@@ -82,7 +82,7 @@ struct FNiagaraGPUSortInfo
 	EGPUSortFlags SortFlags = EGPUSortFlags::None;
 
 	// Set the SortFlags based on the emitter and material constraints.
-	FORCEINLINE void SetSortFlags(bool bHighPrecisionKeys, ENiagaraGpuComputeTickStage::Type ReadyTickStage)
+	inline void SetSortFlags(bool bHighPrecisionKeys, ENiagaraGpuComputeTickStage::Type ReadyTickStage)
 	{
 		SortFlags =
 			EGPUSortFlags::ValuesAsInt32 | 
diff --git a/Engine/Plugins/FX/Niagara/Source/Niagara/Classes/NiagaraGPUSystemTick.h b/Engine/Plugins/FX/Niagara/Source/Niagara/Classes/NiagaraGPUSystemTick.h
--- a/Engine/Plugins/FX/Niagara/Source/Niagara/Classes/NiagaraGPUSystemTick.h
+++ b/Engine/Plugins/FX/Niagara/Source/Niagara/Classes/NiagaraGPUSystemTick.h
@@ -74,30 +74,30 @@ public:
 	void Init(FNiagaraSystemInstance* InSystemInstance);
 	void Destroy();
 
-	FORCEINLINE TArrayView<FNiagaraComputeInstanceData> GetInstances() const
+	inline TArrayView<FNiagaraComputeInstanceData> GetInstances() const
 	{
 		return MakeArrayView(reinterpret_cast<FNiagaraComputeInstanceData*>(InstanceData_ParamData_Packed), InstanceCount);
 	};
 
-	FORCEINLINE void GetGlobalParameters(const FNiagaraComputeInstanceData& InstanceData, void* OutputParameters) const
+	inline void GetGlobalParameters(const FNiagaraComputeInstanceData& InstanceData, void* OutputParameters) const
 	{
 		const uint32 ParamSize = sizeof(FNiagaraGlobalParameters) * (InstanceData.Context->HasInterpolationParameters ? 2 : 1);
 		FMemory::Memcpy(OutputParameters, GlobalParamData, ParamSize);
 	}
 
-	FORCEINLINE void GetSystemParameters(const FNiagaraComputeInstanceData& InstanceData, void* OutputParameters) const
+	inline void GetSystemParameters(const FNiagaraComputeInstanceData& InstanceData, void* OutputParameters) const
 	{
 		const uint32 ParamSize = sizeof(FNiagaraSystemParameters) * (InstanceData.Context->HasInterpolationParameters ? 2 : 1);
 		FMemory::Memcpy(OutputParameters, SystemParamData, ParamSize);
 	}
 
-	FORCEINLINE void GetOwnerParameters(const FNiagaraComputeInstanceData& InstanceData, void* OutputParameters) const
+	inline void GetOwnerParameters(const FNiagaraComputeInstanceData& InstanceData, void* OutputParameters) const
 	{
 		const uint32 ParamSize = sizeof(FNiagaraOwnerParameters) * (InstanceData.Context->HasInterpolationParameters ? 2 : 1);
 		FMemory::Memcpy(OutputParameters, OwnerParamData, ParamSize);
 	}
 
-	FORCEINLINE void GetEmitterParameters(const FNiagaraComputeInstanceData& InstanceData, void* OutputParameters) const
+	inline void GetEmitterParameters(const FNiagaraComputeInstanceData& InstanceData, void* OutputParameters) const
 	{
 		const uint32 ParamSize = sizeof(FNiagaraEmitterParameters) * (InstanceData.Context->HasInterpolationParameters ? 2 : 1);
 		FMemory::Memcpy(OutputParameters, InstanceData.EmitterParamData, ParamSize);
@@ -105,7 +105,7 @@ public:
 
 	void BuildUniformBuffers();
 
-	FORCEINLINE FRHIUniformBuffer* GetExternalUniformBuffer(const FNiagaraComputeInstanceData& InstanceData, bool bPrevious) const
+	inline FRHIUniformBuffer* GetExternalUniformBuffer(const FNiagaraComputeInstanceData& InstanceData, bool bPrevious) const
 	{
 		const int32 InstanceIndex = UE_PTRDIFF_TO_INT32(&InstanceData - GetInstances().GetData());
 		const int32 BufferIndex = InstanceIndex + (bPrevious ? InstanceCount : 0);
diff --git a/Engine/Plugins/FX/Niagara/Source/Niagara/Classes/NiagaraParameterCollection.h b/Engine/Plugins/FX/Niagara/Source/Niagara/Classes/NiagaraParameterCollection.h
--- a/Engine/Plugins/FX/Niagara/Source/Niagara/Classes/NiagaraParameterCollection.h
+++ b/Engine/Plugins/FX/Niagara/Source/Niagara/Classes/NiagaraParameterCollection.h
@@ -169,8 +169,8 @@ public:
 	//TODO: Optional per platform overrides of the above.
 	//TMap<FString, UNiagaraParameterCollectionOverride> PerPlatformOverrides;
 
-	FORCEINLINE UNiagaraParameterCollectionInstance* GetDefaultInstance() { return DefaultInstance; }
-	FORCEINLINE const UNiagaraParameterCollectionInstance* GetDefaultInstance() const { return DefaultInstance; }
+	inline UNiagaraParameterCollectionInstance* GetDefaultInstance() { return DefaultInstance; }
+	inline const UNiagaraParameterCollectionInstance* GetDefaultInstance() const { return DefaultInstance; }
 	
 	/**
 	Takes the friendly name presented to the UI and converts to the real parameter name used under the hood.
diff --git a/Engine/Plugins/FX/Niagara/Source/Niagara/Classes/NiagaraPerfBaseline.h b/Engine/Plugins/FX/Niagara/Source/Niagara/Classes/NiagaraPerfBaseline.h
--- a/Engine/Plugins/FX/Niagara/Source/Niagara/Classes/NiagaraPerfBaseline.h
+++ b/Engine/Plugins/FX/Niagara/Source/Niagara/Classes/NiagaraPerfBaseline.h
@@ -53,7 +53,7 @@ struct FNiagaraPerfBaselineStats
 	*/
 	NIAGARA_API FNiagaraPerfBaselineStats(FAccumulatedParticlePerfStats& Stats, bool bSyncRT);
 
-	FORCEINLINE bool IsValid()const { return PerInstanceAvg_GT > 0.0f && PerInstanceAvg_RT > 0.0f && PerInstanceMax_GT > 0.0f && PerInstanceMax_RT > 0.0f; }
+	inline bool IsValid()const { return PerInstanceAvg_GT > 0.0f && PerInstanceAvg_RT > 0.0f && PerInstanceMax_GT > 0.0f && PerInstanceMax_RT > 0.0f; }
 	/**
 	Compares a system's performance stats with that of a baseline.
 	*/
diff --git a/Engine/Plugins/FX/Niagara/Source/Niagara/Classes/NiagaraPlatformSet.h b/Engine/Plugins/FX/Niagara/Source/Niagara/Classes/NiagaraPlatformSet.h
--- a/Engine/Plugins/FX/Niagara/Source/Niagara/Classes/NiagaraPlatformSet.h
+++ b/Engine/Plugins/FX/Niagara/Source/Niagara/Classes/NiagaraPlatformSet.h
@@ -37,13 +37,13 @@ struct FNiagaraDeviceProfileStateEntry
 	UPROPERTY(EditAnywhere, Category = Profile)
 	uint32 SetQualityLevelMask;
 
-	FORCEINLINE ENiagaraPlatformSelectionState GetState(int32 QualityLevel)const;
+	inline ENiagaraPlatformSelectionState GetState(int32 QualityLevel)const;
 
-	FORCEINLINE void SetState(int32 QualityLevel, ENiagaraPlatformSelectionState State);
+	inline void SetState(int32 QualityLevel, ENiagaraPlatformSelectionState State);
 
-	FORCEINLINE bool AllDefaulted()const { return SetQualityLevelMask == 0; }
+	inline bool AllDefaulted()const { return SetQualityLevelMask == 0; }
 
-	FORCEINLINE bool operator==(const FNiagaraDeviceProfileStateEntry& Other)const
+	inline bool operator==(const FNiagaraDeviceProfileStateEntry& Other)const
 	{
 		return ProfileName == Other.ProfileName && QualityLevelMask == Other.QualityLevelMask && SetQualityLevelMask == Other.SetQualityLevelMask;
 	}
@@ -130,21 +130,21 @@ struct FNiagaraCVarValue
 		float FloatValue;
 	}Data;
 	
-	template<typename T> FORCEINLINE T Get()const;	
-	template<typename T> FORCEINLINE void SetTyped(const FString& Str);
-	FORCEINLINE void Set(const FString& Str, IConsoleVariable* CVar);
+	template<typename T> inline T Get()const;	
+	template<typename T> inline void SetTyped(const FString& Str);
+	inline void Set(const FString& Str, IConsoleVariable* CVar);
 };
 
-template<> FORCEINLINE bool FNiagaraCVarValue::Get<bool>()const{ return Data.bValue; }
-template<> FORCEINLINE void FNiagaraCVarValue::SetTyped<bool>(const FString& Str){ LexFromString(Data.bValue, *Str); }
+template<> inline bool FNiagaraCVarValue::Get<bool>()const{ return Data.bValue; }
+template<> inline void FNiagaraCVarValue::SetTyped<bool>(const FString& Str){ LexFromString(Data.bValue, *Str); }
 
-template<> FORCEINLINE int32 FNiagaraCVarValue::Get<int32>()const { return Data.IntValue; }
-template<> FORCEINLINE void FNiagaraCVarValue::SetTyped<int32>(const FString& Str) { LexFromString(Data.IntValue, *Str); }
+template<> inline int32 FNiagaraCVarValue::Get<int32>()const { return Data.IntValue; }
+template<> inline void FNiagaraCVarValue::SetTyped<int32>(const FString& Str) { LexFromString(Data.IntValue, *Str); }
 
-template<> FORCEINLINE float FNiagaraCVarValue::Get<float>()const { return Data.FloatValue; }
-template<> FORCEINLINE void FNiagaraCVarValue::SetTyped<float>(const FString& Str) { LexFromString(Data.FloatValue, *Str); }
+template<> inline float FNiagaraCVarValue::Get<float>()const { return Data.FloatValue; }
+template<> inline void FNiagaraCVarValue::SetTyped<float>(const FString& Str) { LexFromString(Data.FloatValue, *Str); }
 
-FORCEINLINE void FNiagaraCVarValue::Set(const FString& Str, IConsoleVariable* CVar)
+inline void FNiagaraCVarValue::Set(const FString& Str, IConsoleVariable* CVar)
 {
 	if (CVar->IsVariableBool()) { SetTyped<bool>(Str); }
 	else if (CVar->IsVariableInt()) { SetTyped<int32>(Str); }
@@ -250,12 +250,12 @@ struct FNiagaraPlatformSetCVarCondition
 	template<typename T>
 	bool Evaluate_Internal(const UDeviceProfile* DeviceProfile, int32 QualityLevel, bool bCheckCurrentStateOnly, bool& bOutCanEverPass, bool& bOutCanEverFail)const;
 
-	FORCEINLINE bool CheckValue(bool CVarValue)const { return CVarValue == Value; }
-	FORCEINLINE bool CheckValue(int32 CVarValue)const { return (!bUseMinInt || CVarValue >= MinInt) && (!bUseMaxInt || CVarValue <= MaxInt); }
-	FORCEINLINE bool CheckValue(float CVarValue)const { return (!bUseMinFloat || CVarValue >= MinFloat) && (!bUseMaxFloat || CVarValue <= MaxFloat); }
+	inline bool CheckValue(bool CVarValue)const { return CVarValue == Value; }
+	inline bool CheckValue(int32 CVarValue)const { return (!bUseMinInt || CVarValue >= MinInt) && (!bUseMaxInt || CVarValue <= MaxInt); }
+	inline bool CheckValue(float CVarValue)const { return (!bUseMinFloat || CVarValue >= MinFloat) && (!bUseMaxFloat || CVarValue <= MaxFloat); }
 
 	template<typename T>
-	FORCEINLINE T GetCVarValue(IConsoleVariable* CVar)const;
+	inline T GetCVarValue(IConsoleVariable* CVar)const;
 
 private:
 	mutable IConsoleVariable* CachedCVar = nullptr;
@@ -309,8 +309,8 @@ struct FNiagaraPlatformSet
 {
 	GENERATED_BODY();
 
-	static FORCEINLINE int32 QualityLevelFromMask(int32 QLMask);
-	static FORCEINLINE int32 CreateQualityLevelMask(int32 QL);
+	static inline int32 QualityLevelFromMask(int32 QLMask);
+	static inline int32 CreateQualityLevelMask(int32 QL);
 	static NIAGARA_API FText GetQualityLevelText(int32 QualityLevel);
 	static NIAGARA_API FText GetQualityLevelMaskText(int32 QualityLevelMask);
 
@@ -481,7 +481,7 @@ private:
 };
 
 
-FORCEINLINE int32 FNiagaraPlatformSet::QualityLevelFromMask(int32 QLMask)
+inline int32 FNiagaraPlatformSet::QualityLevelFromMask(int32 QLMask)
 {
 	//Iterate through bits in mask to find the first set quality level.
 	int32 QualityLevel = 0;
@@ -498,19 +498,19 @@ FORCEINLINE int32 FNiagaraPlatformSet::QualityLevelFromMask(int32 QLMask)
 	return INDEX_NONE;
 }
 
-FORCEINLINE int32 FNiagaraPlatformSet::CreateQualityLevelMask(int32 QL)
+inline int32 FNiagaraPlatformSet::CreateQualityLevelMask(int32 QL)
 {
 	return QL == INDEX_NONE ? INDEX_NONE : (1 << QL);
 }
 
 
-FORCEINLINE ENiagaraPlatformSelectionState FNiagaraDeviceProfileStateEntry::GetState(int32 QualityLevel)const
+inline ENiagaraPlatformSelectionState FNiagaraDeviceProfileStateEntry::GetState(int32 QualityLevel)const
 {
 	int32 QLMask = FNiagaraPlatformSet::CreateQualityLevelMask(QualityLevel);
 	return (QLMask & SetQualityLevelMask) == 0 ? ENiagaraPlatformSelectionState::Default : ((QLMask & QualityLevelMask) != 0 ? ENiagaraPlatformSelectionState::Enabled : ENiagaraPlatformSelectionState::Disabled);
 }
 
-FORCEINLINE void FNiagaraDeviceProfileStateEntry::SetState(int32 QualityLevel, ENiagaraPlatformSelectionState State)
+inline void FNiagaraDeviceProfileStateEntry::SetState(int32 QualityLevel, ENiagaraPlatformSelectionState State)
 {
 	int32 QLMask = FNiagaraPlatformSet::CreateQualityLevelMask(QualityLevel);
 	if (State == ENiagaraPlatformSelectionState::Default)
diff --git a/Engine/Plugins/FX/Niagara/Source/Niagara/Classes/NiagaraScript.h b/Engine/Plugins/FX/Niagara/Source/Niagara/Classes/NiagaraScript.h
--- a/Engine/Plugins/FX/Niagara/Source/Niagara/Classes/NiagaraScript.h
+++ b/Engine/Plugins/FX/Niagara/Source/Niagara/Classes/NiagaraScript.h
@@ -1282,9 +1282,9 @@ public:
 	void RaiseOnGPUCompilationComplete();
 
 
-	FORCEINLINE FNiagaraVMExecutableData& GetVMExecutableData() { return CachedScriptVM; }
-	FORCEINLINE const FNiagaraVMExecutableData& GetVMExecutableData() const { return CachedScriptVM; }
-	FORCEINLINE const FNiagaraVMExecutableDataId& GetVMExecutableDataCompilationId() const { return CachedScriptVMId; }
+	inline FNiagaraVMExecutableData& GetVMExecutableData() { return CachedScriptVM; }
+	inline const FNiagaraVMExecutableData& GetVMExecutableData() const { return CachedScriptVM; }
+	inline const FNiagaraVMExecutableDataId& GetVMExecutableDataCompilationId() const { return CachedScriptVMId; }
 
 	TConstArrayView<FNiagaraDataInterfaceGPUParamInfo> GetDataInterfaceGPUParamInfos() const { return CachedScriptVM.ShaderScriptParametersMetadata.DataInterfaceParamInfo; }
 
diff --git a/Engine/Plugins/FX/Niagara/Source/Niagara/Classes/NiagaraScriptExecutionContext.h b/Engine/Plugins/FX/Niagara/Source/Niagara/Classes/NiagaraScriptExecutionContext.h
--- a/Engine/Plugins/FX/Niagara/Source/Niagara/Classes/NiagaraScriptExecutionContext.h
+++ b/Engine/Plugins/FX/Niagara/Source/Niagara/Classes/NiagaraScriptExecutionContext.h
@@ -66,7 +66,7 @@ struct FNiagaraDataSetExecutionInfo
 	}
 
 
-	FORCEINLINE void Init(FNiagaraDataSet* InDataSet, FNiagaraDataBuffer* InInput, int32 InStartInstance, bool bInUpdateInstanceCount)
+	inline void Init(FNiagaraDataSet* InDataSet, FNiagaraDataBuffer* InInput, int32 InStartInstance, bool bInUpdateInstanceCount)
 	{
 		DataSet = InDataSet;
 		Input = InInput;
@@ -82,7 +82,7 @@ struct FNiagaraDataSetExecutionInfo
 		Reset();
 	}
 
-	FORCEINLINE void Reset()
+	inline void Reset()
 	{
 		DataSet = nullptr;
 		Input = nullptr;
diff --git a/Engine/Plugins/FX/Niagara/Source/Niagara/Classes/NiagaraSystem.h b/Engine/Plugins/FX/Niagara/Source/Niagara/Classes/NiagaraSystem.h
--- a/Engine/Plugins/FX/Niagara/Source/Niagara/Classes/NiagaraSystem.h
+++ b/Engine/Plugins/FX/Niagara/Source/Niagara/Classes/NiagaraSystem.h
@@ -395,19 +395,19 @@ private:
 public:
 	NIAGARA_API bool IsReadyToRun() const;
 
-	FORCEINLINE bool NeedsWarmup()const { return WarmupTickCount > 0 && WarmupTickDelta > SMALL_NUMBER; }
-	FORCEINLINE float GetWarmupTime()const { return WarmupTime; }
-	FORCEINLINE int32 GetWarmupTickCount()const { return WarmupTickCount; }
-	FORCEINLINE float GetWarmupTickDelta()const { return WarmupTickDelta; }
+	inline bool NeedsWarmup()const { return WarmupTickCount > 0 && WarmupTickDelta > SMALL_NUMBER; }
+	inline float GetWarmupTime()const { return WarmupTime; }
+	inline int32 GetWarmupTickCount()const { return WarmupTickCount; }
+	inline float GetWarmupTickDelta()const { return WarmupTickDelta; }
 
-	FORCEINLINE bool HasFixedTickDelta() const { return bFixedTickDelta; }
-	FORCEINLINE float GetFixedTickDeltaTime()const { return FixedTickDeltaTime; }
+	inline bool HasFixedTickDelta() const { return bFixedTickDelta; }
+	inline float GetFixedTickDeltaTime()const { return FixedTickDeltaTime; }
 	NIAGARA_API virtual void GetAssetRegistryTags(FAssetRegistryTagsContext Context) const override;
 	UE_DEPRECATED(5.4, "Implement the version that takes FAssetRegistryTagsContext instead.")
 	NIAGARA_API virtual void GetAssetRegistryTags(TArray<FAssetRegistryTag>& OutTags) const override;
 
-	FORCEINLINE bool NeedsDeterminism() const { return bDeterminism; }
-	FORCEINLINE int32 GetRandomSeed() const { return RandomSeed; }
+	inline bool NeedsDeterminism() const { return bDeterminism; }
+	inline int32 GetRandomSeed() const { return RandomSeed; }
 
 	NIAGARA_API void SetWarmupTime(float InWarmupTime);
 	NIAGARA_API void SetWarmupTickDelta(float InWarmupTickDelta);
@@ -536,7 +536,7 @@ public:
 	bool SupportsLargeWorldCoordinates() const { return bSupportLargeWorldCoordinates && bLwcEnabledSettingCached; }
 	ENiagaraLwcTileUpdateMode GetLargeWorldCoordinateTileUpdateMode() const;
 
-	FORCEINLINE bool IsInitialOwnerVelocityFromActor() const { return bInitialOwnerVelocityFromActor; }
+	inline bool IsInitialOwnerVelocityFromActor() const { return bInitialOwnerVelocityFromActor; }
 
 #if WITH_EDITORONLY_DATA
 	NIAGARA_API void ReportAnalyticsData(bool bIsCooking);
@@ -545,20 +545,20 @@ public:
 	NIAGARA_API bool UsesScript(const UNiagaraScript* Script)const; 
 	NIAGARA_API void ForceGraphToRecompileOnNextCheck();
 	NIAGARA_API void SetCompileForEdit(bool bNewCompileForEdit);
-	FORCEINLINE bool GetCompileForEdit() const { return bCompileForEdit; }
+	inline bool GetCompileForEdit() const { return bCompileForEdit; }
 
 	static NIAGARA_API void RequestCompileForEmitter(const FVersionedNiagaraEmitter& InEmitter);
 	static NIAGARA_API void RecomputeExecutionOrderForEmitter(const FVersionedNiagaraEmitter& InEmitter);
 	static NIAGARA_API void RecomputeExecutionOrderForDataInterface(class UNiagaraDataInterface* DataInterface);
 
-	FORCEINLINE bool ShouldUseRapidIterationParameters() const { return bCompileForEdit ? !bBakeOutRapidIteration : !bBakeOutRapidIterationOnCook; }
-	FORCEINLINE bool ShouldTrimAttributes() const { return bCompileForEdit ? bTrimAttributes : bTrimAttributesOnCook; }
-	FORCEINLINE bool ShouldIgnoreParticleReadsForAttributeTrim() const { return bIgnoreParticleReadsForAttributeTrim; }
-	FORCEINLINE bool ShouldDisableDebugSwitches() const { return bCompileForEdit ? bDisableDebugSwitches : bDisableDebugSwitchesOnCook; }
-	FORCEINLINE bool ShouldCompressAttributes() const { return bCompressAttributes; }
+	inline bool ShouldUseRapidIterationParameters() const { return bCompileForEdit ? !bBakeOutRapidIteration : !bBakeOutRapidIterationOnCook; }
+	inline bool ShouldTrimAttributes() const { return bCompileForEdit ? bTrimAttributes : bTrimAttributesOnCook; }
+	inline bool ShouldIgnoreParticleReadsForAttributeTrim() const { return bIgnoreParticleReadsForAttributeTrim; }
+	inline bool ShouldDisableDebugSwitches() const { return bCompileForEdit ? bDisableDebugSwitches : bDisableDebugSwitchesOnCook; }
+	inline bool ShouldCompressAttributes() const { return bCompressAttributes; }
 
-	FORCEINLINE void SetBakeOutRapidIterationOnCook(bool bBakeOut) { bBakeOutRapidIteration = bBakeOut; bBakeOutRapidIterationOnCook = bBakeOut; }
-	FORCEINLINE void SetTrimAttributesOnCook(bool bTrim) { bTrimAttributes = bTrim; bTrimAttributesOnCook = bTrim; }
+	inline void SetBakeOutRapidIterationOnCook(bool bBakeOut) { bBakeOutRapidIteration = bBakeOut; bBakeOutRapidIterationOnCook = bBakeOut; }
+	inline void SetTrimAttributesOnCook(bool bTrim) { bTrimAttributes = bTrim; bTrimAttributesOnCook = bTrim; }
 
 	/** Returns information on the availability of performance stat scopes */
 	NIAGARA_API ENiagaraPerformanceStateScopeMode GetPerformanceStateScopeMode() const;
@@ -727,16 +727,16 @@ public:
 	/** Cache data & accessors from the compiled data, allows us to avoid per instance. */
 	NIAGARA_API void CacheFromCompiledData();
 
-	FORCEINLINE TConstArrayView<FNiagaraEmitterExecutionIndex> GetEmitterExecutionOrder() const { return MakeArrayView(EmitterExecutionOrder); }
-	FORCEINLINE TConstArrayView<FNiagaraRendererExecutionIndex> GetRendererPostTickOrder() const { return MakeArrayView(RendererPostTickOrder); }
-	FORCEINLINE TConstArrayView<FNiagaraRendererExecutionIndex> GetRendererCompletionOrder() const { return MakeArrayView(RendererCompletionOrder); }
+	inline TConstArrayView<FNiagaraEmitterExecutionIndex> GetEmitterExecutionOrder() const { return MakeArrayView(EmitterExecutionOrder); }
+	inline TConstArrayView<FNiagaraRendererExecutionIndex> GetRendererPostTickOrder() const { return MakeArrayView(RendererPostTickOrder); }
+	inline TConstArrayView<FNiagaraRendererExecutionIndex> GetRendererCompletionOrder() const { return MakeArrayView(RendererCompletionOrder); }
 
-	FORCEINLINE TConstArrayView<int32> GetRendererDrawOrder() const { return MakeArrayView(RendererDrawOrder); }
+	inline TConstArrayView<int32> GetRendererDrawOrder() const { return MakeArrayView(RendererDrawOrder); }
 
 	/** When an index inside the EmitterExecutionOrder array has this bit set, it means the corresponding emitter cannot execute in parallel with the previous emitters due to a data dependency. */
 	static constexpr int32 kStartNewOverlapGroupBit = (1 << 31);
 
-	FORCEINLINE UNiagaraParameterCollectionInstance* GetParameterCollectionOverride(UNiagaraParameterCollection* Collection)
+	inline UNiagaraParameterCollectionInstance* GetParameterCollectionOverride(UNiagaraParameterCollection* Collection)
 	{
 		auto* Found = ParameterCollectionOverrides.FindByPredicate(
 			[&](const UNiagaraParameterCollectionInstance* CheckInst)
@@ -765,14 +765,14 @@ protected:
 	uint8 bOverrideScalabilitySettings : 1;
 public:
 	
-	FORCEINLINE bool HasDIsWithPostSimulateTick() const { return bHasDIsWithPostSimulateTick; }
-	FORCEINLINE bool AllDIsPostSimulateCanOverlapFrames() const { return bAllDIsPostSimulateCanOverlapFrames; }
-	FORCEINLINE bool AsyncWorkCanOverlapTickGroups() const { return bAllDIsPostStageCanOverlapTickGroups; }
-	FORCEINLINE bool HasAnyGPUEmitters()const{ return bHasAnyGPUEmitters; }
-	FORCEINLINE bool NeedsGPUContextInitForDataInterfaces() const { return bNeedsGPUContextInitForDataInterfaces; }
+	inline bool HasDIsWithPostSimulateTick() const { return bHasDIsWithPostSimulateTick; }
+	inline bool AllDIsPostSimulateCanOverlapFrames() const { return bAllDIsPostSimulateCanOverlapFrames; }
+	inline bool AsyncWorkCanOverlapTickGroups() const { return bAllDIsPostStageCanOverlapTickGroups; }
+	inline bool HasAnyGPUEmitters()const{ return bHasAnyGPUEmitters; }
+	inline bool NeedsGPUContextInitForDataInterfaces() const { return bNeedsGPUContextInitForDataInterfaces; }
 
 	NIAGARA_API FBox GetFixedBounds() const;
-	FORCEINLINE void SetFixedBounds(const FBox& Box) { FixedBounds = Box;  }
+	inline void SetFixedBounds(const FBox& Box) { FixedBounds = Box;  }
 
 	FBox GetInitialStreamingBounds() const { return bUseInitialStreamingBounds ? InitialStreamingBounds : FBox(); }
 
@@ -780,17 +780,17 @@ public:
 	NIAGARA_API void SetEffectType(UNiagaraEffectType* EffectType);
 
 	FNiagaraSystemScalabilityOverrides& GetSystemScalabilityOverrides() { return SystemScalabilityOverrides; }
-	FORCEINLINE bool GetOverrideScalabilitySettings()const { return bOverrideScalabilitySettings; }
-	FORCEINLINE void SetOverrideScalabilitySettings(bool bOverride) { bOverrideScalabilitySettings = bOverride; }
+	inline bool GetOverrideScalabilitySettings()const { return bOverrideScalabilitySettings; }
+	inline void SetOverrideScalabilitySettings(bool bOverride) { bOverrideScalabilitySettings = bOverride; }
 
 	NIAGARA_API void GatherStaticVariables(TArray<FNiagaraVariable>& OutVars, TArray<FNiagaraVariable>& OutEmitterVars) const;
 #endif
 	NIAGARA_API UNiagaraEffectType* GetEffectType()const;
-	FORCEINLINE const FNiagaraSystemScalabilitySettings& GetScalabilitySettings()const { return CurrentScalabilitySettings; }
+	inline const FNiagaraSystemScalabilitySettings& GetScalabilitySettings()const { return CurrentScalabilitySettings; }
 	NIAGARA_API const FNiagaraSystemScalabilityOverride& GetCurrentOverrideSettings() const;
-	FORCEINLINE bool NeedsSortedSignificanceCull()const{ return bNeedsSortedSignificanceCull; }
+	inline bool NeedsSortedSignificanceCull()const{ return bNeedsSortedSignificanceCull; }
 
-	FORCEINLINE void GetMaxInstanceCounts(int32& OutSystemInstanceMax, int32& OutFXTypeInstanceMax, bool bBudgetAdjusted)const;
+	inline void GetMaxInstanceCounts(int32& OutSystemInstanceMax, int32& OutFXTypeInstanceMax, bool bBudgetAdjusted)const;
 	
 	FNiagaraPlatformSet& GetScalabilityPlatformSet() { return Platforms; }
 	const FNiagaraPlatformSet& GetScalabilityPlatformSet() const { return Platforms; }
@@ -819,9 +819,9 @@ public:
 	const FGuid& GetAssetGuid() const {return AssetGuid;};
 #endif
 
-	FORCEINLINE void RegisterActiveInstance();
-	FORCEINLINE void UnregisterActiveInstance();
-	FORCEINLINE int32& GetActiveInstancesCount() { return ActiveInstances; }
+	inline void RegisterActiveInstance();
+	inline void UnregisterActiveInstance();
+	inline int32& GetActiveInstancesCount() { return ActiveInstances; }
 
 #if WITH_EDITORONLY_DATA
 	NIAGARA_API UNiagaraBakerSettings* GetBakerSettings();
@@ -1131,18 +1131,18 @@ protected:
 	FRenderCommandFence WaitRenderCommandsFence;
 };
 
-FORCEINLINE void UNiagaraSystem::RegisterActiveInstance()
+inline void UNiagaraSystem::RegisterActiveInstance()
 {
 	++ActiveInstances;
 }
 
-FORCEINLINE void UNiagaraSystem::UnregisterActiveInstance()
+inline void UNiagaraSystem::UnregisterActiveInstance()
 {
 	--ActiveInstances;
 }
 
 
-FORCEINLINE void UNiagaraSystem::GetMaxInstanceCounts(int32& OutSystemInstanceMax, int32& OutFXTypeInstanceMax, bool bBudgetAdjusted)const 
+inline void UNiagaraSystem::GetMaxInstanceCounts(int32& OutSystemInstanceMax, int32& OutFXTypeInstanceMax, bool bBudgetAdjusted)const 
 {
 	OutSystemInstanceMax = CurrentScalabilitySettings.MaxSystemInstances;
 	OutFXTypeInstanceMax = CurrentScalabilitySettings.MaxInstances;
diff --git a/Engine/Plugins/FX/Niagara/Source/Niagara/Public/NiagaraCommon.h b/Engine/Plugins/FX/Niagara/Source/Niagara/Public/NiagaraCommon.h
--- a/Engine/Plugins/FX/Niagara/Source/Niagara/Public/NiagaraCommon.h
+++ b/Engine/Plugins/FX/Niagara/Source/Niagara/Public/NiagaraCommon.h
@@ -284,25 +284,25 @@ struct FNiagaraDataSetID
 	UPROPERTY()
 	ENiagaraDataSetType Type;
 
-	FORCEINLINE bool operator==(const FNiagaraDataSetID& Other)const
+	inline bool operator==(const FNiagaraDataSetID& Other)const
 	{
 		return Name == Other.Name && Type == Other.Type;
 	}
 
-	FORCEINLINE bool operator!=(const FNiagaraDataSetID& Other)const
+	inline bool operator!=(const FNiagaraDataSetID& Other)const
 	{
 		return !(*this == Other);
 	}
 };
 
 
-FORCEINLINE FArchive& operator<<(FArchive& Ar, FNiagaraDataSetID& VarInfo)
+inline FArchive& operator<<(FArchive& Ar, FNiagaraDataSetID& VarInfo)
 {
 	Ar << VarInfo.Name << VarInfo.Type;
 	return Ar;
 }
 
-FORCEINLINE uint32 GetTypeHash(const FNiagaraDataSetID& Var)
+inline uint32 GetTypeHash(const FNiagaraDataSetID& Var)
 {
 	return HashCombine(GetTypeHash(Var.Name), (uint32)Var.Type);
 }
@@ -1005,8 +1005,8 @@ struct FVMExternalFunctionBindingInfo
 	UPROPERTY()
 	TArray<FNiagaraVariableBase> VariadicOutputs;
 
-	FORCEINLINE int32 GetNumInputs() const { return InputParamLocations.Num(); }
-	FORCEINLINE int32 GetNumOutputs() const { return NumOutputs; }
+	inline int32 GetNumInputs() const { return InputParamLocations.Num(); }
+	inline int32 GetNumOutputs() const { return NumOutputs; }
 
 	const FVMFunctionSpecifier* FindSpecifier(const FName& Key) const
 	{
@@ -1571,13 +1571,13 @@ namespace FNiagaraUtilities
 	bool AreBufferSRVsAlwaysCreated(EShaderPlatform ShaderPlatform);
 
 	// Helper function to determine if we should sync data from CPU to GPU
-	FORCEINLINE bool ShouldSyncCpuToGpu(ENiagaraGpuSyncMode SyncMode)
+	inline bool ShouldSyncCpuToGpu(ENiagaraGpuSyncMode SyncMode)
 	{
 		return SyncMode == ENiagaraGpuSyncMode::SyncBoth || SyncMode == ENiagaraGpuSyncMode::SyncCpuToGpu;
 	}
 
 	// Helper function to determine if we should sync data from GPU to CPU
-	FORCEINLINE bool ShouldSyncGpuToCpu(ENiagaraGpuSyncMode SyncMode)
+	inline bool ShouldSyncGpuToCpu(ENiagaraGpuSyncMode SyncMode)
 	{
 		return SyncMode == ENiagaraGpuSyncMode::SyncBoth || SyncMode == ENiagaraGpuSyncMode::SyncGpuToCpu;
 	}
@@ -1638,7 +1638,7 @@ struct FNiagaraUserParameterBinding
 	UPROPERTY(EditAnywhere, Category = "User Parameter")
 	FNiagaraVariable Parameter;
 
-	FORCEINLINE bool operator==(const FNiagaraUserParameterBinding& Other)const
+	inline bool operator==(const FNiagaraUserParameterBinding& Other)const
 	{
 		return Other.Parameter == Parameter;
 	}
@@ -1781,12 +1781,12 @@ public:
 	UPROPERTY()
 	bool bDependentVariableFromCustomIterationNamespace;
 
-	FORCEINLINE bool operator==(const FNiagaraCompileDependency& Other)const
+	inline bool operator==(const FNiagaraCompileDependency& Other)const
 	{
 		return DependentVariable == Other.DependentVariable && NodeGuid == Other.NodeGuid && PinGuid == Other.PinGuid && StackGuids == Other.StackGuids;
 	}
 
-	FORCEINLINE bool operator!=(const FNiagaraCompileDependency& Other)const
+	inline bool operator!=(const FNiagaraCompileDependency& Other)const
 	{
 		return !(*this == Other);
 	}
diff --git a/Engine/Plugins/FX/Niagara/Source/Niagara/Public/NiagaraComponent.h b/Engine/Plugins/FX/Niagara/Source/Niagara/Public/NiagaraComponent.h
--- a/Engine/Plugins/FX/Niagara/Source/Niagara/Public/NiagaraComponent.h
+++ b/Engine/Plugins/FX/Niagara/Source/Niagara/Public/NiagaraComponent.h
@@ -636,10 +636,10 @@ public:
 	NIAGARA_API void SetPreviewLODDistance(bool bEnablePreviewLODDistance, float PreviewLODDistance, float PreviewMaxDistance);
 
 	UFUNCTION(BlueprintCallable, Category = Preview, meta = (Keywords = "preview LOD Distance scalability"))
-	FORCEINLINE bool GetPreviewLODDistanceEnabled()const;
+	inline bool GetPreviewLODDistanceEnabled()const;
 
 	UFUNCTION(BlueprintCallable, Category = Preview, meta = (Keywords = "preview LOD Distance scalability"))
-	FORCEINLINE float GetPreviewLODDistance()const;
+	inline float GetPreviewLODDistance()const;
 
 	/**
 	Initializes this component for capturing a performance baseline.
@@ -808,18 +808,18 @@ public:
 	UFUNCTION(BlueprintGetter)
 	NIAGARA_API bool GetForceLocalPlayerEffect()const;
 
-	FORCEINLINE bool IsRegisteredWithScalabilityManager()const { return ScalabilityManagerHandle != INDEX_NONE; }
-	FORCEINLINE int32 GetScalabilityManagerHandle()const { return ScalabilityManagerHandle; }
+	inline bool IsRegisteredWithScalabilityManager()const { return ScalabilityManagerHandle != INDEX_NONE; }
+	inline int32 GetScalabilityManagerHandle()const { return ScalabilityManagerHandle; }
 
-	FORCEINLINE void BeginUpdateContextReset(){ bDuringUpdateContextReset = true; }
-	FORCEINLINE void EndUpdateContextReset(){ bDuringUpdateContextReset = false; }
+	inline void BeginUpdateContextReset(){ bDuringUpdateContextReset = true; }
+	inline void EndUpdateContextReset(){ bDuringUpdateContextReset = false; }
 
 #if WITH_NIAGARA_DEBUGGER	
 	//Cache our scalability state in the component so we have access to it easily and also after it has been removed from the scalability manager.
 	FNiagaraScalabilityState DebugCachedScalabilityState;
 #endif
 
-	FORCEINLINE bool IsUsingCullProxy()const { return CullProxy != nullptr; }
+	inline bool IsUsingCullProxy()const { return CullProxy != nullptr; }
 
 	NIAGARA_API bool ResolveOwnerAllowsScalability(bool bRegister=true);
 
@@ -910,5 +910,5 @@ public:
 	NIAGARA_API FParticlePerfStatsContext GetPerfStatsContext();
 };
 
-FORCEINLINE bool UNiagaraComponent::GetPreviewLODDistanceEnabled()const { return bEnablePreviewLODDistance; }
-FORCEINLINE float UNiagaraComponent::GetPreviewLODDistance()const { return bEnablePreviewLODDistance ? PreviewLODDistance : 0.0f; }
+inline bool UNiagaraComponent::GetPreviewLODDistanceEnabled()const { return bEnablePreviewLODDistance; }
+inline float UNiagaraComponent::GetPreviewLODDistance()const { return bEnablePreviewLODDistance ? PreviewLODDistance : 0.0f; }
diff --git a/Engine/Plugins/FX/Niagara/Source/Niagara/Public/NiagaraDataChannelAccessor.h b/Engine/Plugins/FX/Niagara/Source/Niagara/Public/NiagaraDataChannelAccessor.h
--- a/Engine/Plugins/FX/Niagara/Source/Niagara/Public/NiagaraDataChannelAccessor.h
+++ b/Engine/Plugins/FX/Niagara/Source/Niagara/Public/NiagaraDataChannelAccessor.h
@@ -295,11 +295,11 @@ public:
 	TWeakObjectPtr<const UNiagaraDataChannel> WeakNDC;
 	void ValidateLayout()const { checkf(WeakNDC->IsValid() && WeakNDC->GetLayoutInfo() == DebugCachedLayout, TEXT("NDC Variable being accessed with stale layout info.")); }
 #else
-	FORCEINLINE void ValidateLayout()const {}
+	inline void ValidateLayout()const {}
 #endif
 
 	template<typename T>
-	FORCEINLINE bool WriteInternal(const FNiagaraDataChannelGameDataPtr& Data, int32 Index, const T& InValue)
+	inline bool WriteInternal(const FNiagaraDataChannelGameDataPtr& Data, int32 Index, const T& InValue)
 	{
 		if (VarOffset != INDEX_NONE)
 		{
@@ -309,7 +309,7 @@ public:
 	}	
 	
 	template<typename T>
-	FORCEINLINE bool ReadInternal(const FNiagaraDataChannelGameDataPtr& Data, int32 Index, T& OutValue, bool bPreviousFrame)
+	inline bool ReadInternal(const FNiagaraDataChannelGameDataPtr& Data, int32 Index, T& OutValue, bool bPreviousFrame)
 	{
 		if (VarOffset != INDEX_NONE)
 		{
@@ -399,7 +399,7 @@ public:
 		: FNDCVarAccessorBase(Owner, FNiagaraVariableBase(FNiagaraTypeHelper::GetTypeDef<T>(), VarName), bInIsRequired)
 	{}
 
-	FORCEINLINE bool Write(const FNiagaraDataChannelGameDataPtr& Data, int32 Index, const typename TNiagaraTypeHelper<T>::TLWCType& InValue, int32 StartIndex, int32 Count)
+	inline bool Write(const FNiagaraDataChannelGameDataPtr& Data, int32 Index, const typename TNiagaraTypeHelper<T>::TLWCType& InValue, int32 StartIndex, int32 Count)
 	{
 		if(Index >= 0 && Index < Count)
 		{
@@ -418,7 +418,7 @@ public:
 		: FNDCVarAccessorBase(Owner, FNiagaraVariableBase(FNiagaraTypeHelper::GetTypeDef<T>(), VarName), bInIsRequired)
 	{}
 
-	FORCEINLINE bool Read(const FNiagaraDataChannelGameDataPtr& Data, int32 Index, typename TNiagaraTypeHelper<T>::TLWCType& OutValue, bool bPreviousFrame)
+	inline bool Read(const FNiagaraDataChannelGameDataPtr& Data, int32 Index, typename TNiagaraTypeHelper<T>::TLWCType& OutValue, bool bPreviousFrame)
 	{
 		return ReadInternal<typename TNiagaraTypeHelper<T>::TLWCType>(Data, Index, OutValue, bPreviousFrame);
 	}
diff --git a/Engine/Plugins/FX/Niagara/Source/Niagara/Public/NiagaraGpuComputeDispatchInterface.h b/Engine/Plugins/FX/Niagara/Source/Niagara/Public/NiagaraGpuComputeDispatchInterface.h
--- a/Engine/Plugins/FX/Niagara/Source/Niagara/Public/NiagaraGpuComputeDispatchInterface.h
+++ b/Engine/Plugins/FX/Niagara/Source/Niagara/Public/NiagaraGpuComputeDispatchInterface.h
@@ -103,8 +103,8 @@ public:
 	virtual const FGlobalDistanceFieldParameterData* GetGlobalDistanceFieldData() const = 0;
 
 	/** Get access to the instance count manager. */
-	FORCEINLINE FNiagaraGPUInstanceCountManager& GetGPUInstanceCounterManager() { check(IsInParallelRenderingThread()); return GPUInstanceCounterManager; }
-	FORCEINLINE const FNiagaraGPUInstanceCountManager& GetGPUInstanceCounterManager() const { check(IsInParallelRenderingThread()); return GPUInstanceCounterManager; }
+	inline FNiagaraGPUInstanceCountManager& GetGPUInstanceCounterManager() { check(IsInParallelRenderingThread()); return GPUInstanceCounterManager; }
+	inline const FNiagaraGPUInstanceCountManager& GetGPUInstanceCounterManager() const { check(IsInParallelRenderingThread()); return GPUInstanceCounterManager; }
 
 #if NIAGARA_COMPUTEDEBUG_ENABLED
 	/** Public interface to Niagara compute debugging. */
@@ -165,10 +165,10 @@ public:
 
 	virtual FNiagaraAsyncGpuTraceHelper& GetAsyncGpuTraceHelper() const = 0;
 
-	FORCEINLINE bool IsOutsideSceneRenderer() const { return bIsOutsideSceneRenderer; }
+	inline bool IsOutsideSceneRenderer() const { return bIsOutsideSceneRenderer; }
 
-	FORCEINLINE bool IsFirstViewFamily() const { return bIsFirstViewFamily; }
-	FORCEINLINE bool IsLastViewFamily() const { return bIsLastViewFamily; }
+	inline bool IsFirstViewFamily() const { return bIsFirstViewFamily; }
+	inline bool IsLastViewFamily() const { return bIsLastViewFamily; }
 
 #if WITH_MGPU
 	/**
@@ -184,13 +184,13 @@ public:
 	UE_DEPRECATED(5.1, "ImmediateMode is deprecated for Niagara please migrate to using a FRDGBuilder")
 	virtual void MultiGPUResourceModified(FRHICommandList& RHICmdList, FRHITexture* Texture, bool bRequiredForSimulation, bool bRequiredForRendering) const = 0;
 #else
-	FORCEINLINE void MultiGPUResourceModified(FRDGBuilder& GraphBuilder, FRHIBuffer* Buffer, bool bRequiredForSimulation, bool bRequiredForRendering) const {}
-	FORCEINLINE void MultiGPUResourceModified(FRDGBuilder& GraphBuilder, FRHITexture* Texture, bool bRequiredForSimulation, bool bRequiredForRendering) const {}
+	inline void MultiGPUResourceModified(FRDGBuilder& GraphBuilder, FRHIBuffer* Buffer, bool bRequiredForSimulation, bool bRequiredForRendering) const {}
+	inline void MultiGPUResourceModified(FRDGBuilder& GraphBuilder, FRHITexture* Texture, bool bRequiredForSimulation, bool bRequiredForRendering) const {}
 
 	UE_DEPRECATED(5.1, "ImmediateMode is deprecated for Niagara please migrate to using a FRDGBuilder")
-	FORCEINLINE void MultiGPUResourceModified(FRHICommandList& RHICmdList, FRHIBuffer* Buffer, bool bRequiredForSimulation, bool bRequiredForRendering) const {}
+	inline void MultiGPUResourceModified(FRHICommandList& RHICmdList, FRHIBuffer* Buffer, bool bRequiredForSimulation, bool bRequiredForRendering) const {}
 	UE_DEPRECATED(5.1, "ImmediateMode is deprecated for Niagara please migrate to using a FRDGBuilder")
-	FORCEINLINE void MultiGPUResourceModified(FRHICommandList& RHICmdList, FRHITexture* Texture, bool bRequiredForSimulation, bool bRequiredForRendering) const {}
+	inline void MultiGPUResourceModified(FRHICommandList& RHICmdList, FRHITexture* Texture, bool bRequiredForSimulation, bool bRequiredForRendering) const {}
 #endif
 
 	/**
