Replace some usages of FORCEINLINE with inline in Niagara modules.

diff --git a/Engine/Plugins/FX/Niagara/Source/Niagara/Public/NiagaraModule.h b/Engine/Plugins/FX/Niagara/Source/Niagara/Public/NiagaraModule.h
--- a/Engine/Plugins/FX/Niagara/Source/Niagara/Public/NiagaraModule.h
+++ b/Engine/Plugins/FX/Niagara/Source/Niagara/Public/NiagaraModule.h
@@ -146,133 +146,133 @@ public:
 	NIAGARA_API int32 RenderStatPerfBaselines(UWorld* World, FViewport* Viewport, FCanvas* Canvas, int32 X, int32 Y, const FVector* ViewLocation = nullptr, const FRotator* ViewRotation = nullptr);
 #endif
 
-	FORCEINLINE static bool UseGlobalFXBudget() { return bUseGlobalFXBudget; }
+	inline static bool UseGlobalFXBudget() { return bUseGlobalFXBudget; }
 	static NIAGARA_API void OnUseGlobalFXBudgetChanged(IConsoleVariable* Variable);
 
-	FORCEINLINE static bool DataChannelsEnabled() { return bDataChannelsEnabled; }
+	inline static bool DataChannelsEnabled() { return bDataChannelsEnabled; }
 	static NIAGARA_API void OnDataChannelsEnabledChanged(IConsoleVariable* Variable);
 
-	FORCEINLINE static float GetGlobalSpawnCountScale() { return EngineGlobalSpawnCountScale; }
-	FORCEINLINE static float GetGlobalSystemCountScale() { return EngineGlobalSystemCountScale; }
+	inline static float GetGlobalSpawnCountScale() { return EngineGlobalSpawnCountScale; }
+	inline static float GetGlobalSystemCountScale() { return EngineGlobalSystemCountScale; }
 
 	static NIAGARA_API float EngineGlobalSpawnCountScale;
 	static NIAGARA_API float EngineGlobalSystemCountScale;
 
-	FORCEINLINE static const FNiagaraVariable&  GetVar_Engine_WorldDeltaTime() { return Engine_WorldDeltaTime; }
-	FORCEINLINE static const FNiagaraVariable&  GetVar_Engine_DeltaTime() { return Engine_DeltaTime; }
-	FORCEINLINE static const FNiagaraVariable&  GetVar_Engine_InvDeltaTime() { return Engine_InvDeltaTime; }
-	FORCEINLINE static const FNiagaraVariable&  GetVar_Engine_Time() { return Engine_Time; }
-	FORCEINLINE static const FNiagaraVariable&  GetVar_Engine_RealTime() { return Engine_RealTime; }
-	FORCEINLINE static const FNiagaraVariable&  GetVar_Engine_QualityLevel() { return Engine_QualityLevel; }
-
-	FORCEINLINE static const FNiagaraVariable&  GetVar_Engine_Owner_Position() { return Engine_Owner_Position; }
-	FORCEINLINE static const FNiagaraVariable&  GetVar_Engine_Owner_Velocity() { return Engine_Owner_Velocity; }
-	FORCEINLINE static const FNiagaraVariable&  GetVar_Engine_Owner_XAxis() { return Engine_Owner_XAxis; }
-	FORCEINLINE static const FNiagaraVariable&  GetVar_Engine_Owner_YAxis() { return Engine_Owner_YAxis; }
-	FORCEINLINE static const FNiagaraVariable&  GetVar_Engine_Owner_ZAxis() { return Engine_Owner_ZAxis; }
-	FORCEINLINE static const FNiagaraVariable&  GetVar_Engine_Owner_Scale() { return Engine_Owner_Scale; }
-	FORCEINLINE static const FNiagaraVariable&  GetVar_Engine_Owner_Rotation() { return Engine_Owner_Rotation; }
-	FORCEINLINE static const FNiagaraVariable&  GetVar_Engine_Owner_LWC_Tile() { return Engine_Owner_LWC_Tile; }
-	FORCEINLINE static const FNiagaraVariable&  GetVar_Engine_ExecIndex() { return Engine_ExecIndex; }
-
-	FORCEINLINE static const FNiagaraVariable&  GetVar_Engine_Owner_SystemLocalToWorld() { return Engine_Owner_SystemLocalToWorld; }
-	FORCEINLINE static const FNiagaraVariable&  GetVar_Engine_Owner_SystemWorldToLocal() { return Engine_Owner_SystemWorldToLocal; }
-	FORCEINLINE static const FNiagaraVariable&  GetVar_Engine_Owner_SystemLocalToWorldTransposed() { return Engine_Owner_SystemLocalToWorldTransposed; }
-	FORCEINLINE static const FNiagaraVariable&  GetVar_Engine_Owner_SystemWorldToLocalTransposed() { return Engine_Owner_SystemWorldToLocalTransposed; }
-	FORCEINLINE static const FNiagaraVariable&  GetVar_Engine_Owner_SystemLocalToWorldNoScale() { return Engine_Owner_SystemLocalToWorldNoScale; }
-	FORCEINLINE static const FNiagaraVariable&  GetVar_Engine_Owner_SystemWorldToLocalNoScale() { return Engine_Owner_SystemWorldToLocalNoScale; }
-
-	FORCEINLINE static const FNiagaraVariable&  GetVar_Engine_Owner_TimeSinceRendered() { return Engine_Owner_TimeSinceRendered; }
-	FORCEINLINE static const FNiagaraVariable&  GetVar_Engine_Owner_LODDistance() { return Engine_Owner_LODDistance; }
-	FORCEINLINE static const FNiagaraVariable&  GetVar_Engine_Owner_LODDistanceFraction() { return Engine_Owner_LODDistanceFraction; }
-
-	FORCEINLINE static const FNiagaraVariable&  GetVar_Engine_Owner_ExecutionState() { return Engine_Owner_ExecutionState; }
-
-	FORCEINLINE static const FNiagaraVariable&  GetVar_Engine_ExecutionCount() { return Engine_ExecutionCount; }
-	FORCEINLINE static const FNiagaraVariable&  GetVar_Engine_Emitter_NumParticles() { return Engine_Emitter_NumParticles; }
-	FORCEINLINE static const FNiagaraVariable&  GetVar_Engine_Emitter_SimulationPosition() { return Engine_Emitter_SimulationPosition; }
-	FORCEINLINE static const FNiagaraVariable&  GetVar_Engine_Emitter_TotalSpawnedParticles() { return Engine_Emitter_TotalSpawnedParticles; }
-	FORCEINLINE static const FNiagaraVariable&  GetVar_Engine_Emitter_SpawnCountScale() { return Engine_Emitter_SpawnCountScale; }
-	FORCEINLINE static const FNiagaraVariable&  GetVar_Engine_Emitter_InstanceSeed() { return Engine_Emitter_InstanceSeed; }
-	FORCEINLINE static const FNiagaraVariable&  GetVar_Engine_Emitter_ID() { return Engine_Emitter_ID; }
-	FORCEINLINE static const FNiagaraVariable&  GetVar_Engine_System_TickCount() { return Engine_System_TickCount; }
-	FORCEINLINE static const FNiagaraVariable&  GetVar_Engine_System_NumEmittersAlive() { return Engine_System_NumEmittersAlive; }
-	FORCEINLINE static const FNiagaraVariable&  GetVar_Engine_System_SignificanceIndex() { return Engine_System_SignificanceIndex; }
-	FORCEINLINE static const FNiagaraVariable&  GetVar_Engine_System_RandomSeed() { return Engine_System_RandomSeed; }
-
-	FORCEINLINE static const FNiagaraVariable& GetVar_Engine_System_CurrentTimeStep() { return Engine_System_CurrentTimeStep; }
-	FORCEINLINE static const FNiagaraVariable& GetVar_Engine_System_NumTimeSteps() { return Engine_System_NumTimeSteps; }
-	FORCEINLINE static const FNiagaraVariable& GetVar_Engine_System_TimeStepFraction() { return Engine_System_TimeStepFraction; }
-	FORCEINLINE static const FNiagaraVariable& GetVar_Engine_System_NumParticles() { return Engine_System_NumParticles; }
-
-	FORCEINLINE static const FNiagaraVariable&  GetVar_Engine_System_NumEmitters() { return Engine_System_NumEmitters; }
-	FORCEINLINE static const FNiagaraVariable&  GetVar_Engine_NumSystemInstances() { return Engine_NumSystemInstances; }
-
-	FORCEINLINE static const FNiagaraVariable&  GetVar_Engine_GlobalSpawnCountScale() { return Engine_GlobalSpawnCountScale; }
-	FORCEINLINE static const FNiagaraVariable&  GetVar_Engine_GlobalSystemScale() { return Engine_GlobalSystemScale; }
-
-	FORCEINLINE static const FNiagaraVariable&  GetVar_Engine_System_Age() { return Engine_System_Age; }
-	FORCEINLINE static const FNiagaraVariable&  GetVar_Emitter_Age() { return Emitter_Age; }
-	FORCEINLINE static const FNiagaraVariable&  GetVar_Emitter_LocalSpace() { return Emitter_LocalSpace; }
-	FORCEINLINE static const FNiagaraVariable&  GetVar_Emitter_Determinism() { return Emitter_Determinism; }
-	FORCEINLINE static const FNiagaraVariable&  GetVar_Emitter_InterpolatedSpawn() { return Emitter_InterpolatedSpawn; }
-	FORCEINLINE static const FNiagaraVariable&  GetVar_Emitter_OverrideGlobalSpawnCountScale() { return Emitter_OverrideGlobalSpawnCountScale; }
-	FORCEINLINE static const FNiagaraVariable&  GetVar_Emitter_RandomSeed() { return Emitter_RandomSeed; }
-	FORCEINLINE static const FNiagaraVariable&  GetVar_Emitter_SpawnRate() { return Emitter_SpawnRate; }
-	FORCEINLINE static const FNiagaraVariable&  GetVar_Emitter_SpawnInterval() { return Emitter_SpawnInterval; }
-	FORCEINLINE static const FNiagaraVariable&  GetVar_Emitter_SimulationTarget() { return Emitter_SimulationTarget; }
-	FORCEINLINE static const FNiagaraVariable&  GetVar_ScriptUsage() { return ScriptUsage; }
-	FORCEINLINE static const FNiagaraVariable&  GetVar_ScriptContext() { return ScriptContext; }
-	FORCEINLINE static const FNiagaraVariable&  GetVar_FunctionDebugState() { return FunctionDebugState; }
-	FORCEINLINE static const FNiagaraVariable&  GetVar_Emitter_InterpSpawnStartDt() { return Emitter_InterpSpawnStartDt; }
-	FORCEINLINE static const FNiagaraVariable&  GetVar_Emitter_SpawnGroup() { return Emitter_SpawnGroup; }
-
-	FORCEINLINE static const FNiagaraVariable&  GetVar_Particles_UniqueID() { return Particles_UniqueID; }
-	FORCEINLINE static const FNiagaraVariable&  GetVar_Particles_ID() { return Particles_ID; }
-	FORCEINLINE static const FNiagaraVariable&  GetVar_Particles_Position() { return Particles_Position; }
-	FORCEINLINE static const FNiagaraVariable&  GetVar_Particles_Velocity() { return Particles_Velocity; }
-	FORCEINLINE static const FNiagaraVariable&  GetVar_Particles_Color() { return Particles_Color; }
-	FORCEINLINE static const FNiagaraVariable&  GetVar_Particles_SpriteRotation() { return Particles_SpriteRotation; }
-	FORCEINLINE static const FNiagaraVariable&  GetVar_Particles_Age() { return Particles_Age; }
-	FORCEINLINE static const FNiagaraVariable&  GetVar_Particles_NormalizedAge() { return Particles_NormalizedAge; }
-	FORCEINLINE static const FNiagaraVariable&  GetVar_Particles_SpriteSize() { return Particles_SpriteSize; }
-	FORCEINLINE static const FNiagaraVariable&  GetVar_Particles_SpriteFacing() { return Particles_SpriteFacing; }
-	FORCEINLINE static const FNiagaraVariable&  GetVar_Particles_SpriteAlignment() { return Particles_SpriteAlignment; }
-	FORCEINLINE static const FNiagaraVariable&  GetVar_Particles_SubImageIndex() { return Particles_SubImageIndex; }
-	FORCEINLINE static const FNiagaraVariable&  GetVar_Particles_DynamicMaterialParameter() { return Particles_DynamicMaterialParameter; }
-	FORCEINLINE static const FNiagaraVariable&  GetVar_Particles_DynamicMaterialParameter1() { return Particles_DynamicMaterialParameter1; }
-	FORCEINLINE static const FNiagaraVariable&  GetVar_Particles_DynamicMaterialParameter2() { return Particles_DynamicMaterialParameter2; }
-	FORCEINLINE static const FNiagaraVariable&  GetVar_Particles_DynamicMaterialParameter3() { return Particles_DynamicMaterialParameter3; }
-	FORCEINLINE static const FNiagaraVariable&  GetVar_Particles_Scale() { return Particles_Scale; }
-	FORCEINLINE static const FNiagaraVariable&  GetVar_Particles_Lifetime() { return Particles_Lifetime; }
-	FORCEINLINE static const FNiagaraVariable&  GetVar_Particles_MeshOrientation() { return Particles_MeshOrientation; }
-	FORCEINLINE static const FNiagaraVariable&  GetVar_Particles_UVScale() { return Particles_UVScale; }
-	FORCEINLINE static const FNiagaraVariable&  GetVar_Particles_PivotOffset() { return Particles_PivotOffset; }
-	FORCEINLINE static const FNiagaraVariable&  GetVar_Particles_CameraOffset() { return Particles_CameraOffset; }
-	FORCEINLINE static const FNiagaraVariable&  GetVar_Particles_MaterialRandom() { return Particles_MaterialRandom; }
-	FORCEINLINE static const FNiagaraVariable&  GetVar_Particles_LightRadius() { return Particles_LightRadius; }
-	FORCEINLINE static const FNiagaraVariable&  GetVar_Particles_LightExponent() { return Particles_LightExponent; }
-	FORCEINLINE static const FNiagaraVariable&  GetVar_Particles_LightEnabled() { return Particles_LightEnabled; }
-	FORCEINLINE static const FNiagaraVariable&  GetVar_Particles_LightVolumetricScattering() { return Particles_LightVolumetricScattering; }
-	FORCEINLINE static const FNiagaraVariable&  GetVar_Particles_LightSpecularScale() { return Particles_LightSpecularScale; }
-	FORCEINLINE static const FNiagaraVariable&  GetVar_Particles_LightDiffuseScale() { return Particles_LightDiffuseScale; }
-	FORCEINLINE static const FNiagaraVariable&  GetVar_Particles_RibbonID() { return Particles_RibbonID; }
-	FORCEINLINE static const FNiagaraVariable&  GetVar_Particles_RibbonWidth() { return Particles_RibbonWidth; }
-	FORCEINLINE static const FNiagaraVariable&  GetVar_Particles_RibbonTwist() { return Particles_RibbonTwist; }
-	FORCEINLINE static const FNiagaraVariable&  GetVar_Particles_RibbonFacing() { return Particles_RibbonFacing; }
-	FORCEINLINE static const FNiagaraVariable&  GetVar_Particles_RibbonLinkOrder() { return Particles_RibbonLinkOrder; }
-	FORCEINLINE static const FNiagaraVariable&  GetVar_Particles_RibbonUVDistance() { return Particles_RibbonUVDistance; }
-	FORCEINLINE static const FNiagaraVariable&  GetVar_Particles_RibbonU0Override() { return Particles_RibbonU0Override; }
-	FORCEINLINE static const FNiagaraVariable&  GetVar_Particles_RibbonV0RangeOverride() { return Particles_RibbonV0RangeOverride; }
-	FORCEINLINE static const FNiagaraVariable&  GetVar_Particles_RibbonU1Override() { return Particles_RibbonU1Override; }
-	FORCEINLINE static const FNiagaraVariable&  GetVar_Particles_RibbonV1RangeOverride() { return Particles_RibbonV1RangeOverride; }
-	FORCEINLINE static const FNiagaraVariable&  GetVar_Particles_VisibilityTag() { return Particles_VisibilityTag; }
-	FORCEINLINE static const FNiagaraVariable&  GetVar_Particles_MeshIndex() { return Particles_MeshIndex; }
-	FORCEINLINE static const FNiagaraVariable&  GetVar_Particles_ComponentsEnabled() { return Particles_ComponentsEnabled; }
+	inline static const FNiagaraVariable&  GetVar_Engine_WorldDeltaTime() { return Engine_WorldDeltaTime; }
+	inline static const FNiagaraVariable&  GetVar_Engine_DeltaTime() { return Engine_DeltaTime; }
+	inline static const FNiagaraVariable&  GetVar_Engine_InvDeltaTime() { return Engine_InvDeltaTime; }
+	inline static const FNiagaraVariable&  GetVar_Engine_Time() { return Engine_Time; }
+	inline static const FNiagaraVariable&  GetVar_Engine_RealTime() { return Engine_RealTime; }
+	inline static const FNiagaraVariable&  GetVar_Engine_QualityLevel() { return Engine_QualityLevel; }
+
+	inline static const FNiagaraVariable&  GetVar_Engine_Owner_Position() { return Engine_Owner_Position; }
+	inline static const FNiagaraVariable&  GetVar_Engine_Owner_Velocity() { return Engine_Owner_Velocity; }
+	inline static const FNiagaraVariable&  GetVar_Engine_Owner_XAxis() { return Engine_Owner_XAxis; }
+	inline static const FNiagaraVariable&  GetVar_Engine_Owner_YAxis() { return Engine_Owner_YAxis; }
+	inline static const FNiagaraVariable&  GetVar_Engine_Owner_ZAxis() { return Engine_Owner_ZAxis; }
+	inline static const FNiagaraVariable&  GetVar_Engine_Owner_Scale() { return Engine_Owner_Scale; }
+	inline static const FNiagaraVariable&  GetVar_Engine_Owner_Rotation() { return Engine_Owner_Rotation; }
+	inline static const FNiagaraVariable&  GetVar_Engine_Owner_LWC_Tile() { return Engine_Owner_LWC_Tile; }
+	inline static const FNiagaraVariable&  GetVar_Engine_ExecIndex() { return Engine_ExecIndex; }
+
+	inline static const FNiagaraVariable&  GetVar_Engine_Owner_SystemLocalToWorld() { return Engine_Owner_SystemLocalToWorld; }
+	inline static const FNiagaraVariable&  GetVar_Engine_Owner_SystemWorldToLocal() { return Engine_Owner_SystemWorldToLocal; }
+	inline static const FNiagaraVariable&  GetVar_Engine_Owner_SystemLocalToWorldTransposed() { return Engine_Owner_SystemLocalToWorldTransposed; }
+	inline static const FNiagaraVariable&  GetVar_Engine_Owner_SystemWorldToLocalTransposed() { return Engine_Owner_SystemWorldToLocalTransposed; }
+	inline static const FNiagaraVariable&  GetVar_Engine_Owner_SystemLocalToWorldNoScale() { return Engine_Owner_SystemLocalToWorldNoScale; }
+	inline static const FNiagaraVariable&  GetVar_Engine_Owner_SystemWorldToLocalNoScale() { return Engine_Owner_SystemWorldToLocalNoScale; }
+
+	inline static const FNiagaraVariable&  GetVar_Engine_Owner_TimeSinceRendered() { return Engine_Owner_TimeSinceRendered; }
+	inline static const FNiagaraVariable&  GetVar_Engine_Owner_LODDistance() { return Engine_Owner_LODDistance; }
+	inline static const FNiagaraVariable&  GetVar_Engine_Owner_LODDistanceFraction() { return Engine_Owner_LODDistanceFraction; }
+
+	inline static const FNiagaraVariable&  GetVar_Engine_Owner_ExecutionState() { return Engine_Owner_ExecutionState; }
+
+	inline static const FNiagaraVariable&  GetVar_Engine_ExecutionCount() { return Engine_ExecutionCount; }
+	inline static const FNiagaraVariable&  GetVar_Engine_Emitter_NumParticles() { return Engine_Emitter_NumParticles; }
+	inline static const FNiagaraVariable&  GetVar_Engine_Emitter_SimulationPosition() { return Engine_Emitter_SimulationPosition; }
+	inline static const FNiagaraVariable&  GetVar_Engine_Emitter_TotalSpawnedParticles() { return Engine_Emitter_TotalSpawnedParticles; }
+	inline static const FNiagaraVariable&  GetVar_Engine_Emitter_SpawnCountScale() { return Engine_Emitter_SpawnCountScale; }
+	inline static const FNiagaraVariable&  GetVar_Engine_Emitter_InstanceSeed() { return Engine_Emitter_InstanceSeed; }
+	inline static const FNiagaraVariable&  GetVar_Engine_Emitter_ID() { return Engine_Emitter_ID; }
+	inline static const FNiagaraVariable&  GetVar_Engine_System_TickCount() { return Engine_System_TickCount; }
+	inline static const FNiagaraVariable&  GetVar_Engine_System_NumEmittersAlive() { return Engine_System_NumEmittersAlive; }
+	inline static const FNiagaraVariable&  GetVar_Engine_System_SignificanceIndex() { return Engine_System_SignificanceIndex; }
+	inline static const FNiagaraVariable&  GetVar_Engine_System_RandomSeed() { return Engine_System_RandomSeed; }
+
+	inline static const FNiagaraVariable& GetVar_Engine_System_CurrentTimeStep() { return Engine_System_CurrentTimeStep; }
+	inline static const FNiagaraVariable& GetVar_Engine_System_NumTimeSteps() { return Engine_System_NumTimeSteps; }
+	inline static const FNiagaraVariable& GetVar_Engine_System_TimeStepFraction() { return Engine_System_TimeStepFraction; }
+	inline static const FNiagaraVariable& GetVar_Engine_System_NumParticles() { return Engine_System_NumParticles; }
+
+	inline static const FNiagaraVariable&  GetVar_Engine_System_NumEmitters() { return Engine_System_NumEmitters; }
+	inline static const FNiagaraVariable&  GetVar_Engine_NumSystemInstances() { return Engine_NumSystemInstances; }
+
+	inline static const FNiagaraVariable&  GetVar_Engine_GlobalSpawnCountScale() { return Engine_GlobalSpawnCountScale; }
+	inline static const FNiagaraVariable&  GetVar_Engine_GlobalSystemScale() { return Engine_GlobalSystemScale; }
+
+	inline static const FNiagaraVariable&  GetVar_Engine_System_Age() { return Engine_System_Age; }
+	inline static const FNiagaraVariable&  GetVar_Emitter_Age() { return Emitter_Age; }
+	inline static const FNiagaraVariable&  GetVar_Emitter_LocalSpace() { return Emitter_LocalSpace; }
+	inline static const FNiagaraVariable&  GetVar_Emitter_Determinism() { return Emitter_Determinism; }
+	inline static const FNiagaraVariable&  GetVar_Emitter_InterpolatedSpawn() { return Emitter_InterpolatedSpawn; }
+	inline static const FNiagaraVariable&  GetVar_Emitter_OverrideGlobalSpawnCountScale() { return Emitter_OverrideGlobalSpawnCountScale; }
+	inline static const FNiagaraVariable&  GetVar_Emitter_RandomSeed() { return Emitter_RandomSeed; }
+	inline static const FNiagaraVariable&  GetVar_Emitter_SpawnRate() { return Emitter_SpawnRate; }
+	inline static const FNiagaraVariable&  GetVar_Emitter_SpawnInterval() { return Emitter_SpawnInterval; }
+	inline static const FNiagaraVariable&  GetVar_Emitter_SimulationTarget() { return Emitter_SimulationTarget; }
+	inline static const FNiagaraVariable&  GetVar_ScriptUsage() { return ScriptUsage; }
+	inline static const FNiagaraVariable&  GetVar_ScriptContext() { return ScriptContext; }
+	inline static const FNiagaraVariable&  GetVar_FunctionDebugState() { return FunctionDebugState; }
+	inline static const FNiagaraVariable&  GetVar_Emitter_InterpSpawnStartDt() { return Emitter_InterpSpawnStartDt; }
+	inline static const FNiagaraVariable&  GetVar_Emitter_SpawnGroup() { return Emitter_SpawnGroup; }
+
+	inline static const FNiagaraVariable&  GetVar_Particles_UniqueID() { return Particles_UniqueID; }
+	inline static const FNiagaraVariable&  GetVar_Particles_ID() { return Particles_ID; }
+	inline static const FNiagaraVariable&  GetVar_Particles_Position() { return Particles_Position; }
+	inline static const FNiagaraVariable&  GetVar_Particles_Velocity() { return Particles_Velocity; }
+	inline static const FNiagaraVariable&  GetVar_Particles_Color() { return Particles_Color; }
+	inline static const FNiagaraVariable&  GetVar_Particles_SpriteRotation() { return Particles_SpriteRotation; }
+	inline static const FNiagaraVariable&  GetVar_Particles_Age() { return Particles_Age; }
+	inline static const FNiagaraVariable&  GetVar_Particles_NormalizedAge() { return Particles_NormalizedAge; }
+	inline static const FNiagaraVariable&  GetVar_Particles_SpriteSize() { return Particles_SpriteSize; }
+	inline static const FNiagaraVariable&  GetVar_Particles_SpriteFacing() { return Particles_SpriteFacing; }
+	inline static const FNiagaraVariable&  GetVar_Particles_SpriteAlignment() { return Particles_SpriteAlignment; }
+	inline static const FNiagaraVariable&  GetVar_Particles_SubImageIndex() { return Particles_SubImageIndex; }
+	inline static const FNiagaraVariable&  GetVar_Particles_DynamicMaterialParameter() { return Particles_DynamicMaterialParameter; }
+	inline static const FNiagaraVariable&  GetVar_Particles_DynamicMaterialParameter1() { return Particles_DynamicMaterialParameter1; }
+	inline static const FNiagaraVariable&  GetVar_Particles_DynamicMaterialParameter2() { return Particles_DynamicMaterialParameter2; }
+	inline static const FNiagaraVariable&  GetVar_Particles_DynamicMaterialParameter3() { return Particles_DynamicMaterialParameter3; }
+	inline static const FNiagaraVariable&  GetVar_Particles_Scale() { return Particles_Scale; }
+	inline static const FNiagaraVariable&  GetVar_Particles_Lifetime() { return Particles_Lifetime; }
+	inline static const FNiagaraVariable&  GetVar_Particles_MeshOrientation() { return Particles_MeshOrientation; }
+	inline static const FNiagaraVariable&  GetVar_Particles_UVScale() { return Particles_UVScale; }
+	inline static const FNiagaraVariable&  GetVar_Particles_PivotOffset() { return Particles_PivotOffset; }
+	inline static const FNiagaraVariable&  GetVar_Particles_CameraOffset() { return Particles_CameraOffset; }
+	inline static const FNiagaraVariable&  GetVar_Particles_MaterialRandom() { return Particles_MaterialRandom; }
+	inline static const FNiagaraVariable&  GetVar_Particles_LightRadius() { return Particles_LightRadius; }
+	inline static const FNiagaraVariable&  GetVar_Particles_LightExponent() { return Particles_LightExponent; }
+	inline static const FNiagaraVariable&  GetVar_Particles_LightEnabled() { return Particles_LightEnabled; }
+	inline static const FNiagaraVariable&  GetVar_Particles_LightVolumetricScattering() { return Particles_LightVolumetricScattering; }
+	inline static const FNiagaraVariable&  GetVar_Particles_LightSpecularScale() { return Particles_LightSpecularScale; }
+	inline static const FNiagaraVariable&  GetVar_Particles_LightDiffuseScale() { return Particles_LightDiffuseScale; }
+	inline static const FNiagaraVariable&  GetVar_Particles_RibbonID() { return Particles_RibbonID; }
+	inline static const FNiagaraVariable&  GetVar_Particles_RibbonWidth() { return Particles_RibbonWidth; }
+	inline static const FNiagaraVariable&  GetVar_Particles_RibbonTwist() { return Particles_RibbonTwist; }
+	inline static const FNiagaraVariable&  GetVar_Particles_RibbonFacing() { return Particles_RibbonFacing; }
+	inline static const FNiagaraVariable&  GetVar_Particles_RibbonLinkOrder() { return Particles_RibbonLinkOrder; }
+	inline static const FNiagaraVariable&  GetVar_Particles_RibbonUVDistance() { return Particles_RibbonUVDistance; }
+	inline static const FNiagaraVariable&  GetVar_Particles_RibbonU0Override() { return Particles_RibbonU0Override; }
+	inline static const FNiagaraVariable&  GetVar_Particles_RibbonV0RangeOverride() { return Particles_RibbonV0RangeOverride; }
+	inline static const FNiagaraVariable&  GetVar_Particles_RibbonU1Override() { return Particles_RibbonU1Override; }
+	inline static const FNiagaraVariable&  GetVar_Particles_RibbonV1RangeOverride() { return Particles_RibbonV1RangeOverride; }
+	inline static const FNiagaraVariable&  GetVar_Particles_VisibilityTag() { return Particles_VisibilityTag; }
+	inline static const FNiagaraVariable&  GetVar_Particles_MeshIndex() { return Particles_MeshIndex; }
+	inline static const FNiagaraVariable&  GetVar_Particles_ComponentsEnabled() { return Particles_ComponentsEnabled; }
 	
-	FORCEINLINE static const FNiagaraVariable&  GetVar_DataInstance_Alive() { return DataInstance_Alive; }
-	FORCEINLINE static const FNiagaraVariable&  GetVar_BeginDefaults() { return Translator_BeginDefaults; }
-	FORCEINLINE static const FNiagaraVariable&  GetVar_CallID() { return Translator_CallID; }
+	inline static const FNiagaraVariable&  GetVar_DataInstance_Alive() { return DataInstance_Alive; }
+	inline static const FNiagaraVariable&  GetVar_BeginDefaults() { return Translator_BeginDefaults; }
+	inline static const FNiagaraVariable&  GetVar_CallID() { return Translator_CallID; }
 	
 	FOnProcessQueue OnProcessQueue;
 	
diff --git a/Engine/Plugins/FX/Niagara/Source/Niagara/Public/NiagaraParameterStore.h b/Engine/Plugins/FX/Niagara/Source/Niagara/Public/NiagaraParameterStore.h
--- a/Engine/Plugins/FX/Niagara/Source/Niagara/Public/NiagaraParameterStore.h
+++ b/Engine/Plugins/FX/Niagara/Source/Niagara/Public/NiagaraParameterStore.h
@@ -56,7 +56,7 @@ struct FNiagaraParameterStoreBinding
 			check(InDestOffset < (int32)0xFFFF);
 			check(InSize < (int32)0xFFFF);
 		}
-		FORCEINLINE bool operator==(const FParameterBinding& Other)const
+		inline bool operator==(const FParameterBinding& Other)const
 		{
 			return SrcOffset == Other.SrcOffset && DestOffset == Other.DestOffset && Size == Other.Size;
 		}
@@ -76,7 +76,7 @@ struct FNiagaraParameterStoreBinding
 			check(InSrcOffset < (int32)0xFFFF);
 			check(InDestOffset < (int32)0xFFFF);
 		}
-		FORCEINLINE bool operator==(const FInterfaceBinding& Other)const
+		inline bool operator==(const FInterfaceBinding& Other)const
 		{
 			return SrcOffset == Other.SrcOffset && DestOffset == Other.DestOffset;
 		}
@@ -95,7 +95,7 @@ struct FNiagaraParameterStoreBinding
 			check(InSrcOffset < (int32)0xFFFF);
 			check(InDestOffset < (int32)0xFFFF);
 		}
-		FORCEINLINE bool operator==(const FUObjectBinding& Other)const
+		inline bool operator==(const FUObjectBinding& Other)const
 		{
 			return SrcOffset == Other.SrcOffset && DestOffset == Other.DestOffset;
 		}
@@ -105,8 +105,8 @@ struct FNiagaraParameterStoreBinding
 	/** Bindings of UObject params. Src and Dest offsets.*/
 	TArray<FUObjectBinding> UObjectBindings;
 
-	FORCEINLINE_DEBUGGABLE void Empty(FNiagaraParameterStore* DestStore, FNiagaraParameterStore* SrcStore);
-	FORCEINLINE_DEBUGGABLE bool Initialize(FNiagaraParameterStore* DestStore, FNiagaraParameterStore* SrcStore, const FNiagaraBoundParameterArray* BoundParameters = nullptr);
+	inline void Empty(FNiagaraParameterStore* DestStore, FNiagaraParameterStore* SrcStore);
+	inline bool Initialize(FNiagaraParameterStore* DestStore, FNiagaraParameterStore* SrcStore, const FNiagaraBoundParameterArray* BoundParameters = nullptr);
 	NIAGARA_API bool VerifyBinding(const FNiagaraParameterStore* DestStore, const FNiagaraParameterStore* SrcStore)const;
 	void CopyParameters(FNiagaraParameterStore* DestStore, const FNiagaraParameterStore* SrcStore) const;
 	NIAGARA_API void Dump(const FNiagaraParameterStore* DestStore, const FNiagaraParameterStore* SrcStore)const;
@@ -116,7 +116,7 @@ struct FNiagaraParameterStoreBinding
 private:
 	
 	template <typename TVisitor>
-	FORCEINLINE_DEBUGGABLE static void MatchParameters(FNiagaraParameterStore* DestStore, FNiagaraParameterStore* SrcStore, TVisitor Visitor);
+	inline static void MatchParameters(FNiagaraParameterStore* DestStore, FNiagaraParameterStore* SrcStore, TVisitor Visitor);
 
 	bool BindParameters(FNiagaraParameterStore* DestStore, FNiagaraParameterStore* SrcStore, const FNiagaraBoundParameterArray* BoundParameters = nullptr);
 };
@@ -127,8 +127,8 @@ struct FNiagaraVariableWithOffset : public FNiagaraVariableBase
 	GENERATED_USTRUCT_BODY()
 
 	// Those constructor enforce that there are no data allocated.
-	FORCEINLINE FNiagaraVariableWithOffset() : Offset(INDEX_NONE) {}
-	FORCEINLINE FNiagaraVariableWithOffset(const FNiagaraVariableBase& InVariable, int32 InOffset, const FNiagaraLwcStructConverter& InStructConverter) : FNiagaraVariableBase(InVariable.GetType(), InVariable.GetName()), Offset(InOffset), StructConverter(InStructConverter) {}
+	inline FNiagaraVariableWithOffset() : Offset(INDEX_NONE) {}
+	inline FNiagaraVariableWithOffset(const FNiagaraVariableBase& InVariable, int32 InOffset, const FNiagaraLwcStructConverter& InStructConverter) : FNiagaraVariableBase(InVariable.GetType(), InVariable.GetName()), Offset(InOffset), StructConverter(InStructConverter) {}
 
 	bool Serialize(FArchive& Ar);
 #if WITH_EDITORONLY_DATA
@@ -277,17 +277,17 @@ public:
 		return ResourceSize;
 	}
 
-	FORCEINLINE bool GetParametersDirty() const { return bParametersDirty; }
-	FORCEINLINE bool GetInterfacesDirty() const { return bInterfacesDirty; }
-	FORCEINLINE bool GetUObjectsDirty() const { return bUObjectsDirty; }
-	FORCEINLINE bool GetPositionDataDirty() const { return bPositionDataDirty; }
+	inline bool GetParametersDirty() const { return bParametersDirty; }
+	inline bool GetInterfacesDirty() const { return bInterfacesDirty; }
+	inline bool GetUObjectsDirty() const { return bUObjectsDirty; }
+	inline bool GetPositionDataDirty() const { return bPositionDataDirty; }
 
-	FORCEINLINE void MarkParametersDirty() { bParametersDirty = true; }
-	FORCEINLINE void MarkInterfacesDirty() { bInterfacesDirty = true; }
-	FORCEINLINE void MarkUObjectsDirty() { bUObjectsDirty = true; }
-	FORCEINLINE void MarkPositionDataDirty() { bPositionDataDirty = true; }
+	inline void MarkParametersDirty() { bParametersDirty = true; }
+	inline void MarkInterfacesDirty() { bInterfacesDirty = true; }
+	inline void MarkUObjectsDirty() { bUObjectsDirty = true; }
+	inline void MarkPositionDataDirty() { bPositionDataDirty = true; }
 
-	FORCEINLINE uint32 GetLayoutVersion() const { return LayoutVersion; }
+	inline uint32 GetLayoutVersion() const { return LayoutVersion; }
 
 	/** Binds this parameter store to another, by default if we find no matching parameters we will not maintain a pointer to the store. */
 	NIAGARA_API void Bind(FNiagaraParameterStore* DestStore, const FNiagaraBoundParameterArray* BoundParameters = nullptr);
@@ -343,7 +343,7 @@ public:
 	/** Removes all parameters from this store but doesn't change memory allocations. */
 	NIAGARA_API virtual void Reset(bool bClearBindings = true);
 
-	FORCEINLINE void GetParameters(TArray<FNiagaraVariable>& OutParameters) const
+	inline void GetParameters(TArray<FNiagaraVariable>& OutParameters) const
 	{
 		auto ParameterVariables = ReadParameterVariables();
 
@@ -354,14 +354,14 @@ public:
 		}
 	}
 
-	FORCEINLINE TArray<FNiagaraParameterStore*>& GetSourceParameterStores() { return SourceStores; }
+	inline TArray<FNiagaraParameterStore*>& GetSourceParameterStores() { return SourceStores; }
 
-	FORCEINLINE const TArray<TObjectPtr<UObject>>& GetUObjects()const { return UObjects; }
-	FORCEINLINE const TArray<UNiagaraDataInterface*>& GetDataInterfaces()const { return DataInterfaces; }
-	FORCEINLINE const TArray<uint8>& GetParameterDataArray()const { return ParameterData; }
+	inline const TArray<TObjectPtr<UObject>>& GetUObjects()const { return UObjects; }
+	inline const TArray<UNiagaraDataInterface*>& GetDataInterfaces()const { return DataInterfaces; }
+	inline const TArray<uint8>& GetParameterDataArray()const { return ParameterData; }
 
-	FORCEINLINE int32 Num() const {return SortedParameterOffsets.Num(); }
-	FORCEINLINE bool IsEmpty() const { return SortedParameterOffsets.Num() == 0; }
+	inline int32 Num() const {return SortedParameterOffsets.Num(); }
+	inline bool IsEmpty() const { return SortedParameterOffsets.Num() == 0; }
 
 	NIAGARA_API virtual void SanityCheckData(bool bInitInterfaces = true);
 
@@ -369,7 +369,7 @@ public:
 	NIAGARA_API virtual void InitFromSource(const FNiagaraParameterStore* SrcStore, bool bNotifyAsDirty);
 
 	/** Gets the index of the passed parameter. If it is a data interface, this is an offset into the data interface table, otherwise a byte offset into the parameter data buffer. */
-	FORCEINLINE_DEBUGGABLE int32 IndexOf(const FNiagaraVariableBase& Parameter) const
+	inline int32 IndexOf(const FNiagaraVariableBase& Parameter) const
 	{
 		const int32* Off = FindParameterOffset(Parameter);
 		return Off ? *Off : (int32)INDEX_NONE;
@@ -389,7 +389,7 @@ public:
 
 	/** Gets the typed parameter data. */
 	template<typename T>
-	FORCEINLINE_DEBUGGABLE void GetParameterValue(T& OutValue, const FNiagaraVariableBase& Parameter)const
+	inline void GetParameterValue(T& OutValue, const FNiagaraVariableBase& Parameter)const
 	{
 		check(Parameter.GetSizeInBytes() == sizeof(T));
 		int32 Offset = IndexOf(Parameter);
@@ -400,7 +400,7 @@ public:
 	}
 
 	template<typename T>
-	FORCEINLINE_DEBUGGABLE T GetParameterValue(const FNiagaraVariableBase& Parameter)const
+	inline T GetParameterValue(const FNiagaraVariableBase& Parameter)const
 	{
 		check(Parameter.GetSizeInBytes() == sizeof(T));
 		int32 Offset = IndexOf(Parameter);
@@ -412,7 +412,7 @@ public:
 	}
 
 	template<typename T>
-	FORCEINLINE_DEBUGGABLE TOptional<T> GetParameterOptionalValue(const FNiagaraVariableBase& Parameter)const
+	inline TOptional<T> GetParameterOptionalValue(const FNiagaraVariableBase& Parameter)const
 	{
 		check(Parameter.GetSizeInBytes() == sizeof(T));
 		TOptional<T> OptionalValue;
@@ -425,7 +425,7 @@ public:
 	}
 
 	template<typename T>
-	FORCEINLINE_DEBUGGABLE T GetParameterValueOrDefault(const FNiagaraVariableBase& Parameter, const T& DefaultValue) const
+	inline T GetParameterValueOrDefault(const FNiagaraVariableBase& Parameter, const T& DefaultValue) const
 	{
 		check(Parameter.GetSizeInBytes() == sizeof(T));
 		int32 Offset = IndexOf(Parameter);
@@ -433,32 +433,32 @@ public:
 	}
 
 	UE_DEPRECATED(5.5, "GetParameterData without a type side is deprecated as it's unsafe")
-	FORCEINLINE const uint8* GetParameterData(int32 Offset) const
+	inline const uint8* GetParameterData(int32 Offset) const
 	{
 		return ParameterData.GetData() + Offset;
 	}
 
-	FORCEINLINE const uint8* GetParameterData(int32 Offset, int32 SizeInBytes) const
+	inline const uint8* GetParameterData(int32 Offset, int32 SizeInBytes) const
 	{
 		ensureMsgf(Offset >= 0 && (Offset + SizeInBytes) <= ParameterData.Num(), TEXT("OOB immutable access offset(%d) size(%d) store size(%d)"), Offset, SizeInBytes, ParameterData.Num());
 		return ParameterData.GetData() + Offset;
 	}
 
-	FORCEINLINE const uint8* GetParameterData(int32 Offset, const FNiagaraTypeDefinition& TypeDef) const
+	inline const uint8* GetParameterData(int32 Offset, const FNiagaraTypeDefinition& TypeDef) const
 	{
 		ensureMsgf(Offset >= 0 && (Offset + TypeDef.GetSize()) <= ParameterData.Num(), TEXT("OOB immutable access offset(%d) size(%d) store size(%d)"), Offset, TypeDef.GetSize(), ParameterData.Num());
 		return ParameterData.GetData() + Offset;
 	}
 
 	/** Returns a none const pointer to the parameter data, any modifications to the data will not be recorded and automatically pushed to other parameter stores, use with caution. */
-	FORCEINLINE uint8* GetMutableParameterData(int32 Offset, const FNiagaraTypeDefinition& TypeDef)
+	inline uint8* GetMutableParameterData(int32 Offset, const FNiagaraTypeDefinition& TypeDef)
 	{
 		ensureMsgf(Offset >= 0 && (Offset + TypeDef.GetSize()) <= ParameterData.Num(), TEXT("OOB mutable access offset(%d) size(%d) store size(%d)"), Offset, TypeDef.GetSize(), ParameterData.Num());
 		return ParameterData.GetData() + Offset;
 	}
 
 	/** Returns the parameter data for the passed parameter if it exists in this store. Null if not. */
-	FORCEINLINE_DEBUGGABLE const uint8* GetParameterData(const FNiagaraVariableBase& Parameter) const
+	inline const uint8* GetParameterData(const FNiagaraVariableBase& Parameter) const
 	{
 		int32 Offset = IndexOf(Parameter);
 		if (Offset != INDEX_NONE)
@@ -468,7 +468,7 @@ public:
 		return nullptr;
 	}
 
-	FORCEINLINE_DEBUGGABLE const uint8* GetParameterData(const FNiagaraVariableWithOffset& Parameter) const
+	inline const uint8* GetParameterData(const FNiagaraVariableWithOffset& Parameter) const
 	{
 		check(Parameter.Offset != INDEX_NONE);
 		return GetParameterData_Internal(Parameter.Offset);
@@ -481,7 +481,7 @@ public:
 	NIAGARA_API bool CopyParameterData(const FNiagaraVariable& Parameter, uint8* DestinationData) const;
 
 	/** Returns the data interface at the passed offset. */
-	FORCEINLINE UNiagaraDataInterface* GetDataInterface(int32 Offset)const
+	inline UNiagaraDataInterface* GetDataInterface(int32 Offset)const
 	{
 		if (DataInterfaces.IsValidIndex(Offset))
 		{
@@ -511,7 +511,7 @@ public:
 	NIAGARA_API void SortParameters();
 
 	/** Returns the UObject at the passed offset. */
-	FORCEINLINE TObjectPtr<UObject> GetUObject(int32 Offset)const
+	inline TObjectPtr<UObject> GetUObject(int32 Offset)const
 	{
 		if (UObjects.IsValidIndex(Offset))
 		{
@@ -521,7 +521,7 @@ public:
 		return nullptr;
 	}
 
-	FORCEINLINE TObjectPtr<UObject> GetUObject(const FNiagaraVariable& Parameter)const
+	inline TObjectPtr<UObject> GetUObject(const FNiagaraVariable& Parameter)const
 	{
 		int32 Offset = IndexOf(Parameter);
 		TObjectPtr<UObject> Obj = GetUObject(Offset);
@@ -557,7 +557,7 @@ public:
 	NIAGARA_API void ResolvePositions(FNiagaraLWCConverter LwcConverter);
 
 	template<typename T>
-	FORCEINLINE_DEBUGGABLE bool SetParameterValue(const T& InValue, const FNiagaraVariable& Param, bool bAdd=false)
+	inline bool SetParameterValue(const T& InValue, const FNiagaraVariable& Param, bool bAdd=false)
 	{
 		check(Param.GetSizeInBytes() == sizeof(T));
 #if WITH_EDITOR
@@ -592,7 +592,7 @@ public:
 		return false;
 	}
 
-	FORCEINLINE_DEBUGGABLE void SetParameterData(const uint8* Data, int32 Offset, int32 Size)
+	inline void SetParameterData(const uint8* Data, int32 Offset, int32 Size)
 	{
 		check(Data != nullptr);
 		check(Offset >= 0 && (Offset + Size) <= ParameterData.Num());
@@ -606,13 +606,13 @@ public:
 
 	NIAGARA_API bool SetParameterData(const uint8* Data, FNiagaraVariable Param, bool bAdd = false);
 
-	FORCEINLINE_DEBUGGABLE void SetDataInterface(UNiagaraDataInterface* InInterface, int32 Offset)
+	inline void SetDataInterface(UNiagaraDataInterface* InInterface, int32 Offset)
 	{
 		DataInterfaces[Offset] = InInterface;
 		OnInterfaceChange();
 	}
 
-	FORCEINLINE_DEBUGGABLE void SetDataInterface(UNiagaraDataInterface* InInterface, const FNiagaraVariable& Parameter)
+	inline void SetDataInterface(UNiagaraDataInterface* InInterface, const FNiagaraVariable& Parameter)
 	{
 		int32 Offset = IndexOf(Parameter);
 		if (Offset != INDEX_NONE)
@@ -622,13 +622,13 @@ public:
 		}
 	}
 
-	FORCEINLINE_DEBUGGABLE void SetUObject(TObjectPtr<UObject> InObject, int32 Offset)
+	inline void SetUObject(TObjectPtr<UObject> InObject, int32 Offset)
 	{
 		UObjects[Offset] = InObject;
 		OnUObjectChange();
 	}
 
-	FORCEINLINE_DEBUGGABLE void SetUObject(TObjectPtr<UObject> InObject, const FNiagaraVariable& Parameter)
+	inline void SetUObject(TObjectPtr<UObject> InObject, const FNiagaraVariable& Parameter)
 	{
 		int32 Offset = IndexOf(Parameter);
 		if (Offset != INDEX_NONE)
@@ -638,7 +638,7 @@ public:
 		}
 	}
 
-	FORCEINLINE void OnParameterChange() 
+	inline void OnParameterChange() 
 	{ 
 		bParametersDirty = true;
 #if WITH_EDITOR
@@ -649,7 +649,7 @@ public:
 #endif
 	}
 
-	FORCEINLINE void OnInterfaceChange() 
+	inline void OnInterfaceChange() 
 	{ 
 		bInterfacesDirty = true;
 #if WITH_EDITOR
@@ -660,7 +660,7 @@ public:
 #endif
 	}
 
-	FORCEINLINE void OnUObjectChange()
+	inline void OnUObjectChange()
 	{
 		bUObjectsDirty = true;
 #if WITH_EDITOR
@@ -671,7 +671,7 @@ public:
 #endif
 	}
 
-	FORCEINLINE void PostGenericEditChange()
+	inline void PostGenericEditChange()
 	{
 		bUObjectsDirty = true;
 		bInterfacesDirty = true;
@@ -722,24 +722,24 @@ protected:
 };
 
 template<>
-FORCEINLINE bool FNiagaraParameterStore::SetParameterValue(const FVector3d& InValue, const FNiagaraVariable& Param, bool bAdd)
+inline bool FNiagaraParameterStore::SetParameterValue(const FVector3d& InValue, const FNiagaraVariable& Param, bool bAdd)
 {
 	return SetParameterValue((FVector3f)InValue, Param, bAdd);
 }
 
 template<>
-FORCEINLINE bool FNiagaraParameterStore::SetParameterValue(const FVector4d& InValue, const FNiagaraVariable& Param, bool bAdd)
+inline bool FNiagaraParameterStore::SetParameterValue(const FVector4d& InValue, const FNiagaraVariable& Param, bool bAdd)
 {
 	return SetParameterValue((FVector4f)InValue, Param, bAdd);
 }
 
 template<>
-FORCEINLINE bool FNiagaraParameterStore::SetParameterValue(const FQuat4d& InValue, const FNiagaraVariable& Param, bool bAdd)
+inline bool FNiagaraParameterStore::SetParameterValue(const FQuat4d& InValue, const FNiagaraVariable& Param, bool bAdd)
 {
 	return SetParameterValue((FQuat4f)InValue, Param, bAdd);
 }
 
-FORCEINLINE_DEBUGGABLE void FNiagaraParameterStoreBinding::Empty(FNiagaraParameterStore* DestStore, FNiagaraParameterStore* SrcStore)
+inline void FNiagaraParameterStoreBinding::Empty(FNiagaraParameterStore* DestStore, FNiagaraParameterStore* SrcStore)
 {
 	if (DestStore)
 	{
@@ -752,7 +752,7 @@ FORCEINLINE_DEBUGGABLE void FNiagaraParameterStoreBinding::Empty(FNiagaraParamet
 	UObjectBindings.Reset();
 }
 
-FORCEINLINE_DEBUGGABLE bool FNiagaraParameterStoreBinding::Initialize(FNiagaraParameterStore* DestStore, FNiagaraParameterStore* SrcStore, const FNiagaraBoundParameterArray* BoundParameters)
+inline bool FNiagaraParameterStoreBinding::Initialize(FNiagaraParameterStore* DestStore, FNiagaraParameterStore* SrcStore, const FNiagaraBoundParameterArray* BoundParameters)
 {
 	checkSlow(DestStore);
 	checkSlow(SrcStore);
@@ -796,7 +796,7 @@ struct FNiagaraParameterDirectBinding
 #endif
 	{}
 
-	FORCEINLINE bool IsBound()const { return ValuePtr != nullptr; }
+	inline bool IsBound()const { return ValuePtr != nullptr; }
 
 	T* Init(FNiagaraParameterStore& InStore, const FNiagaraVariable& DestVariable)
 	{
@@ -810,7 +810,7 @@ struct FNiagaraParameterDirectBinding
 		return ValuePtr;
 	}
 
-	FORCEINLINE void SetValue(const T& InValue)
+	inline void SetValue(const T& InValue)
 	{
 #if NIAGARA_VALIDATE_DIRECT_BINDINGS
 		checkSlow(BoundVariable.GetSizeInBytes() == sizeof(T));
@@ -822,7 +822,7 @@ struct FNiagaraParameterDirectBinding
 		}
 	}
 
-	FORCEINLINE T GetValue()const 
+	inline T GetValue()const 
 	{
 #if NIAGARA_VALIDATE_DIRECT_BINDINGS
 		checkSlow(BoundVariable.GetSizeInBytes() == sizeof(T));
@@ -853,7 +853,7 @@ struct FNiagaraParameterDirectBinding<FMatrix44f>
 #endif
 	{}
 
-	FORCEINLINE bool IsBound()const { return ValuePtr != nullptr; }
+	inline bool IsBound()const { return ValuePtr != nullptr; }
 
 	FMatrix44f* Init(FNiagaraParameterStore& InStore, const FNiagaraVariable& DestVariable)
 	{
@@ -867,7 +867,7 @@ struct FNiagaraParameterDirectBinding<FMatrix44f>
 		return ValuePtr;
 	}
 
-	FORCEINLINE void SetValue(const FMatrix44f& InValue)
+	inline void SetValue(const FMatrix44f& InValue)
 	{
 #if NIAGARA_VALIDATE_DIRECT_BINDINGS
 		checkSlow(BoundVariable.GetSizeInBytes() == sizeof(FMatrix44f));
@@ -879,7 +879,7 @@ struct FNiagaraParameterDirectBinding<FMatrix44f>
 		}
 	}
 
-	FORCEINLINE FMatrix44f GetValue()const
+	inline FMatrix44f GetValue()const
 	{
 #if NIAGARA_VALIDATE_DIRECT_BINDINGS
 		checkSlow(BoundVariable.GetSizeInBytes() == sizeof(FMatrix44f));
@@ -911,7 +911,7 @@ struct FNiagaraParameterDirectBinding<FVector4f>
 #endif
 	{}
 
-	FORCEINLINE bool IsBound()const { return ValuePtr != nullptr; }
+	inline bool IsBound()const { return ValuePtr != nullptr; }
 
 	FVector4f* Init(FNiagaraParameterStore& InStore, const FNiagaraVariable& DestVariable)
 	{
@@ -925,7 +925,7 @@ struct FNiagaraParameterDirectBinding<FVector4f>
 		return ValuePtr;
 	}
 
-	FORCEINLINE void SetValue(const FVector4f& InValue)
+	inline void SetValue(const FVector4f& InValue)
 	{
 #if NIAGARA_VALIDATE_DIRECT_BINDINGS
 		checkSlow(BoundVariable.GetSizeInBytes() == sizeof(FVector4f));
@@ -937,7 +937,7 @@ struct FNiagaraParameterDirectBinding<FVector4f>
 		}
 	}
 
-	FORCEINLINE FVector4f GetValue()const
+	inline FVector4f GetValue()const
 	{
 #if NIAGARA_VALIDATE_DIRECT_BINDINGS
 		check(BoundVariable.GetSizeInBytes() == sizeof(FVector4f));
@@ -969,7 +969,7 @@ struct FNiagaraParameterDirectBinding<FQuat4f>
 #endif
 	{}
 
-	FORCEINLINE bool IsBound()const { return ValuePtr != nullptr; }
+	inline bool IsBound()const { return ValuePtr != nullptr; }
 
 	FQuat4f* Init(FNiagaraParameterStore& InStore, const FNiagaraVariable& DestVariable)
 	{
@@ -983,7 +983,7 @@ struct FNiagaraParameterDirectBinding<FQuat4f>
 		return ValuePtr;
 	}
 
-	FORCEINLINE void SetValue(const FQuat4f& InValue)
+	inline void SetValue(const FQuat4f& InValue)
 	{
 #if NIAGARA_VALIDATE_DIRECT_BINDINGS
 		checkSlow(BoundVariable.GetSizeInBytes() == sizeof(FQuat4f));
@@ -995,7 +995,7 @@ struct FNiagaraParameterDirectBinding<FQuat4f>
 		}
 	}
 
-	FORCEINLINE FQuat4f GetValue()const
+	inline FQuat4f GetValue()const
 	{
 #if NIAGARA_VALIDATE_DIRECT_BINDINGS
 		checkSlow(BoundVariable.GetSizeInBytes() == sizeof(FQuat4f));
@@ -1027,7 +1027,7 @@ struct FNiagaraParameterDirectBinding<FNiagaraBool>
 #endif
 	{}
 
-	FORCEINLINE bool IsBound()const { return ValuePtr != nullptr; }
+	inline bool IsBound()const { return ValuePtr != nullptr; }
 
 	uint32* Init(FNiagaraParameterStore& InStore, const FNiagaraVariable& DestVariable)
 	{
@@ -1042,7 +1042,7 @@ struct FNiagaraParameterDirectBinding<FNiagaraBool>
 		return ValuePtr;
 	}
 
-	FORCEINLINE void SetValue(const FNiagaraBool& InValue)
+	inline void SetValue(const FNiagaraBool& InValue)
 	{
 #if NIAGARA_VALIDATE_DIRECT_BINDINGS
 		checkSlow(BoundVariable.GetSizeInBytes() == sizeof(FNiagaraBool));
@@ -1055,7 +1055,7 @@ struct FNiagaraParameterDirectBinding<FNiagaraBool>
 		}
 	}
 
-	FORCEINLINE void SetValue(const bool& InValue)
+	inline void SetValue(const bool& InValue)
 	{
 #if NIAGARA_VALIDATE_DIRECT_BINDINGS
 		checkSlow(BoundVariable.GetSizeInBytes() == sizeof(FNiagaraBool));
@@ -1074,7 +1074,7 @@ struct FNiagaraParameterDirectBinding<FNiagaraBool>
 		}
 	}
 
-	FORCEINLINE FNiagaraBool GetValue()const
+	inline FNiagaraBool GetValue()const
 	{
 #if NIAGARA_VALIDATE_DIRECT_BINDINGS
 		checkSlow(BoundVariable.GetSizeInBytes() == sizeof(FNiagaraBool));
@@ -1108,7 +1108,7 @@ struct FNiagaraParameterDirectBinding<UObject*>
 		: UObjectOffset(INDEX_NONE), BoundStore(nullptr)
 	{}
 
-	FORCEINLINE bool IsBound()const { return UObjectOffset != INDEX_NONE; }
+	inline bool IsBound()const { return UObjectOffset != INDEX_NONE; }
 
 	TObjectPtr<UObject> Init(FNiagaraParameterStore& InStore, const FNiagaraVariable& DestVariable)
 	{
@@ -1126,7 +1126,7 @@ struct FNiagaraParameterDirectBinding<UObject*>
 		return nullptr;
 	}
 
-	FORCEINLINE void SetValue(UObject* InValue)
+	inline void SetValue(UObject* InValue)
 	{
 		if (UObjectOffset != INDEX_NONE)
 		{
@@ -1137,7 +1137,7 @@ struct FNiagaraParameterDirectBinding<UObject*>
 		}
 	}
 
-	FORCEINLINE TObjectPtr<UObject> GetValue() const
+	inline TObjectPtr<UObject> GetValue() const
 	{
 		if (UObjectOffset != INDEX_NONE)
 		{
@@ -1150,7 +1150,7 @@ struct FNiagaraParameterDirectBinding<UObject*>
 	}
 
 	template<class TObjectType>
-	FORCEINLINE TObjectType* GetValueOrDefault(TObjectType* DefaultValue) const
+	inline TObjectType* GetValueOrDefault(TObjectType* DefaultValue) const
 	{
 		if (UObjectOffset != INDEX_NONE)
 		{
@@ -1163,7 +1163,7 @@ struct FNiagaraParameterDirectBinding<UObject*>
 	}
 
 	template<class TObjectType>
-	FORCEINLINE TObjectType* GetValue() const
+	inline TObjectType* GetValue() const
 	{
 		return GetValueOrDefault<TObjectType>(nullptr);
 	}
diff --git a/Engine/Plugins/FX/Niagara/Source/Niagara/Public/NiagaraPreviewGrid.h b/Engine/Plugins/FX/Niagara/Source/Niagara/Public/NiagaraPreviewGrid.h
--- a/Engine/Plugins/FX/Niagara/Source/Niagara/Public/NiagaraPreviewGrid.h
+++ b/Engine/Plugins/FX/Niagara/Source/Niagara/Public/NiagaraPreviewGrid.h
@@ -216,7 +216,7 @@ private:
 	void GeneratePreviews();
 	void TickPreviews();
 
-	FORCEINLINE int32 PreviewIndex(int32 InX, int32 InY)const { return InX * NumY + InY; }
+	inline int32 PreviewIndex(int32 InX, int32 InY)const { return InX * NumY + InY; }
 
 public:
 
diff --git a/Engine/Plugins/FX/Niagara/Source/Niagara/Public/NiagaraRenderer.h b/Engine/Plugins/FX/Niagara/Source/Niagara/Public/NiagaraRenderer.h
--- a/Engine/Plugins/FX/Niagara/Source/Niagara/Public/NiagaraRenderer.h
+++ b/Engine/Plugins/FX/Niagara/Source/Niagara/Public/NiagaraRenderer.h
@@ -38,11 +38,11 @@ struct FNiagaraDynamicDataBase
 
 	NIAGARA_API bool IsGpuLowLatencyTranslucencyEnabled() const;
 	NIAGARA_API FNiagaraDataBuffer* GetParticleDataToRender(FRHICommandListBase& RHICmdList, bool bIsLowLatencyTranslucent = false) const;
-	FORCEINLINE FMaterialRelevance GetMaterialRelevance() const { return MaterialRelevance; }
+	inline FMaterialRelevance GetMaterialRelevance() const { return MaterialRelevance; }
 
-	FORCEINLINE void SetMaterialRelevance(FMaterialRelevance NewRelevance) { MaterialRelevance = NewRelevance; }
+	inline void SetMaterialRelevance(FMaterialRelevance NewRelevance) { MaterialRelevance = NewRelevance; }
 
-	FORCEINLINE FNiagaraSystemInstanceID GetSystemInstanceID() const { return SystemInstanceID; }
+	inline FNiagaraSystemInstanceID GetSystemInstanceID() const { return SystemInstanceID; }
 
 	virtual void ApplyMaterialOverride(int32 MaterialIndex, UMaterialInterface* MaterialOverride) {};
 
@@ -102,10 +102,10 @@ public:
 	static NIAGARA_API FVector4f CalcMacroUVParameters(const FSceneView& View, FVector MacroUVPosition, float MacroUVRadius);
 
 	NIAGARA_API void SetDynamicData_RenderThread(FNiagaraDynamicDataBase* NewDynamicData);
-	FORCEINLINE FNiagaraDynamicDataBase *GetDynamicData() const { return DynamicDataRender; }
-	FORCEINLINE bool HasDynamicData() const { return DynamicDataRender != nullptr; }
-	FORCEINLINE bool HasLights() const { return bHasLights; }
-	FORCEINLINE bool IsMotionBlurEnabled() const { return bMotionBlurEnabled; }
+	inline FNiagaraDynamicDataBase *GetDynamicData() const { return DynamicDataRender; }
+	inline bool HasDynamicData() const { return DynamicDataRender != nullptr; }
+	inline bool HasLights() const { return bHasLights; }
+	inline bool IsMotionBlurEnabled() const { return bMotionBlurEnabled; }
 
 #if RHI_RAYTRACING
 	PRAGMA_DISABLE_DEPRECATION_WARNINGS
@@ -115,21 +115,21 @@ public:
 	virtual void GetDynamicRayTracingInstances(FRayTracingInstanceCollector& Collector, const FNiagaraSceneProxy* Proxy) {}
 #endif
 
-	FORCEINLINE static FRHIShaderResourceView* GetSrvOrDefaultFloat(const FRWBuffer& RWBuffer) { return RWBuffer.SRV.IsValid() ? (FRHIShaderResourceView*)RWBuffer.SRV : GetDummyFloatBuffer(); }
-	FORCEINLINE static FRHIShaderResourceView* GetSrvOrDefaultHalf(const FRWBuffer& RWBuffer) { return RWBuffer.SRV.IsValid() ? (FRHIShaderResourceView*)RWBuffer.SRV : GetDummyHalfBuffer(); }
-	FORCEINLINE static FRHIShaderResourceView* GetSrvOrDefaultUInt(const FRWBuffer& RWBuffer) { return RWBuffer.SRV.IsValid() ? (FRHIShaderResourceView*)RWBuffer.SRV : GetDummyUIntBuffer(); }
-	FORCEINLINE static FRHIShaderResourceView* GetSrvOrDefaultInt(const FRWBuffer& RWBuffer) { return RWBuffer.SRV.IsValid() ? (FRHIShaderResourceView*)RWBuffer.SRV : GetDummyIntBuffer(); }
+	inline static FRHIShaderResourceView* GetSrvOrDefaultFloat(const FRWBuffer& RWBuffer) { return RWBuffer.SRV.IsValid() ? (FRHIShaderResourceView*)RWBuffer.SRV : GetDummyFloatBuffer(); }
+	inline static FRHIShaderResourceView* GetSrvOrDefaultHalf(const FRWBuffer& RWBuffer) { return RWBuffer.SRV.IsValid() ? (FRHIShaderResourceView*)RWBuffer.SRV : GetDummyHalfBuffer(); }
+	inline static FRHIShaderResourceView* GetSrvOrDefaultUInt(const FRWBuffer& RWBuffer) { return RWBuffer.SRV.IsValid() ? (FRHIShaderResourceView*)RWBuffer.SRV : GetDummyUIntBuffer(); }
+	inline static FRHIShaderResourceView* GetSrvOrDefaultInt(const FRWBuffer& RWBuffer) { return RWBuffer.SRV.IsValid() ? (FRHIShaderResourceView*)RWBuffer.SRV : GetDummyIntBuffer(); }
 
-	FORCEINLINE static FRHIShaderResourceView* GetSrvOrDefaultFloat(FGlobalDynamicReadBuffer::FAllocation& Allocation) { return Allocation.IsValid() ? (FRHIShaderResourceView*)Allocation.SRV : GetDummyFloatBuffer(); }
-	FORCEINLINE static FRHIShaderResourceView* GetSrvOrDefaultHalf(FGlobalDynamicReadBuffer::FAllocation& Allocation) { return Allocation.IsValid() ? (FRHIShaderResourceView*)Allocation.SRV : GetDummyHalfBuffer(); }
-	FORCEINLINE static FRHIShaderResourceView* GetSrvOrDefaultUInt(FGlobalDynamicReadBuffer::FAllocation& Allocation) { return Allocation.IsValid() ? (FRHIShaderResourceView*)Allocation.SRV : GetDummyUIntBuffer(); }
-	FORCEINLINE static FRHIShaderResourceView* GetSrvOrDefaultInt(FGlobalDynamicReadBuffer::FAllocation& Allocation) { return Allocation.IsValid() ? (FRHIShaderResourceView*)Allocation.SRV : GetDummyIntBuffer(); }
+	inline static FRHIShaderResourceView* GetSrvOrDefaultFloat(FGlobalDynamicReadBuffer::FAllocation& Allocation) { return Allocation.IsValid() ? (FRHIShaderResourceView*)Allocation.SRV : GetDummyFloatBuffer(); }
+	inline static FRHIShaderResourceView* GetSrvOrDefaultHalf(FGlobalDynamicReadBuffer::FAllocation& Allocation) { return Allocation.IsValid() ? (FRHIShaderResourceView*)Allocation.SRV : GetDummyHalfBuffer(); }
+	inline static FRHIShaderResourceView* GetSrvOrDefaultUInt(FGlobalDynamicReadBuffer::FAllocation& Allocation) { return Allocation.IsValid() ? (FRHIShaderResourceView*)Allocation.SRV : GetDummyUIntBuffer(); }
+	inline static FRHIShaderResourceView* GetSrvOrDefaultInt(FGlobalDynamicReadBuffer::FAllocation& Allocation) { return Allocation.IsValid() ? (FRHIShaderResourceView*)Allocation.SRV : GetDummyIntBuffer(); }
 
-	FORCEINLINE static FRHIShaderResourceView* GetSrvOrDefaultFloat(FRHIShaderResourceView* InSRV) { return InSRV ? InSRV : GetDummyFloatBuffer(); }
-	FORCEINLINE static FRHIShaderResourceView* GetSrvOrDefaultFloat2(FRHIShaderResourceView* InSRV) { return InSRV ? InSRV : GetDummyFloat2Buffer(); }
-	FORCEINLINE static FRHIShaderResourceView* GetSrvOrDefaultFloat4(FRHIShaderResourceView* InSRV) { return InSRV ? InSRV : GetDummyFloat4Buffer(); }
-	FORCEINLINE static FRHIShaderResourceView* GetSrvOrDefaultUInt(FRHIShaderResourceView* InSRV) { return InSRV ? InSRV : GetDummyUIntBuffer(); }
-	FORCEINLINE static FRHIShaderResourceView* GetSrvOrDefaultInt(FRHIShaderResourceView* InSRV) { return InSRV ? InSRV : GetDummyIntBuffer(); }
+	inline static FRHIShaderResourceView* GetSrvOrDefaultFloat(FRHIShaderResourceView* InSRV) { return InSRV ? InSRV : GetDummyFloatBuffer(); }
+	inline static FRHIShaderResourceView* GetSrvOrDefaultFloat2(FRHIShaderResourceView* InSRV) { return InSRV ? InSRV : GetDummyFloat2Buffer(); }
+	inline static FRHIShaderResourceView* GetSrvOrDefaultFloat4(FRHIShaderResourceView* InSRV) { return InSRV ? InSRV : GetDummyFloat4Buffer(); }
+	inline static FRHIShaderResourceView* GetSrvOrDefaultUInt(FRHIShaderResourceView* InSRV) { return InSRV ? InSRV : GetDummyUIntBuffer(); }
+	inline static FRHIShaderResourceView* GetSrvOrDefaultInt(FRHIShaderResourceView* InSRV) { return InSRV ? InSRV : GetDummyIntBuffer(); }
 
 	static NIAGARA_API FRHIShaderResourceView* GetDummyFloatBuffer();
 	static NIAGARA_API FRHIShaderResourceView* GetDummyFloat2Buffer();
@@ -147,7 +147,7 @@ public:
 		return GPUOffset == INDEX_NONE ? Variables[Fallback].GetGPUOffset() : GPUOffset;
 	}
 
-	FORCEINLINE ENiagaraSimTarget GetSimTarget() const { return SimTarget; }
+	inline ENiagaraSimTarget GetSimTarget() const { return SimTarget; }
 
 	virtual void GetUsedMaterials(TArray<UMaterialInterface*>& UsedMaterials, bool bGetDebugMaterials) { UsedMaterials.Append(BaseMaterials_GT); }
 
diff --git a/Engine/Plugins/FX/Niagara/Source/Niagara/Public/NiagaraRendererProperties.h b/Engine/Plugins/FX/Niagara/Source/Niagara/Public/NiagaraRendererProperties.h
--- a/Engine/Plugins/FX/Niagara/Source/Niagara/Public/NiagaraRendererProperties.h
+++ b/Engine/Plugins/FX/Niagara/Source/Niagara/Public/NiagaraRendererProperties.h
@@ -124,16 +124,16 @@ struct FNiagaraRendererVariableInfo
 	}
 	NIAGARA_API FNiagaraRendererVariableInfo(const FNiagaraDataSetCompiledData* CompiledData, const FNiagaraVariableBase& Variable);
 
-	FORCEINLINE int32 GetNumComponents() const { return NumComponents; }
+	inline int32 GetNumComponents() const { return NumComponents; }
 
-	FORCEINLINE int32 GetRawGPUOffset() const { return GPUBufferOffset == kInvalidOffset ? INDEX_NONE : GPUBufferOffset; }
-	FORCEINLINE int32 GetGPUOffset() const { return GetRawGPUOffset() | (bHalfType ? (1 << 31) : 0); }
+	inline int32 GetRawGPUOffset() const { return GPUBufferOffset == kInvalidOffset ? INDEX_NONE : GPUBufferOffset; }
+	inline int32 GetGPUOffset() const { return GetRawGPUOffset() | (bHalfType ? (1 << 31) : 0); }
 
-	FORCEINLINE int32 GetRawDatasetOffset() const { return DatasetOffset == kInvalidOffset ? INDEX_NONE : DatasetOffset; }
-	FORCEINLINE int32 GetEncodedDatasetOffset() const { return GetRawDatasetOffset() | (bHalfType ? (1 << 31) : 0); }
+	inline int32 GetRawDatasetOffset() const { return DatasetOffset == kInvalidOffset ? INDEX_NONE : DatasetOffset; }
+	inline int32 GetEncodedDatasetOffset() const { return GetRawDatasetOffset() | (bHalfType ? (1 << 31) : 0); }
 
-	FORCEINLINE bool ShouldUpload() const { return bUpload; }
-	FORCEINLINE bool IsHalfType() const { return bHalfType; }
+	inline bool ShouldUpload() const { return bUpload; }
+	inline bool IsHalfType() const { return bHalfType; }
 
 protected:
 	uint16 DatasetOffset = kInvalidOffset;
diff --git a/Engine/Plugins/FX/Niagara/Source/Niagara/Public/NiagaraRendererRibbons.h b/Engine/Plugins/FX/Niagara/Source/Niagara/Public/NiagaraRendererRibbons.h
--- a/Engine/Plugins/FX/Niagara/Source/Niagara/Public/NiagaraRendererRibbons.h
+++ b/Engine/Plugins/FX/Niagara/Source/Niagara/Public/NiagaraRendererRibbons.h
@@ -127,7 +127,7 @@ struct FRibbonMultiRibbonInfo
 	FRibbonMultiRibbonInfoBufferEntry BufferEntry;
 	
 
-	FORCEINLINE bool UseInvertOrder(const FVector& ViewDirection, const FVector& ViewOriginForDistanceCulling, ENiagaraRibbonDrawDirection DrawDirection) const
+	inline bool UseInvertOrder(const FVector& ViewDirection, const FVector& ViewOriginForDistanceCulling, ENiagaraRibbonDrawDirection DrawDirection) const
 	{
 		const double StartDist = FVector::DotProduct(ViewDirection, StartPos - ViewOriginForDistanceCulling);
 		const double EndDist = FVector::DotProduct(ViewDirection, EndPos - ViewOriginForDistanceCulling);
diff --git a/Engine/Plugins/FX/Niagara/Source/Niagara/Public/NiagaraSceneProxy.h b/Engine/Plugins/FX/Niagara/Source/Niagara/Public/NiagaraSceneProxy.h
--- a/Engine/Plugins/FX/Niagara/Source/Niagara/Public/NiagaraSceneProxy.h
+++ b/Engine/Plugins/FX/Niagara/Source/Niagara/Public/NiagaraSceneProxy.h
@@ -56,7 +56,7 @@ public:
 	virtual bool HasRayTracingRepresentation() const override { return true; }
 #endif
 
-	FORCEINLINE const FMatrix& GetLocalToWorldInverse() const { return LocalToWorldInverse; }
+	inline const FMatrix& GetLocalToWorldInverse() const { return LocalToWorldInverse; }
 
 	NIAGARA_API const FVector3f& GetLWCRenderTile() const;
 
diff --git a/Engine/Plugins/FX/Niagara/Source/Niagara/Public/NiagaraSystemInstance.h b/Engine/Plugins/FX/Niagara/Source/Niagara/Public/NiagaraSystemInstance.h
--- a/Engine/Plugins/FX/Niagara/Source/Niagara/Public/NiagaraSystemInstance.h
+++ b/Engine/Plugins/FX/Niagara/Source/Niagara/Public/NiagaraSystemInstance.h
@@ -117,7 +117,7 @@ private:
 	EResetMode DeferredResetMode = EResetMode::None;
 
 public:
-	FORCEINLINE bool GetAreDataInterfacesInitialized() const { return bDataInterfacesInitialized; }
+	inline bool GetAreDataInterfacesInitialized() const { return bDataInterfacesInitialized; }
 
 	/** Creates a new Niagara system instance. */
 	NIAGARA_API FNiagaraSystemInstance(UWorld& InWorld, UNiagaraSystem& InAsset, FNiagaraUserRedirectionParameterStore* InOverrideParameters = nullptr,
@@ -138,7 +138,7 @@ public:
 	NIAGARA_API void OnPooledReuse(UWorld& NewWorld);
 
 	NIAGARA_API void SetPaused(bool bInPaused);
-	FORCEINLINE bool IsPaused() const { return (SystemInstanceState == ENiagaraSystemInstanceState::PendingSpawnPaused) || (SystemInstanceState == ENiagaraSystemInstanceState::Paused); }
+	inline bool IsPaused() const { return (SystemInstanceState == ENiagaraSystemInstanceState::PendingSpawnPaused) || (SystemInstanceState == ENiagaraSystemInstanceState::Paused); }
 
 	NIAGARA_API void SetSolo(bool bInSolo);
 
@@ -157,17 +157,17 @@ public:
 	// Unbinds the parameters that would be bound via BindToParameterStore
 	void UnbindFromParameterStore(FNiagaraParameterStore& ParameterStore);
 
-	FORCEINLINE FNiagaraParameterStore& GetInstanceParameters() { return InstanceParameters; }
+	inline FNiagaraParameterStore& GetInstanceParameters() { return InstanceParameters; }
 	NIAGARA_API FNiagaraLWCConverter GetLWCConverter(bool bLocalSpaceEmitter = false) const;
 	NIAGARA_API FTransform GetLWCSimToWorld(bool bLocalSpaceEmitter = false) const;
 	void UpdateLWCTile(const FVector3f& InLWCTile, FNiagaraDataSet& EmitterDataSet);
 
-	FORCEINLINE uint32 GetParameterIndex(bool PreviousFrame = false) const
+	inline uint32 GetParameterIndex(bool PreviousFrame = false) const
 	{
 		return (!!(PreviousFrame && ParametersValid) ^ !!CurrentFrameIndex) ? 1 : 0;
 	}
 
-	FORCEINLINE void FlipParameterBuffers()
+	inline void FlipParameterBuffers()
 	{
 		CurrentFrameIndex = ~CurrentFrameIndex;
 		
@@ -178,11 +178,11 @@ public:
 		}
 	}
 
-	FORCEINLINE const FNiagaraGlobalParameters& GetGlobalParameters(bool PreviousFrame = false) const { return GlobalParameters[GetParameterIndex(PreviousFrame)]; }
-	FORCEINLINE const FNiagaraSystemParameters& GetSystemParameters(bool PreviousFrame = false) const { return SystemParameters[GetParameterIndex(PreviousFrame)]; }
-	FORCEINLINE const FNiagaraOwnerParameters& GetOwnerParameters(bool PreviousFrame = false) const { return OwnerParameters[GetParameterIndex(PreviousFrame)]; }
-	FORCEINLINE const FNiagaraEmitterParameters& GetEmitterParameters(int32 EmitterIdx, bool PreviousFrame = false) const { return EmitterParameters[EmitterIdx * 2 + GetParameterIndex(PreviousFrame)]; }
-	FORCEINLINE FNiagaraEmitterParameters& EditEmitterParameters(int32 EmitterIdx) { return EmitterParameters[EmitterIdx * 2 + GetParameterIndex()]; }
+	inline const FNiagaraGlobalParameters& GetGlobalParameters(bool PreviousFrame = false) const { return GlobalParameters[GetParameterIndex(PreviousFrame)]; }
+	inline const FNiagaraSystemParameters& GetSystemParameters(bool PreviousFrame = false) const { return SystemParameters[GetParameterIndex(PreviousFrame)]; }
+	inline const FNiagaraOwnerParameters& GetOwnerParameters(bool PreviousFrame = false) const { return OwnerParameters[GetParameterIndex(PreviousFrame)]; }
+	inline const FNiagaraEmitterParameters& GetEmitterParameters(int32 EmitterIdx, bool PreviousFrame = false) const { return EmitterParameters[EmitterIdx * 2 + GetParameterIndex(PreviousFrame)]; }
+	inline FNiagaraEmitterParameters& EditEmitterParameters(int32 EmitterIdx) { return EmitterParameters[EmitterIdx * 2 + GetParameterIndex()]; }
 	
 	NIAGARA_API FNiagaraWorldManager* GetWorldManager()const;
 	NIAGARA_API bool RequiresGlobalDistanceField() const;
@@ -259,28 +259,28 @@ public:
 
 //	FVector GetOwnerVelocity() const { return OwnerVelocityParam.GetValue(); }
 
-	FORCEINLINE bool IsComplete()const { return ActualExecutionState == ENiagaraExecutionState::Complete || ActualExecutionState == ENiagaraExecutionState::Disabled; }
-	FORCEINLINE bool IsDisabled()const { return ActualExecutionState == ENiagaraExecutionState::Disabled; }
+	inline bool IsComplete()const { return ActualExecutionState == ENiagaraExecutionState::Complete || ActualExecutionState == ENiagaraExecutionState::Disabled; }
+	inline bool IsDisabled()const { return ActualExecutionState == ENiagaraExecutionState::Disabled; }
 
 	/** Gets the simulation for the supplied emitter handle. */
 	NIAGARA_API FNiagaraEmitterInstancePtr GetSimulationForHandle(const FNiagaraEmitterHandle& EmitterHandle) const;
 
-	FORCEINLINE UWorld* GetWorld() const { return World; }
-	FORCEINLINE UNiagaraSystem* GetSystem() const { return System; }
-	FORCEINLINE USceneComponent* GetAttachComponent() { return AttachComponent.Get(); }
-	FORCEINLINE FNiagaraUserRedirectionParameterStore* GetOverrideParameters() { return OverrideParameters; }
-	FORCEINLINE const FNiagaraUserRedirectionParameterStore* GetOverrideParameters() const { return OverrideParameters; }
+	inline UWorld* GetWorld() const { return World; }
+	inline UNiagaraSystem* GetSystem() const { return System; }
+	inline USceneComponent* GetAttachComponent() { return AttachComponent.Get(); }
+	inline FNiagaraUserRedirectionParameterStore* GetOverrideParameters() { return OverrideParameters; }
+	inline const FNiagaraUserRedirectionParameterStore* GetOverrideParameters() const { return OverrideParameters; }
 	[[nodiscard]] TArrayView<FNiagaraEmitterInstanceRef> GetEmitters() { return Emitters; }
 	[[nodiscard]] TConstArrayView<FNiagaraEmitterInstanceRef> GetEmitters() const { return Emitters; }
 
 	UE_DEPRECATED(5.5, "This function is experimental and subject to change.")
 	static const INiagaraSceneComponentUtils& GetSceneComponentUtils(FNiagaraSystemInstance* Instance);
 
-	FORCEINLINE const FBox& GetLocalBounds() const { return LocalBounds;  }
-	FORCEINLINE const FVector3f& GetLWCTile() const { return LWCTile;  }
+	inline const FBox& GetLocalBounds() const { return LocalBounds;  }
+	inline const FVector3f& GetLWCTile() const { return LWCTile;  }
 	NIAGARA_API TConstArrayView<FNiagaraEmitterExecutionIndex> GetEmitterExecutionOrder() const;
 
-	FORCEINLINE void SetSystemFixedBounds(const FBox& InLocalBounds) { FixedBounds_GT = InLocalBounds; }
+	inline void SetSystemFixedBounds(const FBox& InLocalBounds) { FixedBounds_GT = InLocalBounds; }
 	NIAGARA_API FBox GetSystemFixedBounds() const;
 	NIAGARA_API void SetEmitterFixedBounds(FName EmitterName, const FBox& InLocalBounds);
 	NIAGARA_API FBox GetEmitterFixedBounds(FName EmitterName) const;
@@ -288,16 +288,16 @@ public:
 	NIAGARA_API FNiagaraEmitterInstance* GetEmitterByID(FNiagaraEmitterID ID)const;
 
 	NIAGARA_API void SetForceSolo(bool bForceSolo);
-	FORCEINLINE bool IsSolo() const { return bSolo; }
+	inline bool IsSolo() const { return bSolo; }
 
-	FORCEINLINE bool NeedsGPUTick() const { return ActiveGPUEmitterCount > 0 /*&& Component->IsRegistered()*/ && !IsComplete();}
+	inline bool NeedsGPUTick() const { return ActiveGPUEmitterCount > 0 /*&& Component->IsRegistered()*/ && !IsComplete();}
 
 	FNiagaraSystemGpuComputeProxy* GetSystemGpuComputeProxy() { return SystemGpuComputeProxy.Get(); }
 
 	/** Gets a multicast delegate which is called after this instance has finished ticking for the frame on the game thread */
-	FORCEINLINE void SetOnPostTick(const FOnPostTick& InPostTickDelegate) { OnPostTickDelegate = InPostTickDelegate; }
+	inline void SetOnPostTick(const FOnPostTick& InPostTickDelegate) { OnPostTickDelegate = InPostTickDelegate; }
 	/** Gets a multicast delegate which is called whenever this instance is complete. */
-	FORCEINLINE void SetOnComplete(const FOnComplete& InOnCompleteDelegate) { OnCompleteDelegate = InOnCompleteDelegate; }
+	inline void SetOnComplete(const FOnComplete& InOnCompleteDelegate) { OnCompleteDelegate = InOnCompleteDelegate; }
 
 #if WITH_EDITOR
 	/** Gets a multicast delegate which is called whenever this instance is initialized with an System asset. */
@@ -316,7 +316,7 @@ public:
 	FNiagaraSystemInstanceID GetId() const { return ID; }
 
 	/** Returns the instance data for a particular interface for this System. */
-	FORCEINLINE void* FindDataInterfaceInstanceData(const UNiagaraDataInterface* Interface)
+	inline void* FindDataInterfaceInstanceData(const UNiagaraDataInterface* Interface)
 	{
 		if (auto* InstDataOffsetPair = DataInterfaceInstanceDataOffsets.FindByPredicate([&](auto& Pair){ return Pair.Key.Get() == Interface;}))
 		{
@@ -326,7 +326,7 @@ public:
 	}
 
 	template<typename TDataType>
-	FORCEINLINE TDataType* FindTypedDataInterfaceInstanceData(const UNiagaraDataInterface* Interface)
+	inline TDataType* FindTypedDataInterfaceInstanceData(const UNiagaraDataInterface* Interface)
 	{
 		if (auto* InstDataOffsetPair = DataInterfaceInstanceDataOffsets.FindByPredicate([&](auto& Pair) { return Pair.Key.Get() == Interface; }))
 		{
@@ -335,12 +335,12 @@ public:
 		return nullptr;
 	}
 
-	FORCEINLINE int32 FindDataInterfaceInstanceDataIndex(UNiagaraDataInterface* Interface)
+	inline int32 FindDataInterfaceInstanceDataIndex(UNiagaraDataInterface* Interface)
 	{
 		return DataInterfaceInstanceDataOffsets.IndexOfByPredicate([&](auto& Pair) { return Pair.Key.Get() == Interface; });
 	}
 
-	FORCEINLINE const void GetDataInterfaceInstanceDataInfo(int32 Index, UNiagaraDataInterface*& OutInterface, void*& OutInstanceData)
+	inline const void GetDataInterfaceInstanceDataInfo(int32 Index, UNiagaraDataInterface*& OutInterface, void*& OutInstanceData)
 	{
 		if(DataInterfaceInstanceDataOffsets.IsValidIndex(Index))
 		{
@@ -354,11 +354,11 @@ public:
 		}
 	}
 	
-	FORCEINLINE const void* FindDataInterfaceInstanceData(const UNiagaraDataInterface* Interface) const { return const_cast<FNiagaraSystemInstance*>(this)->FindDataInterfaceInstanceData(Interface); }
+	inline const void* FindDataInterfaceInstanceData(const UNiagaraDataInterface* Interface) const { return const_cast<FNiagaraSystemInstance*>(this)->FindDataInterfaceInstanceData(Interface); }
 	template<typename TDataType>
-	FORCEINLINE const TDataType* FindTypedDataInterfaceInstanceData(const UNiagaraDataInterface* Interface) const { return const_cast<FNiagaraSystemInstance*>(this)->FindTypedDataInterfaceInstanceData<TDataType>(Interface); }
+	inline const TDataType* FindTypedDataInterfaceInstanceData(const UNiagaraDataInterface* Interface) const { return const_cast<FNiagaraSystemInstance*>(this)->FindTypedDataInterfaceInstanceData<TDataType>(Interface); }
 
-	FORCEINLINE const FNiagaraPerInstanceDIFuncInfo& GetPerInstanceDIFunction(ENiagaraSystemSimulationScript ScriptType, int32 FuncIndex)const { return PerInstanceDIFunctions[(int32)ScriptType][FuncIndex]; }
+	inline const FNiagaraPerInstanceDIFuncInfo& GetPerInstanceDIFunction(ENiagaraSystemSimulationScript ScriptType, int32 FuncIndex)const { return PerInstanceDIFunctions[(int32)ScriptType][FuncIndex]; }
 
 	NIAGARA_API void EvaluateBoundFunction(FName FunctionName, bool& UsedOnCpu, bool& UsedOnGpu) const;
 
@@ -366,15 +366,15 @@ public:
 	NIAGARA_API bool UsesCollection(const UNiagaraParameterCollection* Collection) const;
 #endif
 
-	FORCEINLINE bool IsPendingSpawn() const { return (SystemInstanceState == ENiagaraSystemInstanceState::PendingSpawn) || (SystemInstanceState == ENiagaraSystemInstanceState::PendingSpawnPaused); }
+	inline bool IsPendingSpawn() const { return (SystemInstanceState == ENiagaraSystemInstanceState::PendingSpawn) || (SystemInstanceState == ENiagaraSystemInstanceState::PendingSpawnPaused); }
 
-	FORCEINLINE float GetAge() const { return Age; }
-	FORCEINLINE int32 GetTickCount() const { return TickCount; }
+	inline float GetAge() const { return Age; }
+	inline int32 GetTickCount() const { return TickCount; }
 
-	FORCEINLINE float GetLastRenderTime() const { return LastRenderTime; }
-	FORCEINLINE void SetLastRenderTime(float TimeSeconds) { LastRenderTime = TimeSeconds; }
+	inline float GetLastRenderTime() const { return LastRenderTime; }
+	inline void SetLastRenderTime(float TimeSeconds) { LastRenderTime = TimeSeconds; }
 	
-	FORCEINLINE TSharedPtr<FNiagaraSystemSimulation, ESPMode::ThreadSafe> GetSystemSimulation()const
+	inline TSharedPtr<FNiagaraSystemSimulation, ESPMode::ThreadSafe> GetSystemSimulation()const
 	{
 		return SystemSimulation; 
 	}
@@ -433,8 +433,8 @@ public:
 	NIAGARA_API FNiagaraDataSet* GetEventDataSet(FName EmitterName, FName EventName) const;
 	NIAGARA_API void ClearEventDataSets();
 
-	FORCEINLINE void SetLODDistance(float InLODDistance, float InMaxLODDistance, bool bOverride);
-	FORCEINLINE void ClearLODDistance();
+	inline void SetLODDistance(float InLODDistance, float InMaxLODDistance, bool bOverride);
+	inline void ClearLODDistance();
 
 	NIAGARA_API const FString& GetCrashReporterTag()const;
 
@@ -452,9 +452,9 @@ public:
 	NIAGARA_API ETickingGroup CalculateTickGroup() const;
 
 	/** Gets the current world transform of the system */
-	FORCEINLINE const FTransform& GetWorldTransform() const { return WorldTransform; }
+	inline const FTransform& GetWorldTransform() const { return WorldTransform; }
 	/** Sets the world transform */
-	FORCEINLINE void SetWorldTransform(const FTransform& InTransform) { WorldTransform = InTransform; }
+	inline void SetWorldTransform(const FTransform& InTransform) { WorldTransform = InTransform; }
 
 	int32 GetSystemInstanceIndex() const 
 	{
@@ -465,7 +465,7 @@ public:
 	The significant index for this component. i.e. this is the Nth most significant instance of it's system in the scene.
 	Passed to the script to allow us to scale down internally for less significant systems instances.
 */
-	FORCEINLINE void SetSystemSignificanceIndex(int32 InIndex) { SignificanceIndex = InIndex; }
+	inline void SetSystemSignificanceIndex(int32 InIndex) { SignificanceIndex = InIndex; }
 
 	/** Calculates the distance to use for distance based LODing / culling. */
 	NIAGARA_API float GetLODDistance();
@@ -725,7 +725,7 @@ public:
 	float			SystemState_CurrentLoopAgeEnd = 0.0f;
 };
 
-FORCEINLINE void FNiagaraSystemInstance::SetLODDistance(float InLODDistance, float InMaxLODDistance, bool bOverride)
+inline void FNiagaraSystemInstance::SetLODDistance(float InLODDistance, float InMaxLODDistance, bool bOverride)
 {
 	bLODDistanceIsOverridden = bOverride;
 	bLODDistanceIsValid = true;
@@ -733,7 +733,7 @@ FORCEINLINE void FNiagaraSystemInstance::SetLODDistance(float InLODDistance, flo
 	MaxLODDistance = InMaxLODDistance;
 }
 
-FORCEINLINE void FNiagaraSystemInstance::ClearLODDistance()
+inline void FNiagaraSystemInstance::ClearLODDistance()
 {
 	bLODDistanceIsOverridden = false;
 	bLODDistanceIsValid = false;
diff --git a/Engine/Plugins/FX/Niagara/Source/Niagara/Public/NiagaraSystemInstanceController.h b/Engine/Plugins/FX/Niagara/Source/Niagara/Public/NiagaraSystemInstanceController.h
--- a/Engine/Plugins/FX/Niagara/Source/Niagara/Public/NiagaraSystemInstanceController.h
+++ b/Engine/Plugins/FX/Niagara/Source/Niagara/Public/NiagaraSystemInstanceController.h
@@ -16,19 +16,19 @@
 #if NIAGARA_SYSTEM_INSTANCE_CONTROLLER_ASYNC
 #define NIAGARA_SYSTEM_INSTANCE_CONTROLLER_DEFFERED_METHOD(MethodName) \
 	template <typename... ArgTypes> \
-	FORCEINLINE void MethodName ## _Deferred(ArgTypes... Args) { \
+	inline void MethodName ## _Deferred(ArgTypes... Args) { \
 		using ThisClass = typename TRemovePointer<decltype(this)>::Type; \
 		WorldManager->GetDeferredMethodQueue().Enqueue(FNiagaraDeferredMethodQueue::FMethod::CreateThreadSafeSP(this, &ThisClass::MethodName, Forward<ArgTypes>(Args)...)); \
 	}
 #else
 #define NIAGARA_SYSTEM_INSTANCE_CONTROLLER_DEFFERED_METHOD(MethodName) \
-	template <typename... ArgTypes> FORCEINLINE void MethodName ## _Deferred(ArgTypes... Args) { this->MethodName(Forward<ArgTypes>(Args)...); }
+	template <typename... ArgTypes> inline void MethodName ## _Deferred(ArgTypes... Args) { this->MethodName(Forward<ArgTypes>(Args)...); }
 #endif
 
 /** Used to expose FNiagaraSystemInstance methods without actually providing the interface */
 #define NIAGARA_SYSTEM_INSTANCE_CONTROLLER_SHIM(MethodName, FuncMod) \
 	template <typename... ArgTypes> \
-	FORCEINLINE auto MethodName(ArgTypes... Args) FuncMod { ensure(IsValid()); return SystemInstance->MethodName(Forward<ArgTypes>(Args)...); }
+	inline auto MethodName(ArgTypes... Args) FuncMod { ensure(IsValid()); return SystemInstance->MethodName(Forward<ArgTypes>(Args)...); }
 
 class UMaterialInterface;
 class UTextureRenderTarget;
@@ -70,7 +70,7 @@ public:
 	/** Deactivates the underlying system and queues it up for destruction. After calling, this interface is unusable unless you call Initialize again */
 	void Release();
 
-	FORCEINLINE bool IsValid() const { return SystemInstance.IsValid(); }
+	inline bool IsValid() const { return SystemInstance.IsValid(); }
 	bool HasValidSimulation() const;
 
 	// Render data operations
diff --git a/Engine/Plugins/FX/Niagara/Source/Niagara/Public/NiagaraSystemRenderData.h b/Engine/Plugins/FX/Niagara/Source/Niagara/Public/NiagaraSystemRenderData.h
--- a/Engine/Plugins/FX/Niagara/Source/Niagara/Public/NiagaraSystemRenderData.h
+++ b/Engine/Plugins/FX/Niagara/Source/Niagara/Public/NiagaraSystemRenderData.h
@@ -81,13 +81,13 @@ public:
 	void RecacheRenderers(const FNiagaraSystemInstance& SystemInstance, const FNiagaraSystemInstanceController& Controller);
 
 	void SetRenderingEnabled_GT(bool bInEnabled);
-	FORCEINLINE bool IsRenderingEnabled_GT() const { return bRenderingEnabled_GT && EmitterRenderers_GT.Num() > 0; }
-	FORCEINLINE bool IsRenderingEnabled_RT() const { return bRenderingEnabled_RT && EmitterRenderers_RT.Num() > 0; }
+	inline bool IsRenderingEnabled_GT() const { return bRenderingEnabled_GT && EmitterRenderers_GT.Num() > 0; }
+	inline bool IsRenderingEnabled_RT() const { return bRenderingEnabled_RT && EmitterRenderers_RT.Num() > 0; }
 
-	FORCEINLINE bool HasAnyMotionBlurEnabled() const { return bAnyMotionBlurEnabled; }
-	FORCEINLINE bool HasAnyHeterogeneousVolumesEnabled() const { return bAnyHeterogeneousVolumesEnabled; }
+	inline bool HasAnyMotionBlurEnabled() const { return bAnyMotionBlurEnabled; }
+	inline bool HasAnyHeterogeneousVolumesEnabled() const { return bAnyHeterogeneousVolumesEnabled; }
 
-	FORCEINLINE int32 GetNumRenderers() const { return RendererDrawOrder.Num(); }
+	inline int32 GetNumRenderers() const { return RendererDrawOrder.Num(); }
 
 	FVector3f LWCRenderTile = FVector3f::ZeroVector;
 private:
diff --git a/Engine/Plugins/FX/Niagara/Source/Niagara/Public/NiagaraSystemSimulation.h b/Engine/Plugins/FX/Niagara/Source/Niagara/Public/NiagaraSystemSimulation.h
--- a/Engine/Plugins/FX/Niagara/Source/Niagara/Public/NiagaraSystemSimulation.h
+++ b/Engine/Plugins/FX/Niagara/Source/Niagara/Public/NiagaraSystemSimulation.h
@@ -123,7 +123,7 @@ struct FNiagaraParameterStoreToDataSetBinding
 		}
 	}
 
-	FORCEINLINE_DEBUGGABLE void DataSetToParameterStore(FNiagaraParameterStore& ParameterStore, FNiagaraDataSet& DataSet, int32 DataSetInstanceIndex)
+	inline void DataSetToParameterStore(FNiagaraParameterStore& ParameterStore, FNiagaraDataSet& DataSet, int32 DataSetInstanceIndex)
 	{
 #if NIAGARA_NAN_CHECKING
 		DataSet.CheckForNaNs();
@@ -161,7 +161,7 @@ struct FNiagaraParameterStoreToDataSetBinding
 		ParameterStore.OnParameterChange();
 	}
 
-	FORCEINLINE_DEBUGGABLE void ParameterStoreToDataSet(const FNiagaraParameterStore& ParameterStore, FNiagaraDataSet& DataSet, int32 DataSetInstanceIndex)
+	inline void ParameterStoreToDataSet(const FNiagaraParameterStore& ParameterStore, FNiagaraDataSet& DataSet, int32 DataSetInstanceIndex)
 	{
 		FNiagaraDataBuffer& CurrBuffer = DataSet.GetDestinationDataChecked();
 		const uint8* ParameterData = ParameterStore.GetParameterDataArray().GetData();
@@ -288,7 +288,7 @@ public:
 	void PauseInstance(FNiagaraSystemInstance* Instance);
 	void UnpauseInstance(FNiagaraSystemInstance* Instance);
 
-	FORCEINLINE UNiagaraSystem* GetSystem() const { return System; }
+	inline UNiagaraSystem* GetSystem() const { return System; }
 
 	UNiagaraParameterCollectionInstance* GetParameterCollectionInstance(UNiagaraParameterCollection* Collection);
 
@@ -315,11 +315,11 @@ public:
 
 	ETickingGroup GetTickGroup() const { return SystemTickGroup; }
 
-	FORCEINLINE FNiagaraGpuComputeDispatchInterface* GetDispatchInterface() const { return DispatchInterface; }
+	inline FNiagaraGpuComputeDispatchInterface* GetDispatchInterface() const { return DispatchInterface; }
 
 	ENiagaraGPUTickHandlingMode GetGPUTickHandlingMode() const;
 
-	FORCEINLINE UWorld* GetWorld()const{return World;}
+	inline UWorld* GetWorld()const{return World;}
 
 	const FNiagaraTickInfo& GetTickInfo() { return TickInfo; }
 
diff --git a/Engine/Plugins/FX/Niagara/Source/Niagara/Public/NiagaraTypes.h b/Engine/Plugins/FX/Niagara/Source/Niagara/Public/NiagaraTypes.h
--- a/Engine/Plugins/FX/Niagara/Source/Niagara/Public/NiagaraTypes.h
+++ b/Engine/Plugins/FX/Niagara/Source/Niagara/Public/NiagaraTypes.h
@@ -60,16 +60,16 @@ struct FNiagaraBool
 	bool GetValue() const { return Value != False; }
 
 	/** Sets this niagara bool's raw integer value directly using the special raw integer values expected by the VM and HLSL. */
-	FORCEINLINE void SetRawValue(int32 RawValue) { Value = RawValue; }
+	inline void SetRawValue(int32 RawValue) { Value = RawValue; }
 
 	/** Gets this niagara bools raw integer value expected by the VM and HLSL. */
-	FORCEINLINE int32 GetRawValue() const { return Value; }
+	inline int32 GetRawValue() const { return Value; }
 
 	bool IsValid() const { return Value == True || Value == False; }
 	
 	FNiagaraBool() : Value(False) {}
 	FNiagaraBool(bool bInValue) : Value(bInValue ? True : False) {}
-	FORCEINLINE operator bool() const { return GetValue(); }
+	inline operator bool() const { return GetValue(); }
 
 private:
 	UPROPERTY(EditAnywhere, Category = Parameters)// Must be either FNiagaraBool::True or FNiagaraBool::False.
@@ -83,12 +83,12 @@ struct FNiagaraPosition : public FVector3f
 
 	FNiagaraPosition() {}
 	
-	explicit FORCEINLINE FNiagaraPosition(EForceInit Init) : Super(Init) {}
+	explicit inline FNiagaraPosition(EForceInit Init) : Super(Init) {}
 
-	FORCEINLINE FNiagaraPosition(const float& X, const float& Y, const float& Z) : Super(X, Y, Z) {}
+	inline FNiagaraPosition(const float& X, const float& Y, const float& Z) : Super(X, Y, Z) {}
 
-	FORCEINLINE FNiagaraPosition(const FVector3f& Other) : Super(Other) {}
-	FORCEINLINE FNiagaraPosition(const FVector& Other) : Super(Other) {}
+	inline FNiagaraPosition(const FVector3f& Other) : Super(Other) {}
+	inline FNiagaraPosition(const FVector& Other) : Super(Other) {}
 };
 
 
@@ -229,7 +229,7 @@ struct FNiagaraAssetVersion
 
 	static NIAGARA_API FGuid CreateStableVersionGuid(UObject* Object);
 
-	friend FORCEINLINE uint32 GetTypeHash(const FNiagaraAssetVersion& Version)
+	friend inline uint32 GetTypeHash(const FNiagaraAssetVersion& Version)
 	{
 		return HashCombine(GetTypeHash(Version.MajorVersion), GetTypeHash(Version.MinorVersion));
 	}
@@ -382,7 +382,7 @@ struct FNiagaraID
 	bool operator!=(const FNiagaraID& Other)const { return !(*this == Other); }
 	bool operator<(const FNiagaraID& Other)const { return Index < Other.Index || (Index == Other.Index && AcquireTag < Other.AcquireTag); }
 
-	friend FORCEINLINE uint32 GetTypeHash(const FNiagaraID& ID)
+	friend inline uint32 GetTypeHash(const FNiagaraID& ID)
 	{
 		return HashCombine(GetTypeHash(ID.Index), GetTypeHash(ID.AcquireTag));
 	}
@@ -687,7 +687,7 @@ struct FNiagaraTypeDefinition
 public:
 
 	// Construct blank raw type definition 
-	FORCEINLINE FNiagaraTypeDefinition(UObject* Object)
+	inline FNiagaraTypeDefinition(UObject* Object)
 		: ClassStructOrEnum(Object)
 #if WITH_EDITORONLY_DATA
 		, Struct_DEPRECATED(nullptr), Enum_DEPRECATED(nullptr)
@@ -713,7 +713,7 @@ public:
 		GeneratePathNameHash();
 	}
 
-	FORCEINLINE FNiagaraTypeDefinition(UClass *ClassDef)
+	inline FNiagaraTypeDefinition(UClass *ClassDef)
 		: ClassStructOrEnum(ClassDef), UnderlyingType(UT_Class)
 #if WITH_EDITORONLY_DATA
 		, Struct_DEPRECATED(nullptr), Enum_DEPRECATED(nullptr)
@@ -724,7 +724,7 @@ public:
 		GeneratePathNameHash();
 	}
 
-	FORCEINLINE FNiagaraTypeDefinition(UEnum *EnumDef)
+	inline FNiagaraTypeDefinition(UEnum *EnumDef)
 		: ClassStructOrEnum(EnumDef), UnderlyingType(UT_Enum)
 #if WITH_EDITORONLY_DATA
 		, Struct_DEPRECATED(nullptr), Enum_DEPRECATED(nullptr)
@@ -736,7 +736,7 @@ public:
 	}
 
 	enum class EAllowUnfriendlyStruct : uint8 { Deny, Allow };
-	FORCEINLINE FNiagaraTypeDefinition(UScriptStruct* StructDef, EAllowUnfriendlyStruct AllowUnfriendlyStruct)
+	inline FNiagaraTypeDefinition(UScriptStruct* StructDef, EAllowUnfriendlyStruct AllowUnfriendlyStruct)
 		: ClassStructOrEnum(StructDef), UnderlyingType(UT_Struct)
 #if WITH_EDITORONLY_DATA
 		, Struct_DEPRECATED(nullptr), Enum_DEPRECATED(nullptr)
@@ -748,9 +748,9 @@ public:
 		GeneratePathNameHash();
 	}
 	
-	FORCEINLINE FNiagaraTypeDefinition(UScriptStruct* StructDef) : FNiagaraTypeDefinition(StructDef, EAllowUnfriendlyStruct::Deny) {}
+	inline FNiagaraTypeDefinition(UScriptStruct* StructDef) : FNiagaraTypeDefinition(StructDef, EAllowUnfriendlyStruct::Deny) {}
 
-	FORCEINLINE FNiagaraTypeDefinition(FProperty* Property, EAllowUnfriendlyStruct AllowUnfriendlyStruct)
+	inline FNiagaraTypeDefinition(FProperty* Property, EAllowUnfriendlyStruct AllowUnfriendlyStruct)
 	{
 		if (Property->IsA(FFloatProperty::StaticClass())) { *this = FNiagaraTypeDefinition::GetFloatDef(); }
 		else if (Property->IsA(FUInt16Property::StaticClass())) { *this = FNiagaraTypeDefinition::GetHalfDef(); }
@@ -764,24 +764,24 @@ public:
 	}
 
 	// Construct a blank raw type definition
-	FORCEINLINE FNiagaraTypeDefinition()
+	inline FNiagaraTypeDefinition()
 		: ClassStructOrEnum(nullptr), UnderlyingType(UT_None)
 #if WITH_EDITORONLY_DATA
 		, Struct_DEPRECATED(nullptr), Enum_DEPRECATED(nullptr)
 #endif
 	{}
 
-	FORCEINLINE bool operator !=(const FNiagaraTypeDefinition &Other) const
+	inline bool operator !=(const FNiagaraTypeDefinition &Other) const
 	{
 		return !(*this == Other);
 	}
 
-	FORCEINLINE bool operator == (const FNiagaraTypeDefinition &Other) const
+	inline bool operator == (const FNiagaraTypeDefinition &Other) const
 	{
 		return ClassStructOrEnum == Other.ClassStructOrEnum && UnderlyingType == Other.UnderlyingType && Flags == Other.Flags;
 	}
 
-	FORCEINLINE bool IsSameBaseDefinition(const FNiagaraTypeDefinition& Other) const
+	inline bool IsSameBaseDefinition(const FNiagaraTypeDefinition& Other) const
 	{
 		return ClassStructOrEnum == Other.ClassStructOrEnum && UnderlyingType == Other.UnderlyingType;
 	}
@@ -886,7 +886,7 @@ public:
 
 	NIAGARA_API bool IsDataInterface() const;
 
-	FORCEINLINE bool IsUObject() const
+	inline bool IsUObject() const
 	{
 		const UStruct* ClassStruct = GetStruct();
 		return ClassStruct ? ClassStruct->IsChildOf<UObject>() : false;
@@ -1190,7 +1190,7 @@ private:
 
 	void NIAGARA_API GeneratePathNameHash();
 
-	friend FORCEINLINE uint32 GetTypeHash(const FNiagaraTypeDefinition& Type)
+	friend inline uint32 GetTypeHash(const FNiagaraTypeDefinition& Type)
 	{
 		return HashCombine(HashCombine(Type.UnderlyingTypePathNameHash, GetTypeHash(Type.UnderlyingType)), GetTypeHash(Type.GetFlags()));
 	}
@@ -1274,7 +1274,7 @@ struct FNiagaraVariableBase
 {
 	GENERATED_USTRUCT_BODY()
 
-	FORCEINLINE FNiagaraVariableBase()
+	inline FNiagaraVariableBase()
 		: Name(NAME_None)
 		, TypeDefHandle(FNiagaraTypeDefinition::GetVec4Def())
 #if WITH_EDITORONLY_DATA
@@ -1282,7 +1282,7 @@ struct FNiagaraVariableBase
 #endif
 		{}
 
-	FORCEINLINE FNiagaraVariableBase(const FNiagaraTypeDefinition& InType, const FName& InName)
+	inline FNiagaraVariableBase(const FNiagaraTypeDefinition& InType, const FName& InName)
 		: Name(InName)
 		, TypeDefHandle(InType)
 #if WITH_EDITORONLY_DATA
@@ -1290,7 +1290,7 @@ struct FNiagaraVariableBase
 #endif
 		{}
 
-	FORCEINLINE FNiagaraVariableBase(const FNiagaraVariableCommonReference& VariableRef)
+	inline FNiagaraVariableBase(const FNiagaraVariableCommonReference& VariableRef)
 		: Name(VariableRef.Name)
 		, TypeDefHandle(FNiagaraTypeDefinition(VariableRef.UnderlyingType))
 	{}
@@ -1313,11 +1313,11 @@ struct FNiagaraVariableBase
 		return Name == Other.Name && (TypeDefHandle == Other.TypeDefHandle || (bAllowAssignableTypes && FNiagaraTypeDefinition::TypesAreAssignable(*TypeDefHandle, *Other.TypeDefHandle)));
 	}
 	
-	FORCEINLINE void SetName(FName InName)
+	inline void SetName(FName InName)
 	{
 		Name = InName;
 	}
-	FORCEINLINE const FName& GetName() const
+	inline const FName& GetName() const
 	{
 		return Name;
 	}
@@ -1333,8 +1333,8 @@ struct FNiagaraVariableBase
 		return *TypeDefHandle;
 	}
 
-	FORCEINLINE bool IsDataInterface()const { return GetType().IsDataInterface(); }
-	FORCEINLINE bool IsUObject()const { return GetType().IsUObject(); }
+	inline bool IsDataInterface()const { return GetType().IsDataInterface(); }
+	inline bool IsUObject()const { return GetType().IsUObject(); }
 
 	int32 GetSizeInBytes() const
 	{
@@ -1381,14 +1381,14 @@ struct FNiagaraVariableBase
 		}
 	}
 
-	FORCEINLINE bool IsInNameSpace(const FStringView& Namespace) const
+	inline bool IsInNameSpace(const FStringView& Namespace) const
 	{
 		return IsInNameSpace(Namespace, Name);
 	}
 
 #if WITH_EDITORONLY_DATA
 	// This method should not be used at runtime as we have pre-defined strings in FNiagaraConstants for runtime cases
-	FORCEINLINE bool IsInNameSpace(const FName& Namespace) const
+	inline bool IsInNameSpace(const FName& Namespace) const
 	{
 		FNameBuilder NamespaceString;
 		Namespace.ToString(NamespaceString);
@@ -1416,7 +1416,7 @@ protected:
 	FNiagaraTypeDefinition TypeDef_DEPRECATED;
 #endif
 
-	friend FORCEINLINE uint32 GetTypeHash(const FNiagaraVariableBase& Var)
+	friend inline uint32 GetTypeHash(const FNiagaraVariableBase& Var)
 	{
 		return HashCombine(GetTypeHash(Var.GetType()), GetTypeHash(Var.GetName()));
 	}
@@ -1449,7 +1449,7 @@ struct FNiagaraVariable : public FNiagaraVariableBase
 	{
 	}
 
-	FORCEINLINE FNiagaraVariable(const FNiagaraTypeDefinition& InType, const FName& InName)
+	inline FNiagaraVariable(const FNiagaraTypeDefinition& InType, const FName& InName)
 		: FNiagaraVariableBase(InType, InName)
 	{
 	}
diff --git a/Engine/Plugins/FX/Niagara/Source/Niagara/Public/NiagaraUserRedirectionParameterStore.h b/Engine/Plugins/FX/Niagara/Source/Niagara/Public/NiagaraUserRedirectionParameterStore.h
--- a/Engine/Plugins/FX/Niagara/Source/Niagara/Public/NiagaraUserRedirectionParameterStore.h
+++ b/Engine/Plugins/FX/Niagara/Source/Niagara/Public/NiagaraUserRedirectionParameterStore.h
@@ -29,10 +29,10 @@ public:
 	NIAGARA_API bool RedirectUserVariable(FNiagaraVariableBase& UserVar) const;
 
 	/** Get the list of FNiagaraVariables that are exposed to the user. Note that the values will be stale and are not to be trusted directly. Get the Values using the offset specified by IndexOf or GetParameterValue.*/
-	FORCEINLINE void GetUserParameters(TArray<FNiagaraVariable>& OutParameters) const { return UserParameterRedirects.GenerateKeyArray(OutParameters); }
+	inline void GetUserParameters(TArray<FNiagaraVariable>& OutParameters) const { return UserParameterRedirects.GenerateKeyArray(OutParameters); }
 
 	// ~ Begin FNiagaraParameterStore overrides
-	FORCEINLINE_DEBUGGABLE virtual const int32* FindParameterOffset(const FNiagaraVariableBase& Parameter, bool IgnoreType = false) const override
+	inline virtual const int32* FindParameterOffset(const FNiagaraVariableBase& Parameter, bool IgnoreType = false) const override
 	{
 		if (IgnoreType)
 		{
@@ -52,7 +52,7 @@ public:
 		}
 	}
 	
-	FORCEINLINE_DEBUGGABLE virtual const FNiagaraVariableWithOffset* FindParameterVariable(const FNiagaraVariable& Parameter, bool IgnoreType = false) const override
+	inline virtual const FNiagaraVariableWithOffset* FindParameterVariable(const FNiagaraVariable& Parameter, bool IgnoreType = false) const override
 	{
 		if (IgnoreType)
 		{
diff --git a/Engine/Plugins/FX/Niagara/Source/Niagara/Public/NiagaraWorldManager.h b/Engine/Plugins/FX/Niagara/Source/Niagara/Public/NiagaraWorldManager.h
--- a/Engine/Plugins/FX/Niagara/Source/Niagara/Public/NiagaraWorldManager.h
+++ b/Engine/Plugins/FX/Niagara/Source/Niagara/Public/NiagaraWorldManager.h
@@ -198,7 +198,7 @@ public:
 	void DumpDetails(FOutputDevice& Ar, bool bDetailed);
 	
 	UWorld* GetWorld();
-	FORCEINLINE UWorld* GetWorld()const { return World; }
+	inline UWorld* GetWorld()const { return World; }
 
 	//Various helper functions for scalability culling.
 	
@@ -212,7 +212,7 @@ public:
 	void CalculateScalabilityState(UNiagaraSystem* System, const FNiagaraSystemScalabilitySettings& ScalabilitySettings, UNiagaraEffectType* EffectType, UNiagaraComponent* Component, bool bIsPreCull, float WorstGlobalBudgetUse, FNiagaraScalabilityState& OutState);
 	void CalculateScalabilityState(UNiagaraSystem* System, const FNiagaraSystemScalabilitySettings& ScalabilitySettings, UNiagaraEffectType* EffectType, FVector Location, bool bIsPreCull, float WorstGlobalBudgetUse, FNiagaraScalabilityState& OutState);
 
-	/*FORCEINLINE_DEBUGGABLE*/ void SortedSignificanceCull(UNiagaraEffectType* EffectType, UNiagaraComponent* Component, const FNiagaraSystemScalabilitySettings& ScalabilitySettings, float Significance, int32& EffectTypeInstCount, uint16& SystemInstCount, FNiagaraScalabilityState& OutState);
+	/*inline*/ void SortedSignificanceCull(UNiagaraEffectType* EffectType, UNiagaraComponent* Component, const FNiagaraSystemScalabilitySettings& ScalabilitySettings, float Significance, int32& EffectTypeInstCount, uint16& SystemInstCount, FNiagaraScalabilityState& OutState);
 
 #if DEBUG_SCALABILITY_STATE
 	void DumpScalabilityState();
diff --git a/Engine/Plugins/FX/Niagara/Source/NiagaraCore/Public/NiagaraCompileHash.h b/Engine/Plugins/FX/Niagara/Source/NiagaraCore/Public/NiagaraCompileHash.h
--- a/Engine/Plugins/FX/Niagara/Source/NiagaraCore/Public/NiagaraCompileHash.h
+++ b/Engine/Plugins/FX/Niagara/Source/NiagaraCore/Public/NiagaraCompileHash.h
@@ -66,7 +66,7 @@ private:
 	TArray<uint8> DataHash;
 };
 
-FORCEINLINE uint32 GetTypeHash(const FNiagaraCompileHash& Hash)
+inline uint32 GetTypeHash(const FNiagaraCompileHash& Hash)
 {
 	return Hash.GetTypeHash();
 }
diff --git a/Engine/Plugins/FX/Niagara/Source/NiagaraEditor/Public/Config/NiagaraFavoriteActionsConfig.h b/Engine/Plugins/FX/Niagara/Source/NiagaraEditor/Public/Config/NiagaraFavoriteActionsConfig.h
--- a/Engine/Plugins/FX/Niagara/Source/NiagaraEditor/Public/Config/NiagaraFavoriteActionsConfig.h
+++ b/Engine/Plugins/FX/Niagara/Source/NiagaraEditor/Public/Config/NiagaraFavoriteActionsConfig.h
@@ -56,7 +56,7 @@ struct FNiagaraActionIdentifier
 	}
 };
 
-FORCEINLINE uint32 GetTypeHash(const FNiagaraActionIdentifier& Identity)
+inline uint32 GetTypeHash(const FNiagaraActionIdentifier& Identity)
 {
 	uint32 Hash = 0;
 	
diff --git a/Engine/Plugins/FX/Niagara/Source/NiagaraEditor/Public/NiagaraEditorModule.h b/Engine/Plugins/FX/Niagara/Source/NiagaraEditor/Public/NiagaraEditorModule.h
--- a/Engine/Plugins/FX/Niagara/Source/NiagaraEditor/Public/NiagaraEditorModule.h
+++ b/Engine/Plugins/FX/Niagara/Source/NiagaraEditor/Public/NiagaraEditorModule.h
@@ -159,7 +159,7 @@ private:
 	TMap<UClass*, TSharedPtr<FNiagaraStackObjectPropertyCustomization>> Customizations;
 };
 
-FORCEINLINE uint32 GetTypeHash(const FReservedParameter& ReservedParameter) { return GetTypeHash(ReservedParameter.GetParameter().GetName()); };
+inline uint32 GetTypeHash(const FReservedParameter& ReservedParameter) { return GetTypeHash(ReservedParameter.GetParameter().GetName()); };
 
 /** Niagara Editor module */
 class FNiagaraEditorModule : public IModuleInterface,
@@ -279,7 +279,7 @@ public:
 		EnqueueObjectForDeferredDestructionInternal(ObjectInContainer);
 	}
 
-	FORCEINLINE INiagaraStackObjectIssueGenerator* FindStackObjectIssueGenerator(FName StructName)
+	inline INiagaraStackObjectIssueGenerator* FindStackObjectIssueGenerator(FName StructName)
 	{
 		if (INiagaraStackObjectIssueGenerator** FoundGenerator = StackIssueGenerators.Find(StructName))
 		{
diff --git a/Engine/Plugins/FX/Niagara/Source/NiagaraEditor/Public/NiagaraEditorSettings.h b/Engine/Plugins/FX/Niagara/Source/NiagaraEditor/Public/NiagaraEditorSettings.h
--- a/Engine/Plugins/FX/Niagara/Source/NiagaraEditor/Public/NiagaraEditorSettings.h
+++ b/Engine/Plugins/FX/Niagara/Source/NiagaraEditor/Public/NiagaraEditorSettings.h
@@ -330,7 +330,7 @@ struct FNiagaraCurveEditorSharedSettings
 	bool bCurveOutputSnapEnabled = false;
 };
 
-FORCEINLINE uint32 GetTypeHash(const FNiagaraNamespaceMetadata& NamespaceMetaData)
+inline uint32 GetTypeHash(const FNiagaraNamespaceMetadata& NamespaceMetaData)
 {
 	return GetTypeHash(NamespaceMetaData.GetGuid());
 }
diff --git a/Engine/Plugins/FX/Niagara/Source/NiagaraEditor/Public/NiagaraEditorUtilities.h b/Engine/Plugins/FX/Niagara/Source/NiagaraEditor/Public/NiagaraEditorUtilities.h
--- a/Engine/Plugins/FX/Niagara/Source/NiagaraEditor/Public/NiagaraEditorUtilities.h
+++ b/Engine/Plugins/FX/Niagara/Source/NiagaraEditor/Public/NiagaraEditorUtilities.h
@@ -621,7 +621,7 @@ private:
 		int32 Index;
 	};
 
-	friend FORCEINLINE uint32 GetTypeHash(const FEnumIndexPair& EnumIndexPair)
+	friend inline uint32 GetTypeHash(const FEnumIndexPair& EnumIndexPair)
 	{
 		return HashCombineFast(GetTypeHash(EnumIndexPair.Enum), GetTypeHash(EnumIndexPair.Index));
 	}
diff --git a/Engine/Plugins/FX/Niagara/Source/NiagaraEditor/Public/NiagaraGraph.h b/Engine/Plugins/FX/Niagara/Source/NiagaraEditor/Public/NiagaraGraph.h
--- a/Engine/Plugins/FX/Niagara/Source/NiagaraEditor/Public/NiagaraGraph.h
+++ b/Engine/Plugins/FX/Niagara/Source/NiagaraEditor/Public/NiagaraGraph.h
@@ -53,7 +53,7 @@ public:
 	UPROPERTY()
 	bool bIsUserFacing = true;
 
-	FORCEINLINE bool operator==(const FNiagaraGraphParameterReference& Other)const
+	inline bool operator==(const FNiagaraGraphParameterReference& Other)const
 	{
 		return Other.Key == Key && Other.Value == Value;
 	}
diff --git a/Engine/Plugins/FX/Niagara/Source/NiagaraEditor/Public/ViewModels/Stack/NiagaraParameterHandle.h b/Engine/Plugins/FX/Niagara/Source/NiagaraEditor/Public/ViewModels/Stack/NiagaraParameterHandle.h
--- a/Engine/Plugins/FX/Niagara/Source/NiagaraEditor/Public/ViewModels/Stack/NiagaraParameterHandle.h
+++ b/Engine/Plugins/FX/Niagara/Source/NiagaraEditor/Public/ViewModels/Stack/NiagaraParameterHandle.h
@@ -77,7 +77,7 @@ private:
 	mutable TArray<FName> HandlePartsCache;
 };
 
-FORCEINLINE uint32 GetTypeHash(const FNiagaraParameterHandle& Var)
+inline uint32 GetTypeHash(const FNiagaraParameterHandle& Var)
 {
 	return GetTypeHash(Var.GetParameterHandleString());
 }
diff --git a/Engine/Plugins/FX/Niagara/Source/NiagaraEditor/Public/ViewModels/Stack/NiagaraStackFunctionInput.h b/Engine/Plugins/FX/Niagara/Source/NiagaraEditor/Public/ViewModels/Stack/NiagaraStackFunctionInput.h
--- a/Engine/Plugins/FX/Niagara/Source/NiagaraEditor/Public/ViewModels/Stack/NiagaraStackFunctionInput.h
+++ b/Engine/Plugins/FX/Niagara/Source/NiagaraEditor/Public/ViewModels/Stack/NiagaraStackFunctionInput.h
@@ -592,7 +592,7 @@ private:
 	TOptional<FGuid> LastFunctionGraphChangeId;
 };
 
-FORCEINLINE uint32 GetTypeHash(const UNiagaraStackFunctionInput::FNiagaraAvailableParameterInfo& ParameterInfo)
+inline uint32 GetTypeHash(const UNiagaraStackFunctionInput::FNiagaraAvailableParameterInfo& ParameterInfo)
 {
 	return GetTypeHash(ParameterInfo.Variable);
 }
