- IWYU fix

#rnx

diff --git a/Engine/Plugins/FX/Niagara/Source/Niagara/Internal/DataInterface/NiagaraDataInterfaceArrayDistributionInt.h b/Engine/Plugins/FX/Niagara/Source/Niagara/Internal/DataInterface/NiagaraDataInterfaceArrayDistributionInt.h
--- a/Engine/Plugins/FX/Niagara/Source/Niagara/Internal/DataInterface/NiagaraDataInterfaceArrayDistributionInt.h
+++ b/Engine/Plugins/FX/Niagara/Source/Niagara/Internal/DataInterface/NiagaraDataInterfaceArrayDistributionInt.h
@@ -4,6 +4,7 @@
 
 #include "NiagaraDataInterfaceArray.h"
 #include "NiagaraDataInterfaceArrayImpl.h"
+#include "Misc/TransactionallySafeRWLock.h"
 #include "NiagaraDataInterfaceArrayDistributionInt.generated.h"
 
 USTRUCT(BlueprintType, meta = (DisplayName = "Weighted Distribution Int Entry"))
