- Use UE::TRWScopeLock vs FRWScopeLock to handle FTransactionallySafeRWLock correctly

#rnx

diff --git a/Engine/Plugins/FX/Niagara/Source/Niagara/Private/DataInterface/NiagaraDataInterfaceArrayDistributionInt.cpp b/Engine/Plugins/FX/Niagara/Source/Niagara/Private/DataInterface/NiagaraDataInterfaceArrayDistributionInt.cpp
--- a/Engine/Plugins/FX/Niagara/Source/Niagara/Private/DataInterface/NiagaraDataInterfaceArrayDistributionInt.cpp
+++ b/Engine/Plugins/FX/Niagara/Source/Niagara/Private/DataInterface/NiagaraDataInterfaceArrayDistributionInt.cpp
@@ -7,6 +7,7 @@
 #include "NiagaraFunctionLibrary.h"
 #include "NiagaraShaderParametersBuilder.h"
 
+#include "Misc/ScopeRWLock.h"
 #include "RenderGraphBuilder.h"
 #include "WeightedRandomSampler.h"
 
@@ -78,7 +79,7 @@ namespace NDIArrayDistributionIntPrivate
 				return;
 			}
 
-			FRWScopeLock ScopeLock(OwnerDI->BuiltTableDataGuard, SLT_ReadOnly);
+			UE::TRWScopeLock ScopeLock(OwnerDI->BuiltTableDataGuard, SLT_ReadOnly);
 
 			bNeedsGpuDataUpdate = false;
 
@@ -318,7 +319,7 @@ void UNiagaraDataInterfaceArrayDistributionInt::BuildTableData()
 		EncodeBuiltTableEntry(NewBuiltTableData, 0, 0.0f, 0, 0, 0.0f);
 	}
 
-	FRWScopeLock ScopeLock(BuiltTableDataGuard, SLT_Write);
+	UE::TRWScopeLock ScopeLock(BuiltTableDataGuard, SLT_Write);
 	Swap(BuiltTableData, NewBuiltTableData);
 
 	GetProxyAs<FNDIProxy>()->bNeedsGpuDataUpdate = true;
@@ -328,7 +329,7 @@ int32 UNiagaraDataInterfaceArrayDistributionInt::GetBuiltTableLength() const
 {
 	using namespace NDIArrayDistributionIntPrivate;
 
-	FRWScopeLock ScopeLock(BuiltTableDataGuard, SLT_ReadOnly);
+	UE::TRWScopeLock ScopeLock(BuiltTableDataGuard, SLT_ReadOnly);
 	return BuiltTableData.Num() / GetBuiltTableEntrySize();
 }
 
@@ -376,7 +377,7 @@ void UNiagaraDataInterfaceArrayDistributionInt::VMGet(FVectorVMExternalFunctionC
 	FNDIOutputParam<int32>	OutIndex(Context);
 	FNDIOutputParam<float>	OutWeight(Context);
 
-	FRWScopeLock ScopeLock(BuiltTableDataGuard, SLT_ReadOnly);
+	UE::TRWScopeLock ScopeLock(BuiltTableDataGuard, SLT_ReadOnly);
 	const int32 LengthMinusOne = GetBuiltTableLength() - 1;
 	for (int32 i = 0; i < Context.GetNumInstances(); ++i)
 	{
@@ -399,7 +400,7 @@ void UNiagaraDataInterfaceArrayDistributionInt::VMGetRandomValue(FVectorVMExtern
 	FNDIRandomHelper		RandHelper(Context);
 	FNDIOutputParam<int32>	OutValue(Context);
 
-	FRWScopeLock ScopeLock(BuiltTableDataGuard, SLT_ReadOnly);
+	UE::TRWScopeLock ScopeLock(BuiltTableDataGuard, SLT_ReadOnly);
 	const int32 LengthMinusOne = GetBuiltTableLength() - 1;
 	for (int32 i = 0; i < Context.GetNumInstances(); ++i)
 	{
