Remove some lingering delegate bindings in OnShutdown().
Practically this should have no affect but for tidyness/safety sake we should ensure these are cleaned up.

diff --git a/Engine/Plugins/FX/Niagara/Source/Niagara/Private/NiagaraWorldManager.cpp b/Engine/Plugins/FX/Niagara/Source/Niagara/Private/NiagaraWorldManager.cpp
--- a/Engine/Plugins/FX/Niagara/Source/Niagara/Private/NiagaraWorldManager.cpp
+++ b/Engine/Plugins/FX/Niagara/Source/Niagara/Private/NiagaraWorldManager.cpp
@@ -504,9 +504,11 @@ void FNiagaraWorldManager::OnShutdown()
 	FWorldDelegates::OnPostWorldCleanup.Remove(OnPostWorldCleanupHandle);
 	FWorldDelegates::OnPreWorldFinishDestroy.Remove(OnPreWorldFinishDestroyHandle);
 	FWorldDelegates::OnWorldBeginTearDown.Remove(OnWorldBeginTearDownHandle);
+	FWorldDelegates::OnWorldTickStart.Remove(TickWorldStartHandle);
+	FWorldDelegates::OnWorldTickEnd.Remove(TickWorldEndHandle);
+	FWorldDelegates::OnWorldPreActorTick.Remove(OnWorldPreActorTickHandle);
 	FWorldDelegates::OnWorldPostActorTick.Remove(TickWorldHandle);
 	FWorldDelegates::OnWorldPreSendAllEndOfFrameUpdates.Remove(OnWorldPreSendAllEndOfFrameUpdatesHandle);
-
 	FCoreUObjectDelegates::GetPreGarbageCollectDelegate().Remove(PreGCHandle);
 	FCoreUObjectDelegates::PostReachabilityAnalysis.Remove(PostReachabilityAnalysisHandle);
 	FCoreUObjectDelegates::GetPostGarbageCollect().Remove(PostGCHandle);
