- Renderer memory savings by reducing FNiagaraVariableAttributeBinding
- Sprites 1992 -> 1320 = 672
- Mesh    1688 -> 1160 = 528
- Ribbon  1888 -> 1312 = 576
- Light    800 ->  560 = 240

diff --git a/Engine/Plugins/FX/Niagara/Source/Niagara/Classes/NiagaraConstants.h b/Engine/Plugins/FX/Niagara/Source/Niagara/Classes/NiagaraConstants.h
--- a/Engine/Plugins/FX/Niagara/Source/Niagara/Classes/NiagaraConstants.h
+++ b/Engine/Plugins/FX/Niagara/Source/Niagara/Classes/NiagaraConstants.h
@@ -186,21 +186,21 @@ struct FNiagaraConstants
 	static NIAGARA_API const TArray<FNiagaraVariable>& GetOldPositionTypeVariables();
 
 	static NIAGARA_API const TArray<FNiagaraVariable>& GetCommonParticleAttributes();
-	static NIAGARA_API FText GetAttributeDescription(const FNiagaraVariable& InVar);
-	static NIAGARA_API FString GetAttributeDefaultValue(const FNiagaraVariable& InVar);
-	static NIAGARA_API FNiagaraVariable GetAttributeWithDefaultValue(const FNiagaraVariable& InAttribute);
-	static NIAGARA_API FNiagaraVariable GetAttributeAsParticleDataSetKey(const FNiagaraVariable& InAttribute);
-	static NIAGARA_API FNiagaraVariable GetAttributeAsEmitterDataSetKey(const FNiagaraVariable& InAttribute);
-	static NIAGARA_API FNiagaraVariableAttributeBinding GetAttributeDefaultBinding(const FNiagaraVariable& InAttribute);
+	static NIAGARA_API FText GetAttributeDescription(const FNiagaraVariableBase& InVar);
+	static NIAGARA_API FString GetAttributeDefaultValue(const FNiagaraVariableBase& InVar);
+	static NIAGARA_API FNiagaraVariable GetAttributeWithDefaultValue(const FNiagaraVariableBase& InAttribute);
+	static NIAGARA_API FNiagaraVariableBase GetAttributeAsParticleDataSetKey(const FNiagaraVariableBase& InAttribute);
+	static NIAGARA_API FNiagaraVariableBase GetAttributeAsEmitterDataSetKey(const FNiagaraVariableBase& InAttribute);
+	static NIAGARA_API FNiagaraVariableAttributeBinding GetAttributeDefaultBinding(const FNiagaraVariableBase& InAttribute);
 
-	static NIAGARA_API bool IsNiagaraConstant(const FNiagaraVariable& InVar);
-	static NIAGARA_API const FNiagaraVariableMetaData* GetConstantMetaData(const FNiagaraVariable& InVar);
+	static NIAGARA_API bool IsNiagaraConstant(const FNiagaraVariableBase& InVar);
+	static NIAGARA_API const FNiagaraVariableMetaData* GetConstantMetaData(const FNiagaraVariableBase& InVar);
 
 	static NIAGARA_API const FNiagaraVariable* GetKnownConstant(const FName& InName, bool bAllowPartialNameMatch);
 	static NIAGARA_API const FNiagaraKnownConstantInfo GetKnownConstantInfo(const FName& InName, bool bAllowPartialNameMatch);
 	static NIAGARA_API const FNiagaraVariable *FindStaticSwitchConstant(const FName& InName);
 
-	static NIAGARA_API bool IsEngineManagedAttribute(const FNiagaraVariable& Var);
+	static NIAGARA_API bool IsEngineManagedAttribute(const FNiagaraVariableBase& Var);
 
 	static NIAGARA_API const FName InputPinName;
 	static NIAGARA_API const FName OutputPinName;
@@ -263,11 +263,11 @@ private:
 	static TMap<FName, FNiagaraVariable> UpdatedSystemParameters;
 	static TMap<FNiagaraVariable, FText> SystemStrMap;
 	static TArray<FNiagaraVariable> Attributes;
-	static TMap<FNiagaraVariable, FString> AttrDefaultsStrMap;
-	static TMap<FNiagaraVariable, FNiagaraVariable> AttrDefaultsValueMap;
-	static TMap<FNiagaraVariable, FNiagaraVariable> AttrDataSetKeyMap;
-	static TMap<FNiagaraVariable, FText> AttrDescStrMap;
-	static TMap<FNiagaraVariable, FNiagaraVariableMetaData> AttrMetaData;
+	static TMap<FNiagaraVariableBase, FString> AttrDefaultsStrMap;
+	static TMap<FNiagaraVariableBase, FNiagaraVariable> AttrDefaultsValueMap;
+	static TMap<FNiagaraVariableBase, FNiagaraVariable> AttrDataSetKeyMap;
+	static TMap<FNiagaraVariableBase, FText> AttrDescStrMap;
+	static TMap<FNiagaraVariableBase, FNiagaraVariableMetaData> AttrMetaData;
 	
 	static TArray<FNiagaraVariable> EngineManagedAttributes;
 
diff --git a/Engine/Plugins/FX/Niagara/Source/Niagara/Private/NiagaraCommon.cpp b/Engine/Plugins/FX/Niagara/Source/Niagara/Private/NiagaraCommon.cpp
--- a/Engine/Plugins/FX/Niagara/Source/Niagara/Private/NiagaraCommon.cpp
+++ b/Engine/Plugins/FX/Niagara/Source/Niagara/Private/NiagaraCommon.cpp
@@ -489,7 +489,9 @@ FCriticalSection* FNiagaraStatDatabase::GetCriticalSection() const
 void FNiagaraVariableAttributeBinding::SetValue(const FName& InValue, const FVersionedNiagaraEmitter& InVersionedEmitter, ENiagaraRendererSourceDataMode InSourceMode)
 {
 	UNiagaraEmitter* Emitter = InVersionedEmitter.Emitter;
-	RootVariable.SetName(InValue);
+	RootName = InValue;
+
+	FNiagaraVariableBase RootVariable = GetRootVariable();
 
 	const bool bIsRootParticleValue = RootVariable.IsInNameSpace(FNiagaraConstants::ParticleAttributeNamespaceString);
 	const bool bIsRootUnaliasedEmitterValue = RootVariable.IsInNameSpace(FNiagaraConstants::EmitterNamespaceString);
@@ -549,6 +551,7 @@ void FNiagaraVariableAttributeBinding::SetValue(const FName& InValue, const FVer
 	{
 		ensureMsgf(!bIsStackContextValue, TEXT("Should not get to this point! Should be covered by first two branch expresssions."));
 	}
+	RootName = RootVariable.GetName();
 
 	CacheValues(InVersionedEmitter, InSourceMode);
 }
@@ -558,7 +561,7 @@ void FNiagaraVariableAttributeBinding::SetAsPreviousValue(const FNiagaraVariable
 	static const FString PreviousNamespace = FNiagaraConstants::PreviousNamespace.ToString();
 
 	ParamMapVariable = Src;
-	RootVariable = Src;
+	RootName = Src.GetName();
 	DataSetName = Src.GetName();
 
 	// Split out the name and it's namespace
@@ -595,14 +598,16 @@ void FNiagaraVariableAttributeBinding::SetAsPreviousValue(const FNiagaraVariable
 {
 	static const FString PreviousNamespace = FNiagaraConstants::PreviousNamespace.ToString();
 
-	ParamMapVariable = Src.RootVariable;
-	RootVariable = Src.RootVariable;
-	DataSetName = Src.RootVariable.GetName();
+	const FNiagaraVariableBase SrcRootVariable = Src.GetRootVariable();
+
+	ParamMapVariable = SrcRootVariable;
+	RootName = SrcRootVariable.GetName();
+	DataSetName = SrcRootVariable.GetName();
 
 	// Split out the name and it's namespace
 	TStringBuilder<128> VarName;
 	TArray<FStringView, TInlineAllocator<16>> SplitName;
-	Src.RootVariable.GetName().ToString(VarName);
+	SrcRootVariable.GetName().ToString(VarName);
 	UE::String::ParseTokens(VarName, TEXT('.'), [&SplitName](FStringView Token) { SplitName.Emplace(Token); });
 
 	// If the name already contains a "Previous" in the name, just go with that
@@ -618,7 +623,7 @@ void FNiagaraVariableAttributeBinding::SetAsPreviousValue(const FNiagaraVariable
 
 	if (bIsPrev)
 	{
-		SetValue(Src.RootVariable.GetName(), InEmitter, InSourceMode);
+		SetValue(SrcRootVariable.GetName(), InEmitter, InSourceMode);
 	}
 	else
 	{
@@ -646,35 +651,43 @@ void FNiagaraVariableAttributeBinding::SetAsPreviousValue(const FNiagaraVariable
 
 void FNiagaraVariableAttributeBinding::Setup(const FNiagaraVariableBase& InRootVar, const FNiagaraVariable& InDefaultValue, ENiagaraRendererSourceDataMode InSourceMode)
 {
-	RootVariable = InRootVar;
-	if (InDefaultValue.IsDataAllocated() && InDefaultValue.GetType() == InRootVar.GetType())
-	{
-		RootVariable.SetData(InDefaultValue.GetData());
-	}
+	RootName = InRootVar.GetName();
+	ParamMapVariable.SetType(InRootVar.GetType());
+#if WITH_EDITORONLY_DATA
+	DefaultValue = InDefaultValue;
+#endif
 	SetValue(InRootVar.GetName(), FVersionedNiagaraEmitter(), InSourceMode);
 }
 
 #if WITH_EDITORONLY_DATA
 FString FNiagaraVariableAttributeBinding::GetDefaultValueString() const
 {
-	FString DefaultValueStr = RootVariable.GetName().ToString();
-
-	if (!RootVariable.GetName().IsValid() || RootVariable.IsDataAllocated() == true)
+	FString DefaultValueStr;
+	if (DefaultValue.IsValid())
 	{
-		DefaultValueStr = RootVariable.GetType().ToString(RootVariable.GetData());
-		DefaultValueStr.TrimEndInline();
+		if (DefaultValue.IsDataAllocated())
+		{
+			DefaultValueStr = DefaultValue.GetType().ToString(DefaultValue.GetData());
+			DefaultValueStr.TrimEndInline();
+		}
+		else
+		{
+			DefaultValueStr = DefaultValue.GetName().ToString();
+		}
 	}
 	return DefaultValueStr;
 }
-
 #endif
 
 void FNiagaraVariableAttributeBinding::PostLoad(ENiagaraRendererSourceDataMode InSourceMode)
 {
 #if WITH_EDITORONLY_DATA
+	if (!RootVariable_DEPRECATED.GetName().IsNone())
+	{
+		RootName = RootVariable_DEPRECATED.GetName();
+	}
 	if (BoundVariable.IsValid())
 	{
-		RootVariable.SetType(ParamMapVariable.GetType()); //Sometimes the BoundVariable was bogus in the past. THe DataSet shouldn't be though.
 		SetValue(BoundVariable.GetName(), FVersionedNiagaraEmitter(), InSourceMode);
 		BoundVariable = FNiagaraVariable();
 	}
@@ -691,7 +704,7 @@ void FNiagaraVariableAttributeBinding::Dump() const
 	UE_LOG(LogNiagara, Log, TEXT("PostLoad for FNiagaraVariableAttributeBinding...."));
 	UE_LOG(LogNiagara, Log, TEXT("ParamMapVariable: %s %s"), *ParamMapVariable.GetName().ToString(), *ParamMapVariable.GetType().GetName());
 	UE_LOG(LogNiagara, Log, TEXT("DataSetName: %s"), *DataSetName.ToString());
-	UE_LOG(LogNiagara, Log, TEXT("RootVariable: %s %s"), *RootVariable.GetName().ToString(), *RootVariable.GetType().GetName());
+	UE_LOG(LogNiagara, Log, TEXT("RootVariable: %s"), *RootName.ToString());
 #if WITH_EDITORONLY_DATA
 	UE_LOG(LogNiagara, Log, TEXT("BoundVariable: %s %s"), *BoundVariable.GetName().ToString(), *BoundVariable.GetType().GetName());
 	UE_LOG(LogNiagara, Log, TEXT("CachedDisplayName: %s"), *CachedDisplayName.ToString());
@@ -705,7 +718,7 @@ void FNiagaraVariableAttributeBinding::ResetToDefault(const FNiagaraVariableAttr
 	if (InOther.BindingSourceMode == ImplicitFromSource || InOther.BindingSourceMode == ENiagaraBindingSource::ExplicitEmitter || InOther.BindingSourceMode == ENiagaraBindingSource::ExplicitParticles)
 	{
 		// The default may have been set with a different source mode, so we can't copy values over directly. Instead, we need to copy the implicit values over.
-		FNiagaraVariable TempVar = InOther.RootVariable;
+		FNiagaraVariableBase TempVar = InOther.GetRootVariable();
 		if ((InSourceMode == ENiagaraRendererSourceDataMode::Emitter && InOther.BindingSourceMode == ENiagaraBindingSource::ImplicitFromSource) ||
 			InOther.BindingSourceMode == ENiagaraBindingSource::ExplicitEmitter)
 		{
@@ -723,7 +736,7 @@ void FNiagaraVariableAttributeBinding::ResetToDefault(const FNiagaraVariableAttr
 	}
 	else
 	{
-		SetValue(InOther.RootVariable.GetName(), InEmitter, InSourceMode);
+		SetValue(InOther.RootName, InEmitter, InSourceMode);
 	}
 }
 
@@ -731,7 +744,7 @@ bool FNiagaraVariableAttributeBinding::MatchesDefault(const FNiagaraVariableAttr
 {
 	if (DataSetName != InOther.DataSetName)
 		return false;
-	if (RootVariable.GetName() != InOther.RootVariable.GetName())
+	if (RootName != InOther.RootName)
 		return false;
 	return true;
 }
@@ -790,6 +803,7 @@ void FNiagaraVariableAttributeBinding::CacheValues(const FVersionedNiagaraEmitte
 	// Some older values may have had the root with the emitter unique name as the namespace, fix this up
 	// to meet the new assumptions.
 	UNiagaraEmitter* Emitter = InVersionedEmitter.Emitter;
+	FNiagaraVariableBase RootVariable = GetRootVariable();
 	if (Emitter && RootVariable.IsInNameSpace(Emitter->GetUniqueEmitterName()))
 	{
 		// First, replace unaliased emitter namespace with "Emitter" namespace
@@ -798,9 +812,10 @@ void FNiagaraVariableAttributeBinding::CacheValues(const FVersionedNiagaraEmitte
 
 		// Now strip out "Emitter"
 		RootVariable.SetName(FNiagaraConstants::GetAttributeAsEmitterDataSetKey(RootVariable).GetName());
+		RootName = RootVariable.GetName();
 	}
 
-	ParamMapVariable = (const FNiagaraVariableBase&)RootVariable;
+	ParamMapVariable = RootVariable;
 	DataSetName = ParamMapVariable.GetName();
 	bBindingExistsOnSource = false;
 
diff --git a/Engine/Plugins/FX/Niagara/Source/Niagara/Private/NiagaraConstants.cpp b/Engine/Plugins/FX/Niagara/Source/Niagara/Private/NiagaraConstants.cpp
--- a/Engine/Plugins/FX/Niagara/Source/Niagara/Private/NiagaraConstants.cpp
+++ b/Engine/Plugins/FX/Niagara/Source/Niagara/Private/NiagaraConstants.cpp
@@ -16,11 +16,11 @@ TArray<FNiagaraVariable> FNiagaraConstants::OldPositionTypes;
 TMap<FName, FNiagaraVariable> FNiagaraConstants::UpdatedSystemParameters;
 TMap<FNiagaraVariable, FText> FNiagaraConstants::SystemStrMap;
 TArray<FNiagaraVariable> FNiagaraConstants::Attributes;
-TMap<FNiagaraVariable, FString> FNiagaraConstants::AttrDefaultsStrMap;
-TMap<FNiagaraVariable, FText> FNiagaraConstants::AttrDescStrMap;
-TMap<FNiagaraVariable, FNiagaraVariableMetaData> FNiagaraConstants::AttrMetaData;
-TMap<FNiagaraVariable, FNiagaraVariable> FNiagaraConstants::AttrDefaultsValueMap;
-TMap<FNiagaraVariable, FNiagaraVariable> FNiagaraConstants::AttrDataSetKeyMap;
+TMap<FNiagaraVariableBase, FString> FNiagaraConstants::AttrDefaultsStrMap;
+TMap<FNiagaraVariableBase, FText> FNiagaraConstants::AttrDescStrMap;
+TMap<FNiagaraVariableBase, FNiagaraVariableMetaData> FNiagaraConstants::AttrMetaData;
+TMap<FNiagaraVariableBase, FNiagaraVariable> FNiagaraConstants::AttrDefaultsValueMap;
+TMap<FNiagaraVariableBase, FNiagaraVariable> FNiagaraConstants::AttrDataSetKeyMap;
 TArray<FNiagaraVariable> FNiagaraConstants::EngineManagedAttributes;
 
 const FName FNiagaraConstants::InputPinName("InputPin");
@@ -655,7 +655,7 @@ const TArray<FNiagaraVariable>& FNiagaraConstants::GetStaticSwitchConstants()
 	return SwitchParameters;
 }
 
-bool FNiagaraConstants::IsEngineManagedAttribute(const FNiagaraVariable& Var)
+bool FNiagaraConstants::IsEngineManagedAttribute(const FNiagaraVariableBase& Var)
 {
 	return EngineManagedAttributes.Contains(Var);
 }
@@ -719,7 +719,7 @@ const TArray<FNiagaraVariable>& FNiagaraConstants::GetCommonParticleAttributes()
 	return Attributes;
 }
 
-FString FNiagaraConstants::GetAttributeDefaultValue(const FNiagaraVariable& InAttribute)
+FString FNiagaraConstants::GetAttributeDefaultValue(const FNiagaraVariableBase& InAttribute)
 {
 	check(AttrDefaultsStrMap.Num() != 0);
 	FString* FoundStr = AttrDefaultsStrMap.Find(InAttribute);
@@ -730,7 +730,7 @@ FString FNiagaraConstants::GetAttributeDefaultValue(const FNiagaraVariable& InAt
 	return FString();
 }
 
-FText FNiagaraConstants::GetAttributeDescription(const FNiagaraVariable& InAttribute)
+FText FNiagaraConstants::GetAttributeDescription(const FNiagaraVariableBase& InAttribute)
 {
 	check(AttrDescStrMap.Num() != 0);
 	FText* FoundStr = AttrDescStrMap.Find(InAttribute);
@@ -741,7 +741,7 @@ FText FNiagaraConstants::GetAttributeDescription(const FNiagaraVariable& InAttri
 	return FText();
 }
 
-bool FNiagaraConstants::IsNiagaraConstant(const FNiagaraVariable& InVar)
+bool FNiagaraConstants::IsNiagaraConstant(const FNiagaraVariableBase& InVar)
 {
 	if (GetConstantMetaData(InVar) != nullptr)
 	{
@@ -750,13 +750,13 @@ bool FNiagaraConstants::IsNiagaraConstant(const FNiagaraVariable& InVar)
 	return false;
 }
 
-const FNiagaraVariableMetaData* FNiagaraConstants::GetConstantMetaData(const FNiagaraVariable& InVar)
+const FNiagaraVariableMetaData* FNiagaraConstants::GetConstantMetaData(const FNiagaraVariableBase& InVar)
 {
 	check(AttrMetaData.Num() != 0);
 	return AttrMetaData.Find(InVar);
 }
 
-FNiagaraVariable FNiagaraConstants::GetAttributeWithDefaultValue(const FNiagaraVariable& InAttribute)
+FNiagaraVariable FNiagaraConstants::GetAttributeWithDefaultValue(const FNiagaraVariableBase& InAttribute)
 {
 	check(AttrDefaultsValueMap.Num() != 0);
 	FNiagaraVariable* FoundValue = AttrDefaultsValueMap.Find(InAttribute);
@@ -767,12 +767,12 @@ FNiagaraVariable FNiagaraConstants::GetAttributeWithDefaultValue(const FNiagaraV
 	return FNiagaraVariable();
 }
 
-FNiagaraVariable FNiagaraConstants::GetAttributeAsParticleDataSetKey(const FNiagaraVariable& InVar)
+FNiagaraVariableBase FNiagaraConstants::GetAttributeAsParticleDataSetKey(const FNiagaraVariableBase& InVar)
 {
 	static FString ParticlesString = TEXT("Particles.");
 	static FString StackContextString = TEXT("StackContext.");
 
-	FNiagaraVariable OutVar = InVar;
+	FNiagaraVariableBase OutVar = InVar;
 
 	TStringBuilder<128> DataSetName;
 	InVar.GetName().ToString(DataSetName);
@@ -790,12 +790,12 @@ FNiagaraVariable FNiagaraConstants::GetAttributeAsParticleDataSetKey(const FNiag
 	return OutVar;
 }
 
-FNiagaraVariable FNiagaraConstants::GetAttributeAsEmitterDataSetKey(const FNiagaraVariable& InVar)
+FNiagaraVariableBase FNiagaraConstants::GetAttributeAsEmitterDataSetKey(const FNiagaraVariableBase& InVar)
 {
 	static FString EmitterString = TEXT("Emitter.");
 	static FString StackContextString = TEXT("StackContext.");
 
-	FNiagaraVariable OutVar = InVar;
+	FNiagaraVariableBase OutVar = InVar;
 
 	TStringBuilder<128> DataSetName;
 	InVar.GetName().ToString(DataSetName);
@@ -812,7 +812,7 @@ FNiagaraVariable FNiagaraConstants::GetAttributeAsEmitterDataSetKey(const FNiaga
 	OutVar.SetName(FName(DataSetNameView));
 	return OutVar;
 }
-FNiagaraVariableAttributeBinding FNiagaraConstants::GetAttributeDefaultBinding(const FNiagaraVariable& InVar)
+FNiagaraVariableAttributeBinding FNiagaraConstants::GetAttributeDefaultBinding(const FNiagaraVariableBase& InVar)
 {
 	if (AttrDefaultsValueMap.Num() == 0)
 	{
diff --git a/Engine/Plugins/FX/Niagara/Source/Niagara/Private/NiagaraLightRendererProperties.cpp b/Engine/Plugins/FX/Niagara/Source/Niagara/Private/NiagaraLightRendererProperties.cpp
--- a/Engine/Plugins/FX/Niagara/Source/Niagara/Private/NiagaraLightRendererProperties.cpp
+++ b/Engine/Plugins/FX/Niagara/Source/Niagara/Private/NiagaraLightRendererProperties.cpp
@@ -18,9 +18,15 @@
 #include "Widgets/Text/STextBlock.h"
 #endif
 
-
 #define LOCTEXT_NAMESPACE "UNiagaraLightRendererProperties"
 
+namespace NiagaraLightRendererPropertiesPrivate
+{
+	static FLinearColor GDefaultColor = FLinearColor::White;
+	static float GDefaultRadius = 0.0f;
+	static float GDefaultVolumetricScattering = 0.0f;
+}
+
 TArray<TWeakObjectPtr<UNiagaraLightRendererProperties>> UNiagaraLightRendererProperties::LightRendererPropertiesToDeferredInit;
 
 UNiagaraLightRendererProperties::UNiagaraLightRendererProperties()
@@ -99,6 +105,7 @@ void UNiagaraLightRendererProperties::PostEditChangeProperty(struct FPropertyCha
 /** The bindings depend on variables that are created during the NiagaraModule startup. However, the CDO's are build prior to this being initialized, so we defer setting these values until later.*/
 void UNiagaraLightRendererProperties::InitCDOPropertiesAfterModuleStartup()
 {
+
 	UNiagaraLightRendererProperties* CDO = CastChecked<UNiagaraLightRendererProperties>(UNiagaraLightRendererProperties::StaticClass()->GetDefaultObject());
 	CDO->PositionBinding = FNiagaraConstants::GetAttributeDefaultBinding(SYS_PARAM_PARTICLES_POSITION);
 	CDO->ColorBinding = FNiagaraConstants::GetAttributeDefaultBinding(SYS_PARAM_PARTICLES_COLOR);
@@ -110,6 +117,10 @@ void UNiagaraLightRendererProperties::InitCDOPropertiesAfterModuleStartup()
 	CDO->SpecularScaleBinding = FNiagaraConstants::GetAttributeDefaultBinding(SYS_PARAM_PARTICLES_LIGHT_SPECULAR_SCALE);
 	CDO->DiffuseScaleBinding = FNiagaraConstants::GetAttributeDefaultBinding(SYS_PARAM_PARTICLES_LIGHT_DIFFUSE_SCALE);
 
+	NiagaraLightRendererPropertiesPrivate::GDefaultColor = FNiagaraConstants::GetAttributeWithDefaultValue(SYS_PARAM_PARTICLES_COLOR).GetValue<FLinearColor>();
+	NiagaraLightRendererPropertiesPrivate::GDefaultRadius = FNiagaraConstants::GetAttributeWithDefaultValue(SYS_PARAM_PARTICLES_LIGHT_RADIUS).GetValue<float>();
+	NiagaraLightRendererPropertiesPrivate::GDefaultVolumetricScattering = FNiagaraConstants::GetAttributeWithDefaultValue(SYS_PARAM_PARTICLES_LIGHT_VOLUMETRIC_SCATTERING).GetValue<float>();
+
 	for (TWeakObjectPtr<UNiagaraLightRendererProperties>& WeakLightRendererProperties : LightRendererPropertiesToDeferredInit)
 	{
 		if (WeakLightRendererProperties.Get())
@@ -207,6 +218,21 @@ void UNiagaraLightRendererProperties::CacheFromCompiledData(const FNiagaraDataSe
 	InitParticleDataSetAccessor(DiffuseScaleAccessor, CompiledData, DiffuseScaleBinding);
 }
 
+FLinearColor UNiagaraLightRendererProperties::GetDefaultColor()
+{
+	return NiagaraLightRendererPropertiesPrivate::GDefaultColor;
+}
+
+float UNiagaraLightRendererProperties::GetDefaultRadius()
+{
+	return NiagaraLightRendererPropertiesPrivate::GDefaultRadius;
+}
+
+float UNiagaraLightRendererProperties::GetDefaultVolumetricScattering()
+{
+	return NiagaraLightRendererPropertiesPrivate::GDefaultVolumetricScattering;
+}
+
 void UNiagaraLightRendererProperties::UpdateSourceModeDerivates(ENiagaraRendererSourceDataMode InSourceMode, bool bFromPropertyEdit)
 {
 	Super::UpdateSourceModeDerivates(InSourceMode, bFromPropertyEdit);
diff --git a/Engine/Plugins/FX/Niagara/Source/Niagara/Private/NiagaraRendererLights.cpp b/Engine/Plugins/FX/Niagara/Source/Niagara/Private/NiagaraRendererLights.cpp
--- a/Engine/Plugins/FX/Niagara/Source/Niagara/Private/NiagaraRendererLights.cpp
+++ b/Engine/Plugins/FX/Niagara/Source/Niagara/Private/NiagaraRendererLights.cpp
@@ -161,9 +161,9 @@ FNiagaraDynamicDataBase* FNiagaraRendererLights::GenerateDynamicData(const FNiag
 
 	const FNiagaraParameterStore& ParameterStore = Emitter->GetRendererBoundVariables();
 	const FVector3f DefaultPos = ParameterStore.GetParameterValueOrDefault(Properties->PositionBinding.GetParamMapBindableVariable(), DefaultSimPos);
-	const FLinearColor DefaultColor = ParameterStore.GetParameterValueOrDefault(Properties->ColorBinding.GetParamMapBindableVariable(), Properties->ColorBinding.GetDefaultValue<FLinearColor>());
-	const float DefaultRadius = ParameterStore.GetParameterValueOrDefault(Properties->RadiusBinding.GetParamMapBindableVariable(), Properties->RadiusBinding.GetDefaultValue<float>());
-	const float DefaultScattering = ParameterStore.GetParameterValueOrDefault(Properties->VolumetricScatteringBinding.GetParamMapBindableVariable(), Properties->VolumetricScatteringBinding.GetDefaultValue<float>());
+	const FLinearColor DefaultColor = ParameterStore.GetParameterValueOrDefault(Properties->ColorBinding.GetParamMapBindableVariable(), UNiagaraLightRendererProperties::GetDefaultColor());
+	const float DefaultRadius = ParameterStore.GetParameterValueOrDefault(Properties->RadiusBinding.GetParamMapBindableVariable(), UNiagaraLightRendererProperties::GetDefaultRadius());
+	const float DefaultScattering = ParameterStore.GetParameterValueOrDefault(Properties->VolumetricScatteringBinding.GetParamMapBindableVariable(), UNiagaraLightRendererProperties::GetDefaultVolumetricScattering());
 	const FNiagaraBool DefaultEnabled = ParameterStore.GetParameterValueOrDefault(Properties->LightRenderingEnabledBinding.GetParamMapBindableVariable(), FNiagaraBool(true));
 	const int32 DefaultVisibilityTag = ParameterStore.GetParameterValueOrDefault(Properties->RendererVisibilityTagBinding.GetParamMapBindableVariable(), Properties->RendererVisibility);
 	const float DefaultExponent = ParameterStore.GetParameterValueOrDefault(Properties->LightExponentBinding.GetParamMapBindableVariable(), Properties->DefaultExponent);
diff --git a/Engine/Plugins/FX/Niagara/Source/Niagara/Public/NiagaraCommon.h b/Engine/Plugins/FX/Niagara/Source/Niagara/Public/NiagaraCommon.h
--- a/Engine/Plugins/FX/Niagara/Source/Niagara/Public/NiagaraCommon.h
+++ b/Engine/Plugins/FX/Niagara/Source/Niagara/Public/NiagaraCommon.h
@@ -1319,6 +1319,7 @@ struct FNiagaraVariableAttributeBinding
 
 #if WITH_EDITORONLY_DATA
 	const FName& GetName() const { return CachedDisplayName; }
+
 	[[nodiscard]] NIAGARA_API FString GetDefaultValueString() const;
 #endif
 	const FNiagaraVariableBase& GetParamMapBindableVariable() const { return ParamMapVariable; }
@@ -1329,9 +1330,10 @@ struct FNiagaraVariableAttributeBinding
 	bool IsValid() const { return !DataSetName.IsNone() && ParamMapVariable.GetType().IsValid(); }
 	
 	template<typename T>
+	UE_DEPRECATED(5.7, "To reduce memory usage of renderers the default value is no longer stored inside the attribute binding.  Please use FNiagaraConstants::GetAttributeWithDefaultValue to get the value.")
 	T GetDefaultValue() const
 	{
-		return RootVariable.GetValue<T>();
+		return T{};
 	}
 
 	void NIAGARA_API Setup(const FNiagaraVariableBase& InRootVar, const FNiagaraVariable& InDefaultValue, ENiagaraRendererSourceDataMode InSourceMode = ENiagaraRendererSourceDataMode::Particles);
@@ -1342,10 +1344,13 @@ struct FNiagaraVariableAttributeBinding
 	void NIAGARA_API ResetToDefault(const FNiagaraVariableAttributeBinding& InOther, const FVersionedNiagaraEmitter& InEmitter, ENiagaraRendererSourceDataMode InSourceMode);
 	bool NIAGARA_API MatchesDefault(const FNiagaraVariableAttributeBinding& InOther, ENiagaraRendererSourceDataMode InSourceMode) const;
 
+private:
+	[[nodiscard]] FNiagaraVariableBase GetRootVariable() const { return FNiagaraVariableBase(ParamMapVariable.GetType(), RootName); }
+
 protected:
 	/** The namespace and default value explicitly set by the user. If meant to be derived from the source mode, it will be without a namespace.*/
 	UPROPERTY()
-	FNiagaraVariable RootVariable;
+	FName RootName;
 
 	/** The fully expressed namespace for the variable. If an emitter namespace, this will include the Emitter's unique name.*/
 	UPROPERTY()
@@ -1360,6 +1365,9 @@ protected:
 	UPROPERTY()
 	FNiagaraVariable BoundVariable;
 
+	UPROPERTY(transient)
+	FNiagaraVariable DefaultValue;
+
 	UPROPERTY()
 	FName CachedDisplayName;
 #endif
@@ -1380,6 +1388,9 @@ protected:
 	/** The version of the namespace to be found in an attribute table lookup. I.e. without Particles or Emitter.*/
 	UPROPERTY()
 	FNiagaraVariable DataSetVariable_DEPRECATED;
+
+	UPROPERTY()
+	FNiagaraVariable RootVariable_DEPRECATED;
 #endif
 };
 
diff --git a/Engine/Plugins/FX/Niagara/Source/Niagara/Public/NiagaraLightRendererProperties.h b/Engine/Plugins/FX/Niagara/Source/Niagara/Public/NiagaraLightRendererProperties.h
--- a/Engine/Plugins/FX/Niagara/Source/Niagara/Public/NiagaraLightRendererProperties.h
+++ b/Engine/Plugins/FX/Niagara/Source/Niagara/Public/NiagaraLightRendererProperties.h
@@ -49,6 +49,10 @@ public:
 #endif
 	//UNiagaraRendererProperties Interface END
 
+	static FLinearColor GetDefaultColor();
+	static float GetDefaultRadius();
+	static float GetDefaultVolumetricScattering();
+
 	/** Whether or not to draw a single element for the Emitter or to draw the particles.*/
 	UPROPERTY(EditAnywhere, Category = "Light Rendering")
 	ENiagaraRendererSourceDataMode SourceMode = ENiagaraRendererSourceDataMode::Particles;
