- Default weight to 1.0 on distribution
- Add missing customization file from previous check in

#rnx

diff --git a/Engine/Plugins/FX/Niagara/Source/Niagara/Internal/DataInterface/NiagaraDataInterfaceArrayDistributionInt.h b/Engine/Plugins/FX/Niagara/Source/Niagara/Internal/DataInterface/NiagaraDataInterfaceArrayDistributionInt.h
--- a/Engine/Plugins/FX/Niagara/Source/Niagara/Internal/DataInterface/NiagaraDataInterfaceArrayDistributionInt.h
+++ b/Engine/Plugins/FX/Niagara/Source/Niagara/Internal/DataInterface/NiagaraDataInterfaceArrayDistributionInt.h
@@ -19,7 +19,7 @@ struct FNDIDistributionIntArrayEntry
 	// Used to determine the probability this value will be selected when using random
 	// The higher the weight relative to the other weights the higher the chance this value is selected
 	UPROPERTY(EditAnywhere, Category="Entry")
-	float Weight = 0.0f;
+	float Weight = 1.0f;
 
 	bool operator==(const FNDIDistributionIntArrayEntry& Other) const
 	{
diff --git a/Engine/Plugins/FX/Niagara/Source/NiagaraEditor/Private/Customizations/NiagaraTypeCustomizations.cpp b/Engine/Plugins/FX/Niagara/Source/NiagaraEditor/Private/Customizations/NiagaraTypeCustomizations.cpp
--- a/Engine/Plugins/FX/Niagara/Source/NiagaraEditor/Private/Customizations/NiagaraTypeCustomizations.cpp
+++ b/Engine/Plugins/FX/Niagara/Source/NiagaraEditor/Private/Customizations/NiagaraTypeCustomizations.cpp
@@ -228,10 +228,14 @@ FText FNiagaraVariableAttributeBindingCustomization::GetTooltipText() const
 {
 	if (OwningVersionedEmitter.Emitter && TargetVariableBinding)
 	{
-		FString DefaultValueStr = TargetVariableBinding->GetDefaultValueString();
+		FText TooltipDesc = FText::Format(LOCTEXT("AttributeBindingTooltip", "Use the variable \"{0}\" if it exists."), FText::FromName(TargetVariableBinding->GetName()));
+
+		const FString DefaultValueStr = TargetVariableBinding->GetDefaultValueString();
+		if (DefaultValueStr.Len() > 0)
+		{
+			TooltipDesc = FText::Format(LOCTEXT("AttributeBindingTooltipValue", "{0}.  Use the value \"{1}\" it the variable does not eixst."), TooltipDesc, FText::FromString(DefaultValueStr));
+		}
 
-		FText TooltipDesc = FText::Format(LOCTEXT("AttributeBindingTooltip", "Use the variable \"{0}\" if it exists, otherwise use the default \"{1}\" "), FText::FromName(TargetVariableBinding->GetName()),
-			FText::FromString(DefaultValueStr));
 		return TooltipDesc;
 	}
 	return FText::FromString(TEXT("Missing"));
