- Fix type in tooltip

#rnx

diff --git a/Engine/Plugins/FX/Niagara/Source/NiagaraEditor/Private/Customizations/NiagaraTypeCustomizations.cpp b/Engine/Plugins/FX/Niagara/Source/NiagaraEditor/Private/Customizations/NiagaraTypeCustomizations.cpp
--- a/Engine/Plugins/FX/Niagara/Source/NiagaraEditor/Private/Customizations/NiagaraTypeCustomizations.cpp
+++ b/Engine/Plugins/FX/Niagara/Source/NiagaraEditor/Private/Customizations/NiagaraTypeCustomizations.cpp
@@ -233,7 +233,7 @@ FText FNiagaraVariableAttributeBindingCustomization::GetTooltipText() const
 		const FString DefaultValueStr = TargetVariableBinding->GetDefaultValueString();
 		if (DefaultValueStr.Len() > 0)
 		{
-			TooltipDesc = FText::Format(LOCTEXT("AttributeBindingTooltipValue", "{0}.  Use the value \"{1}\" it the variable does not eixst."), TooltipDesc, FText::FromString(DefaultValueStr));
+			TooltipDesc = FText::Format(LOCTEXT("AttributeBindingTooltipValue", "{0}.  Use the value \"{1}\" if the variable does not exist."), TooltipDesc, FText::FromString(DefaultValueStr));
 		}
 
 		return TooltipDesc;
