- Fix validation accidentally changing usage with flags on materials
- CheckMaterialUsage will both check and set which can result in a race with shader compilation

#tests Disable used with flag on a material used with a Niagara asset, validate the asset and confirm the flag has not been enabled again

diff --git a/Engine/Plugins/FX/Niagara/Source/Niagara/Private/NiagaraRendererProperties.cpp b/Engine/Plugins/FX/Niagara/Source/Niagara/Private/NiagaraRendererProperties.cpp
--- a/Engine/Plugins/FX/Niagara/Source/Niagara/Private/NiagaraRendererProperties.cpp
+++ b/Engine/Plugins/FX/Niagara/Source/Niagara/Private/NiagaraRendererProperties.cpp
@@ -888,22 +888,22 @@ void UNiagaraRendererProperties::GetMaterialUsageFeedback(EMaterialUsage Usage,
 	FString FailedDomainNames;
 	for (UMaterialInterface* Material : Materials)
 	{
-		if (!Material)
+		UMaterial* BaseMaterial = Material ? Material->GetMaterial() : nullptr;
+		if (!BaseMaterial)
 		{
 			continue;
 		}
 
-		if (!Material->CheckMaterialUsage(Usage))
+		bool bOutHasUsage = false;
+		BaseMaterial->NeedsSetMaterialUsage_Concurrent(bOutHasUsage, Usage);
+		if (bOutHasUsage == false)
 		{
 			AppendToString(FailedUsageNames, Material->GetName());
 		}
 
-		if ( UMaterial* BaseMaterial = Material->GetMaterial() )
+		if (InvalidMaterialDomains.Contains(BaseMaterial->MaterialDomain))
 		{
-			if (InvalidMaterialDomains.Contains(BaseMaterial->MaterialDomain))
-			{
-				AppendToString(FailedDomainNames, Material->GetName());
-			}
+			AppendToString(FailedDomainNames, Material->GetName());
 		}
 	}
 
