- Rule for validating material usage in renderers
- This allows us to block content that does not have the appropriate usage flags set

diff --git a/Engine/Plugins/FX/Niagara/Source/Niagara/Private/NiagaraMeshRendererProperties.cpp b/Engine/Plugins/FX/Niagara/Source/Niagara/Private/NiagaraMeshRendererProperties.cpp
--- a/Engine/Plugins/FX/Niagara/Source/Niagara/Private/NiagaraMeshRendererProperties.cpp
+++ b/Engine/Plugins/FX/Niagara/Source/Niagara/Private/NiagaraMeshRendererProperties.cpp
@@ -1210,8 +1210,6 @@ void UNiagaraMeshRendererProperties::GetRendererFeedback(const FVersionedNiagara
 {
 	Super::GetRendererFeedback(InEmitter, OutErrors, OutWarnings, OutInfo);
 
-	GetMaterialUsageFeedback(MATUSAGE_NiagaraMeshParticles, { MD_Volume }, OutWarnings);
-
 	if (MaterialParameters.HasAnyBindings())
 	{
 		TArray<UMaterialInterface*> Materials;
@@ -1222,6 +1220,12 @@ void UNiagaraMeshRendererProperties::GetRendererFeedback(const FVersionedNiagara
 	// TODO: Check if Nanite fallback mesh was generated
 }
 
+void UNiagaraMeshRendererProperties::GetRendererMaterialUsageFeedback(TArray<FNiagaraRendererFeedback>& OutFeedback) const
+{
+	Super::GetRendererMaterialUsageFeedback(OutFeedback);
+	GetMaterialUsageFeedback(MATUSAGE_NiagaraMeshParticles, { MD_Volume }, OutFeedback);
+}
+
 void UNiagaraMeshRendererProperties::BeginDestroy()
 {
 	Super::BeginDestroy();
diff --git a/Engine/Plugins/FX/Niagara/Source/Niagara/Private/NiagaraRendererProperties.cpp b/Engine/Plugins/FX/Niagara/Source/Niagara/Private/NiagaraRendererProperties.cpp
--- a/Engine/Plugins/FX/Niagara/Source/Niagara/Private/NiagaraRendererProperties.cpp
+++ b/Engine/Plugins/FX/Niagara/Source/Niagara/Private/NiagaraRendererProperties.cpp
@@ -867,6 +867,8 @@ void UNiagaraRendererProperties::GetRendererFeedback(const FVersionedNiagaraEmit
 	{
 		OutInfo.Add(FNiagaraRendererFeedback(InfoText));
 	}
+
+	GetRendererMaterialUsageFeedback(OutWarnings);
 }
 
 void UNiagaraRendererProperties::GetMaterialUsageFeedback(EMaterialUsage Usage, TConstArrayView<EMaterialDomain> InvalidMaterialDomains, TArray<FNiagaraRendererFeedback>& OutFeedback) const
diff --git a/Engine/Plugins/FX/Niagara/Source/Niagara/Private/NiagaraRibbonRendererProperties.cpp b/Engine/Plugins/FX/Niagara/Source/Niagara/Private/NiagaraRibbonRendererProperties.cpp
--- a/Engine/Plugins/FX/Niagara/Source/Niagara/Private/NiagaraRibbonRendererProperties.cpp
+++ b/Engine/Plugins/FX/Niagara/Source/Niagara/Private/NiagaraRibbonRendererProperties.cpp
@@ -643,8 +643,6 @@ void UNiagaraRibbonRendererProperties::GetRendererFeedback(const FVersionedNiaga
 {
 	Super::GetRendererFeedback(InEmitter, OutErrors, OutWarnings, OutInfo);
 
-	GetMaterialUsageFeedback(MATUSAGE_NiagaraRibbons, OutWarnings);
-
 	const FVersionedNiagaraEmitterData* EmitterData = InEmitter.GetEmitterData();
 
 	// If this is a renderer made before we used UniqueID give the user an option to switch into the new mode
@@ -711,6 +709,12 @@ void UNiagaraRibbonRendererProperties::GetRendererFeedback(const FVersionedNiaga
 	}
 }
 
+void UNiagaraRibbonRendererProperties::GetRendererMaterialUsageFeedback(TArray<FNiagaraRendererFeedback>& OutFeedback) const
+{
+	Super::GetRendererMaterialUsageFeedback(OutFeedback);
+	GetMaterialUsageFeedback(MATUSAGE_NiagaraRibbons, OutFeedback);
+}
+
 void UNiagaraRibbonRendererProperties::RenameVariable(const FNiagaraVariableBase& OldVariable, const FNiagaraVariableBase& NewVariable, const FVersionedNiagaraEmitter& InEmitter)
 {
 	Super::RenameVariable(OldVariable, NewVariable, InEmitter);
diff --git a/Engine/Plugins/FX/Niagara/Source/Niagara/Private/NiagaraSpriteRendererProperties.cpp b/Engine/Plugins/FX/Niagara/Source/Niagara/Private/NiagaraSpriteRendererProperties.cpp
--- a/Engine/Plugins/FX/Niagara/Source/Niagara/Private/NiagaraSpriteRendererProperties.cpp
+++ b/Engine/Plugins/FX/Niagara/Source/Niagara/Private/NiagaraSpriteRendererProperties.cpp
@@ -623,8 +623,6 @@ void UNiagaraSpriteRendererProperties::GetRendererFeedback(const FVersionedNiaga
 {
 	Super::GetRendererFeedback(InEmitter, OutErrors, OutWarnings, OutInfo);
 
-	GetMaterialUsageFeedback(MATUSAGE_NiagaraSprites, OutWarnings);
-
 	if (MaterialParameters.HasAnyBindings())
 	{
 		TArray<UMaterialInterface*> Materials;
@@ -663,6 +661,12 @@ void UNiagaraSpriteRendererProperties::GetRendererTooltipWidgets(const FNiagaraE
 	}
 }
 
+void UNiagaraSpriteRendererProperties::GetRendererMaterialUsageFeedback(TArray<FNiagaraRendererFeedback>& OutFeedback) const
+{
+	Super::GetRendererMaterialUsageFeedback(OutFeedback);
+	GetMaterialUsageFeedback(MATUSAGE_NiagaraSprites, OutFeedback);
+}
+
 void UNiagaraSpriteRendererProperties::UpdateCutoutTexture()
 {
 	if (bUseMaterialCutoutTexture == false)
diff --git a/Engine/Plugins/FX/Niagara/Source/Niagara/Public/NiagaraMeshRendererProperties.h b/Engine/Plugins/FX/Niagara/Source/Niagara/Public/NiagaraMeshRendererProperties.h
--- a/Engine/Plugins/FX/Niagara/Source/Niagara/Public/NiagaraMeshRendererProperties.h
+++ b/Engine/Plugins/FX/Niagara/Source/Niagara/Public/NiagaraMeshRendererProperties.h
@@ -162,6 +162,7 @@ public:
 	NIAGARA_API virtual	void GetRendererWidgets(const FNiagaraEmitterInstance* InEmitter, TArray<TSharedPtr<SWidget>>& OutWidgets, TSharedPtr<FAssetThumbnailPool> InThumbnailPool) const override;
 	NIAGARA_API virtual	void GetRendererTooltipWidgets(const FNiagaraEmitterInstance* InEmitter, TArray<TSharedPtr<SWidget>>& OutWidgets, TSharedPtr<FAssetThumbnailPool> InThumbnailPool) const override;
 	NIAGARA_API virtual void GetRendererFeedback(const FVersionedNiagaraEmitter& InEmitter, TArray<FNiagaraRendererFeedback>& OutErrors, TArray<FNiagaraRendererFeedback>& OutWarnings, TArray<FNiagaraRendererFeedback>& OutInfo) const override;
+	NIAGARA_API virtual void GetRendererMaterialUsageFeedback(TArray<FNiagaraRendererFeedback>& OutFeedback) const override;
 	NIAGARA_API void OnMeshChanged();
 	NIAGARA_API void OnMeshPostBuild(UStaticMesh*);
 	NIAGARA_API void OnAssetReimported(UObject*);
diff --git a/Engine/Plugins/FX/Niagara/Source/Niagara/Public/NiagaraRendererProperties.h b/Engine/Plugins/FX/Niagara/Source/Niagara/Public/NiagaraRendererProperties.h
--- a/Engine/Plugins/FX/Niagara/Source/Niagara/Public/NiagaraRendererProperties.h
+++ b/Engine/Plugins/FX/Niagara/Source/Niagara/Public/NiagaraRendererProperties.h
@@ -382,6 +382,12 @@ public:
 	virtual void GetRendererFeedback(const FVersionedNiagaraEmitter& InEmitter, TArray<FText>& OutErrors, TArray<FText>& OutWarnings, TArray<FText>& OutInfo) const {}
 	NIAGARA_API virtual void GetRendererFeedback(const FVersionedNiagaraEmitter& InEmitter, TArray<FNiagaraRendererFeedback>& OutErrors, TArray<FNiagaraRendererFeedback>& OutWarnings, TArray<FNiagaraRendererFeedback>& OutInfo) const;
 
+	/**
+	Provides feedback on any material usage issues, i.e. incompatable material assignment or missing usage flags
+	Automaticaly called from GetRendererFeedback where OutFeedback results in Errors
+	*/
+	NIAGARA_API virtual void GetRendererMaterialUsageFeedback(TArray<FNiagaraRendererFeedback>& OutFeedback) const {}
+
 	void GetMaterialUsageFeedback(EMaterialUsage Usage, TConstArrayView<EMaterialDomain> InvalidMaterialDomains, TArray<FNiagaraRendererFeedback>& OutFeedback) const;
 	void GetMaterialUsageFeedback(EMaterialUsage Usage, TArray<FNiagaraRendererFeedback>& OutFeedback) const;
 
diff --git a/Engine/Plugins/FX/Niagara/Source/Niagara/Public/NiagaraRibbonRendererProperties.h b/Engine/Plugins/FX/Niagara/Source/Niagara/Public/NiagaraRibbonRendererProperties.h
--- a/Engine/Plugins/FX/Niagara/Source/Niagara/Public/NiagaraRibbonRendererProperties.h
+++ b/Engine/Plugins/FX/Niagara/Source/Niagara/Public/NiagaraRibbonRendererProperties.h
@@ -228,6 +228,7 @@ public:
 	NIAGARA_API virtual void GetRendererWidgets(const FNiagaraEmitterInstance* InEmitter, TArray<TSharedPtr<SWidget>>& OutWidgets, TSharedPtr<FAssetThumbnailPool> InThumbnailPool) const override;
 	NIAGARA_API virtual void GetRendererTooltipWidgets(const FNiagaraEmitterInstance* InEmitter, TArray<TSharedPtr<SWidget>>& OutWidgets, TSharedPtr<FAssetThumbnailPool> InThumbnailPool) const override;
 	NIAGARA_API virtual void GetRendererFeedback(const FVersionedNiagaraEmitter& InEmitter, TArray<FNiagaraRendererFeedback>& OutErrors, TArray<FNiagaraRendererFeedback>& OutWarnings, TArray<FNiagaraRendererFeedback>& OutInfo) const override;
+	NIAGARA_API virtual void GetRendererMaterialUsageFeedback(TArray<FNiagaraRendererFeedback>& OutFeedback) const override;
 
 	NIAGARA_API virtual TArray<FNiagaraVariable> GetBoundAttributes() const override; 
 #endif
diff --git a/Engine/Plugins/FX/Niagara/Source/Niagara/Public/NiagaraSpriteRendererProperties.h b/Engine/Plugins/FX/Niagara/Source/Niagara/Public/NiagaraSpriteRendererProperties.h
--- a/Engine/Plugins/FX/Niagara/Source/Niagara/Public/NiagaraSpriteRendererProperties.h
+++ b/Engine/Plugins/FX/Niagara/Source/Niagara/Public/NiagaraSpriteRendererProperties.h
@@ -145,6 +145,7 @@ public:
 	NIAGARA_API virtual void GetRendererTooltipWidgets(const FNiagaraEmitterInstance* InEmitter, TArray<TSharedPtr<SWidget>>& OutWidgets, TSharedPtr<FAssetThumbnailPool> InThumbnailPool) const override;
 	NIAGARA_API virtual void GetRendererFeedback(const FVersionedNiagaraEmitter& InEmitter, TArray<FText>& OutErrors, TArray<FText>& OutWarnings, TArray<FText>& OutInfo) const override;
 	NIAGARA_API virtual void GetRendererFeedback(const FVersionedNiagaraEmitter& InEmitter, TArray<FNiagaraRendererFeedback>& OutErrors, TArray<FNiagaraRendererFeedback>& OutWarnings, TArray<FNiagaraRendererFeedback>& OutInfo) const override;
+	NIAGARA_API virtual void GetRendererMaterialUsageFeedback(TArray<FNiagaraRendererFeedback>& OutFeedback) const override;
 #endif
 	virtual ENiagaraRendererSourceDataMode GetCurrentSourceMode() const override { return SourceMode; }
 
diff --git a/Engine/Plugins/FX/Niagara/Source/NiagaraEditor/Private/NiagaraValidationRules.cpp b/Engine/Plugins/FX/Niagara/Source/NiagaraEditor/Private/NiagaraValidationRules.cpp
--- a/Engine/Plugins/FX/Niagara/Source/NiagaraEditor/Private/NiagaraValidationRules.cpp
+++ b/Engine/Plugins/FX/Niagara/Source/NiagaraEditor/Private/NiagaraValidationRules.cpp
@@ -1697,4 +1697,31 @@ void UNiagaraValidationRule_ModuleSimTargetRestriction::CheckValidity(const FNia
 	}
 }
 
+void UNiagaraValidationRule_MaterialUsage::CheckValidity(const FNiagaraValidationContext& Context, TArray<FNiagaraValidationResult>& OutResults) const
+{
+	TArray<TSharedRef<FNiagaraEmitterHandleViewModel>> EmitterHandleViewModels = Context.ViewModel->GetEmitterHandleViewModels();
+	for (TSharedRef<FNiagaraEmitterHandleViewModel> EmitterHandleModel : EmitterHandleViewModels)
+	{
+		FVersionedNiagaraEmitterData* EmitterData = EmitterHandleModel.Get().GetEmitterHandle()->GetEmitterData();
+		EmitterData->ForEachRenderer(
+			[&OutResults, EmitterHandleModel, this](UNiagaraRendererProperties* RendererProperties)
+			{
+				TArray<FNiagaraRendererFeedback> Feedback;
+				RendererProperties->GetRendererMaterialUsageFeedback(Feedback);
+				if (Feedback.Num() == 0)
+				{
+					return;
+				}
+
+				OutResults.Emplace(
+					FailedUsageSeverity,
+					LOCTEXT("MissingMaterialUsage", "Material usage flags are not set correctly."),
+					FText::Format(LOCTEXT("MissingMaterialUsage", "Materials on renderer {0} are missing the usage flag required to render with."), FText::FromName(RendererProperties->GetFName())),
+					NiagaraValidation::GetRendererStackItem(EmitterHandleModel.Get().GetEmitterStackViewModel(), RendererProperties)
+				);
+			}
+		);
+	}
+}
+
 #undef LOCTEXT_NAMESPACE
diff --git a/Engine/Plugins/FX/Niagara/Source/NiagaraEditor/Public/NiagaraValidationRules.h b/Engine/Plugins/FX/Niagara/Source/NiagaraEditor/Public/NiagaraValidationRules.h
--- a/Engine/Plugins/FX/Niagara/Source/NiagaraEditor/Public/NiagaraValidationRules.h
+++ b/Engine/Plugins/FX/Niagara/Source/NiagaraEditor/Public/NiagaraValidationRules.h
@@ -565,3 +565,16 @@ public:
 	UPROPERTY(EditAnywhere, Category = Validation)
 	ENiagaraSimTarget SupportedSimTarget = ENiagaraSimTarget::CPUSim;
 };
+
+/** Checks that the materials assigned to renderers have the correct usage with flags enabled. */
+UCLASS(Category = "Validation", DisplayName = "Material Usage")
+class UNiagaraValidationRule_MaterialUsage : public UNiagaraValidationRule
+{
+	GENERATED_BODY()
+
+public:
+	UPROPERTY(EditAnywhere, Category = Validation)
+	ENiagaraValidationSeverity FailedUsageSeverity = ENiagaraValidationSeverity::Error;
+
+	virtual void CheckValidity(const FNiagaraValidationContext& Context, TArray<FNiagaraValidationResult>& OutResults) const override;
+};
