- Add a method to get the prob/alias information for an entry

#rnx

#virtualized

diff --git a/Engine/Plugins/FX/Niagara/Shaders/Private/NiagaraDataInterfaceArrayDistributionIntTemplate.ush b/Engine/Plugins/FX/Niagara/Shaders/Private/NiagaraDataInterfaceArrayDistributionIntTemplate.ush
--- a/Engine/Plugins/FX/Niagara/Shaders/Private/NiagaraDataInterfaceArrayDistributionIntTemplate.ush
+++ b/Engine/Plugins/FX/Niagara/Shaders/Private/NiagaraDataInterfaceArrayDistributionIntTemplate.ush
@@ -27,6 +27,14 @@ void Get_{ParameterName}(in int InIndex, out NDIDistributionIntArrayEntry OutEnt
 	OutEntry.Weight = asfloat(Data.y);
 }
 
+void GetProbabilityAlias_{ParameterName}(in int InIndex, out float Prob, out int Alias)
+{
+	const int ClampedIndex = clamp(InIndex, 0, {ParameterName}_Length - 1);
+	const int Offset = ClampedIndex * 16; 
+	const uint2 Data = {ParameterName}_BuiltTableData.Load2(Offset);
+	Prob = asfloat(Data.x);
+	Alias = asint(Data.y);
+}
 
 void GetRandomValue_{ParameterName}(NiagaraRandInfo RandInfo, out int OutValue)
 {
diff --git a/Engine/Plugins/FX/Niagara/Source/Niagara/Internal/DataInterface/NiagaraDataInterfaceArrayDistributionInt.h b/Engine/Plugins/FX/Niagara/Source/Niagara/Internal/DataInterface/NiagaraDataInterfaceArrayDistributionInt.h
--- a/Engine/Plugins/FX/Niagara/Source/Niagara/Internal/DataInterface/NiagaraDataInterfaceArrayDistributionInt.h
+++ b/Engine/Plugins/FX/Niagara/Source/Niagara/Internal/DataInterface/NiagaraDataInterfaceArrayDistributionInt.h
@@ -86,6 +86,7 @@ private:
 	void VMIsLastIndex(FVectorVMExternalFunctionContext& Context) const;
 	void VMGetLastIndex(FVectorVMExternalFunctionContext& Context) const;
 	void VMGet(FVectorVMExternalFunctionContext& Context) const;
+	void VMGetProbAlias(FVectorVMExternalFunctionContext& Context) const;
 	void VMGetRandomValue(FVectorVMExternalFunctionContext& Context) const;
 
 private:
diff --git a/Engine/Plugins/FX/Niagara/Source/Niagara/Private/DataInterface/NiagaraDataInterfaceArrayDistributionInt.cpp b/Engine/Plugins/FX/Niagara/Source/Niagara/Private/DataInterface/NiagaraDataInterfaceArrayDistributionInt.cpp
--- a/Engine/Plugins/FX/Niagara/Source/Niagara/Private/DataInterface/NiagaraDataInterfaceArrayDistributionInt.cpp
+++ b/Engine/Plugins/FX/Niagara/Source/Niagara/Private/DataInterface/NiagaraDataInterfaceArrayDistributionInt.cpp
@@ -17,6 +17,7 @@
 
 namespace NDIArrayDistributionIntPrivate
 {
+	static const FName NAME_GetProbabilityAlias("GetProbabilityAlias");
 	static const FName NAME_GetRandomValue("GetRandomValue");
 
 	BEGIN_SHADER_PARAMETER_STRUCT(FShaderParameters, )
@@ -53,6 +54,16 @@ namespace NDIArrayDistributionIntPrivate
 		OutWeight = reinterpret_cast<const float*>(Data)[3];
 	}
 
+	void DecodeTableProbAlias(const TArray<uint8>& TableData, int32 Index, float& OutProb, int32& OutAlias)
+	{
+		const int32 EntrySize = GetBuiltTableEntrySize();
+		const int32 Offset = Index * EntrySize;
+		check(TableData.Num() >= Offset + EntrySize);
+		const uint8* Data = TableData.GetData() + Offset;
+		OutProb = reinterpret_cast<const float*>(Data)[0];
+		OutAlias = reinterpret_cast<const int32*>(Data)[1];
+	}
+
 	void DecodeTableIndirect(const TArray<uint8>& TableData, int32 Index, float Probability, int32& OutValue, float& OutWeight)
 	{
 		const int32 EntrySize = GetBuiltTableEntrySize();
@@ -183,6 +194,13 @@ void UNiagaraDataInterfaceArrayDistributionInt::GetFunctionsInternal(TArray<FNia
 	DefaultImmutableSig.Inputs.Emplace(FNiagaraTypeDefinition(UNiagaraDataInterfaceArrayDistributionInt::StaticClass()), TEXT("Distribution Int Array"));
 
 	FNiagaraDataInterfaceArrayImplInternal::GetImmutableFunctions(OutFunctions, DefaultImmutableSig, FNiagaraTypeDefinition(FNDIDistributionIntArrayEntry::StaticStruct()));
+	{
+		FNiagaraFunctionSignature& Sig = OutFunctions.Add_GetRef(DefaultImmutableSig);
+		Sig.Name = NAME_GetProbabilityAlias;
+		Sig.Inputs.Emplace(FNiagaraTypeDefinition::GetIntDef(), TEXT("Index"));
+		Sig.Outputs.Emplace(FNiagaraTypeDefinition::GetFloatDef(), TEXT("Probability"));
+		Sig.Outputs.Emplace(FNiagaraTypeDefinition::GetIntDef(), TEXT("Alias"));
+	}
 	{
 		FNiagaraFunctionSignature& Sig = OutFunctions.Add_GetRef(DefaultImmutableSig);
 		Sig.Name = NAME_GetRandomValue;
@@ -212,6 +230,10 @@ void UNiagaraDataInterfaceArrayDistributionInt::GetVMExternalFunction(const FVME
 	{
 		OutFunc = FVMExternalFunction::CreateUObject(this, &UNiagaraDataInterfaceArrayDistributionInt::VMGet);
 	}
+	else if (BindingInfo.Name == NDIArrayDistributionIntPrivate::NAME_GetProbabilityAlias)
+	{
+		OutFunc = FVMExternalFunction::CreateUObject(this, &UNiagaraDataInterfaceArrayDistributionInt::VMGetProbAlias);
+	}
 	else if (BindingInfo.Name == NDIArrayDistributionIntPrivate::NAME_GetRandomValue)
 	{
 		OutFunc = FVMExternalFunction::CreateUObject(this, &UNiagaraDataInterfaceArrayDistributionInt::VMGetRandomValue);
@@ -393,6 +415,30 @@ void UNiagaraDataInterfaceArrayDistributionInt::VMGet(FVectorVMExternalFunctionC
 	}
 }
 
+void UNiagaraDataInterfaceArrayDistributionInt::VMGetProbAlias(FVectorVMExternalFunctionContext& Context) const
+{
+	using namespace NDIArrayDistributionIntPrivate;
+
+	FNDIInputParam<int32>	InIndex(Context);
+	FNDIOutputParam<float>	OutProb(Context);
+	FNDIOutputParam<int32>	OutAlias(Context);
+
+	UE::TRWScopeLock ScopeLock(BuiltTableDataGuard, SLT_ReadOnly);
+	const int32 LengthMinusOne = GetBuiltTableLength() - 1;
+	for (int32 i = 0; i < Context.GetNumInstances(); ++i)
+	{
+		const int32 Index = InIndex.GetAndAdvance();
+		const int32 ClampedIndex = FMath::Clamp(Index, 0, LengthMinusOne);
+
+		float EntryProb = 0.0f;
+		int32 EntryAlias = 0;
+		DecodeTableProbAlias(BuiltTableData, ClampedIndex, EntryProb, EntryAlias);
+
+		OutProb.SetAndAdvance(EntryProb);
+		OutAlias.SetAndAdvance(EntryAlias);
+	}
+}
+
 void UNiagaraDataInterfaceArrayDistributionInt::VMGetRandomValue(FVectorVMExternalFunctionContext& Context) const
 {
 	using namespace NDIArrayDistributionIntPrivate;
