- Add missing PostLoad calls to attribute bindings

#rnx

diff --git a/Engine/Plugins/FX/Niagara/Source/Niagara/Private/NiagaraRendererProperties.cpp b/Engine/Plugins/FX/Niagara/Source/Niagara/Private/NiagaraRendererProperties.cpp
--- a/Engine/Plugins/FX/Niagara/Source/Niagara/Private/NiagaraRendererProperties.cpp
+++ b/Engine/Plugins/FX/Niagara/Source/Niagara/Private/NiagaraRendererProperties.cpp
@@ -1199,6 +1199,7 @@ bool UNiagaraRendererProperties::NeedsLoadForTargetPlatform(const ITargetPlatfor
 
 void UNiagaraRendererProperties::PostLoadBindings(ENiagaraRendererSourceDataMode InSourceMode)
 {
+	RendererEnabledBinding.PostLoad(ENiagaraRendererSourceDataMode::Emitter);
 	for (int32 i = 0; i < AttributeBindings.Num(); i++)
 	{
 		FNiagaraVariableAttributeBinding* Binding = const_cast<FNiagaraVariableAttributeBinding*>(AttributeBindings[i]);
diff --git a/Engine/Plugins/FX/Niagara/Source/Niagara/Private/NiagaraSimulationStageBase.cpp b/Engine/Plugins/FX/Niagara/Source/Niagara/Private/NiagaraSimulationStageBase.cpp
--- a/Engine/Plugins/FX/Niagara/Source/Niagara/Private/NiagaraSimulationStageBase.cpp
+++ b/Engine/Plugins/FX/Niagara/Source/Niagara/Private/NiagaraSimulationStageBase.cpp
@@ -253,6 +253,10 @@ void UNiagaraSimulationStageGeneric::PostLoad()
 		ParticleIterationStateBinding.Setup(GetDefaultParticleStateBinding(), GetDefaultParticleStateBinding(), ENiagaraRendererSourceDataMode::Particles);
 	}
 
+	// Make sure we post load all FNiagaraVariableAttributeBinding
+	EnabledBinding.PostLoad(ENiagaraRendererSourceDataMode::Emitter);
+	ParticleIterationStateBinding.PostLoad(ENiagaraRendererSourceDataMode::Particles);
+
 	// Handle Deprecated Properties
 	if (ElementCountBinding_DEPRECATED.IsValid())
 	{
