- Work to prepare generating shaders from stateless templates

#rnx

diff --git a/Engine/Plugins/FX/Niagara/Shaders/Private/Stateless/Modules/NiagaraStatelessModule_CalculateAccurateVelocity.ush b/Engine/Plugins/FX/Niagara/Shaders/Private/Stateless/Modules/NiagaraStatelessModule_CalculateAccurateVelocity.ush
--- a/Engine/Plugins/FX/Niagara/Shaders/Private/Stateless/Modules/NiagaraStatelessModule_CalculateAccurateVelocity.ush
+++ b/Engine/Plugins/FX/Niagara/Shaders/Private/Stateless/Modules/NiagaraStatelessModule_CalculateAccurateVelocity.ush
@@ -1,7 +1,5 @@
 // Copyright Epic Games, Inc. All Rights Reserved.
 
-#include "../NiagaraStatelessCommon.ush"
-
 void CalculateAccurateVelocity_Simulate(inout FStatelessParticle Particle)
 {  
 	Particle.Velocity			= (Particle.Position - Particle.PreviousPosition) * Particle.InvDeltaTime;
diff --git a/Engine/Plugins/FX/Niagara/Shaders/Private/Stateless/Modules/NiagaraStatelessModule_CameraOffset.ush b/Engine/Plugins/FX/Niagara/Shaders/Private/Stateless/Modules/NiagaraStatelessModule_CameraOffset.ush
--- a/Engine/Plugins/FX/Niagara/Shaders/Private/Stateless/Modules/NiagaraStatelessModule_CameraOffset.ush
+++ b/Engine/Plugins/FX/Niagara/Shaders/Private/Stateless/Modules/NiagaraStatelessModule_CameraOffset.ush
@@ -1,7 +1,5 @@
 // Copyright Epic Games, Inc. All Rights Reserved.
 
-#include "../NiagaraStatelessCommon.ush"
-
 FNiagaraStatelessBuiltDistributionType CameraOffset_Distribution;
 
 void CameraOffset_Simulate(inout FStatelessParticle Particle)
diff --git a/Engine/Plugins/FX/Niagara/Shaders/Private/Stateless/Modules/NiagaraStatelessModule_DynamicMaterialParameters.ush b/Engine/Plugins/FX/Niagara/Shaders/Private/Stateless/Modules/NiagaraStatelessModule_DynamicMaterialParameters.ush
--- a/Engine/Plugins/FX/Niagara/Shaders/Private/Stateless/Modules/NiagaraStatelessModule_DynamicMaterialParameters.ush
+++ b/Engine/Plugins/FX/Niagara/Shaders/Private/Stateless/Modules/NiagaraStatelessModule_DynamicMaterialParameters.ush
@@ -1,7 +1,5 @@
 // Copyright Epic Games, Inc. All Rights Reserved.
 
-#include "../NiagaraStatelessCommon.ush"
-
 uint	DynamicMaterialParameters_ChannelMask;
 FNiagaraStatelessBuiltDistributionType	DynamicMaterialParameters_Parameter0X;
 FNiagaraStatelessBuiltDistributionType	DynamicMaterialParameters_Parameter0Y;
@@ -22,7 +20,7 @@ FNiagaraStatelessBuiltDistributionType	DynamicMaterialParameters_Parameter3W;
 	
 void DynamicMaterialParameters_Simulate(inout FStatelessParticle Particle)
 {
-	Particle.DynamicMaterialParameter0 = 0.0f;
+	Particle.DynamicMaterialParameter = 0.0f;
 	Particle.DynamicMaterialParameter1 = 0.0f;
 	Particle.DynamicMaterialParameter2 = 0.0f;
 	Particle.DynamicMaterialParameter3 = 0.0f;
@@ -35,10 +33,10 @@ void DynamicMaterialParameters_Simulate(inout FStatelessParticle Particle)
 		ParameterSampler[2].Init(2, DynamicMaterialParameters_Parameter0Z);
 		ParameterSampler[3].Init(3, DynamicMaterialParameters_Parameter0W);
 		
-		Particle.DynamicMaterialParameter0.x = (DynamicMaterialParameters_ChannelMask & 0x0001) != 0 ? ParameterSampler[0].GetValue(Particle.NormalizedAge) : 0.0f;
-		Particle.DynamicMaterialParameter0.y = (DynamicMaterialParameters_ChannelMask & 0x0002) != 0 ? ParameterSampler[1].GetValue(Particle.NormalizedAge) : 0.0f;
-		Particle.DynamicMaterialParameter0.z = (DynamicMaterialParameters_ChannelMask & 0x0004) != 0 ? ParameterSampler[2].GetValue(Particle.NormalizedAge) : 0.0f;
-		Particle.DynamicMaterialParameter0.w = (DynamicMaterialParameters_ChannelMask & 0x0008) != 0 ? ParameterSampler[3].GetValue(Particle.NormalizedAge) : 0.0f;
+		Particle.DynamicMaterialParameter.x = (DynamicMaterialParameters_ChannelMask & 0x0001) != 0 ? ParameterSampler[0].GetValue(Particle.NormalizedAge) : 0.0f;
+		Particle.DynamicMaterialParameter.y = (DynamicMaterialParameters_ChannelMask & 0x0002) != 0 ? ParameterSampler[1].GetValue(Particle.NormalizedAge) : 0.0f;
+		Particle.DynamicMaterialParameter.z = (DynamicMaterialParameters_ChannelMask & 0x0004) != 0 ? ParameterSampler[2].GetValue(Particle.NormalizedAge) : 0.0f;
+		Particle.DynamicMaterialParameter.w = (DynamicMaterialParameters_ChannelMask & 0x0008) != 0 ? ParameterSampler[3].GetValue(Particle.NormalizedAge) : 0.0f;
 	}
 
 	if ((DynamicMaterialParameters_ChannelMask & 0x00f0) != 0)
diff --git a/Engine/Plugins/FX/Niagara/Shaders/Private/Stateless/Modules/NiagaraStatelessModule_InitialMeshOrientation.ush b/Engine/Plugins/FX/Niagara/Shaders/Private/Stateless/Modules/NiagaraStatelessModule_InitialMeshOrientation.ush
--- a/Engine/Plugins/FX/Niagara/Shaders/Private/Stateless/Modules/NiagaraStatelessModule_InitialMeshOrientation.ush
+++ b/Engine/Plugins/FX/Niagara/Shaders/Private/Stateless/Modules/NiagaraStatelessModule_InitialMeshOrientation.ush
@@ -1,7 +1,5 @@
 // Copyright Epic Games, Inc. All Rights Reserved.
 
-#include "../NiagaraStatelessCommon.ush"
-
 float3	InitialMeshOrientation_RotationScale;
 float3	InitialMeshOrientation_RotationBias;
 
diff --git a/Engine/Plugins/FX/Niagara/Shaders/Private/Stateless/Modules/NiagaraStatelessModule_InitializeParticle.ush b/Engine/Plugins/FX/Niagara/Shaders/Private/Stateless/Modules/NiagaraStatelessModule_InitializeParticle.ush
--- a/Engine/Plugins/FX/Niagara/Shaders/Private/Stateless/Modules/NiagaraStatelessModule_InitializeParticle.ush
+++ b/Engine/Plugins/FX/Niagara/Shaders/Private/Stateless/Modules/NiagaraStatelessModule_InitializeParticle.ush
@@ -1,7 +1,5 @@
 // Copyright Epic Games, Inc. All Rights Reserved.
 
-#include "../NiagaraStatelessCommon.ush"
-
 static const uint EInitializeParticleModuleFlag_UniformSpriteSize	= 1 << 0;
 static const uint EInitializeParticleModuleFlag_UniformMeshScale	= 1 << 1;
 
diff --git a/Engine/Plugins/FX/Niagara/Shaders/Private/Stateless/Modules/NiagaraStatelessModule_MeshIndex.ush b/Engine/Plugins/FX/Niagara/Shaders/Private/Stateless/Modules/NiagaraStatelessModule_MeshIndex.ush
--- a/Engine/Plugins/FX/Niagara/Shaders/Private/Stateless/Modules/NiagaraStatelessModule_MeshIndex.ush
+++ b/Engine/Plugins/FX/Niagara/Shaders/Private/Stateless/Modules/NiagaraStatelessModule_MeshIndex.ush
@@ -1,7 +1,5 @@
 // Copyright Epic Games, Inc. All Rights Reserved.
 
-#include "../NiagaraStatelessCommon.ush"
-
 int MeshIndex_Index;
 int MeshIndex_TableOffset;
 int MeshIndex_TableNumElements;	// -1 to avoid added instructions
diff --git a/Engine/Plugins/FX/Niagara/Shaders/Private/Stateless/Modules/NiagaraStatelessModule_MeshRotationRate.ush b/Engine/Plugins/FX/Niagara/Shaders/Private/Stateless/Modules/NiagaraStatelessModule_MeshRotationRate.ush
--- a/Engine/Plugins/FX/Niagara/Shaders/Private/Stateless/Modules/NiagaraStatelessModule_MeshRotationRate.ush
+++ b/Engine/Plugins/FX/Niagara/Shaders/Private/Stateless/Modules/NiagaraStatelessModule_MeshRotationRate.ush
@@ -1,7 +1,5 @@
 // Copyright Epic Games, Inc. All Rights Reserved.
 
-#include "../NiagaraStatelessCommon.ush"
-
 int										MeshRotationRate_ModuleEnabled;
 float3									MeshRotationRate_Scale;
 float3									MeshRotationRate_Bias;
diff --git a/Engine/Plugins/FX/Niagara/Shaders/Private/Stateless/Modules/NiagaraStatelessModule_RotateAroundPoint.ush b/Engine/Plugins/FX/Niagara/Shaders/Private/Stateless/Modules/NiagaraStatelessModule_RotateAroundPoint.ush
--- a/Engine/Plugins/FX/Niagara/Shaders/Private/Stateless/Modules/NiagaraStatelessModule_RotateAroundPoint.ush
+++ b/Engine/Plugins/FX/Niagara/Shaders/Private/Stateless/Modules/NiagaraStatelessModule_RotateAroundPoint.ush
@@ -1,8 +1,5 @@
 // Copyright Epic Games, Inc. All Rights Reserved.
 
-#include "../NiagaraStatelessCommon.ush"
-#include "/Plugin/FX/Niagara/Private/NiagaraQuaternionUtils.ush"
-
 float	RotateAroundPoint_RateScale;
 float	RotateAroundPoint_RateBias;
 float	RotateAroundPoint_RadiusScale;
diff --git a/Engine/Plugins/FX/Niagara/Shaders/Private/Stateless/Modules/NiagaraStatelessModule_ScaleColor.ush b/Engine/Plugins/FX/Niagara/Shaders/Private/Stateless/Modules/NiagaraStatelessModule_ScaleColor.ush
--- a/Engine/Plugins/FX/Niagara/Shaders/Private/Stateless/Modules/NiagaraStatelessModule_ScaleColor.ush
+++ b/Engine/Plugins/FX/Niagara/Shaders/Private/Stateless/Modules/NiagaraStatelessModule_ScaleColor.ush
@@ -1,7 +1,5 @@
 // Copyright Epic Games, Inc. All Rights Reserved.
 
-#include "../NiagaraStatelessCommon.ush"
-
 FNiagaraStatelessBuiltDistributionType	ScaleColor_Distribution;
 
 void ScaleColor_Simulate(inout FStatelessParticle Particle)
diff --git a/Engine/Plugins/FX/Niagara/Shaders/Private/Stateless/Modules/NiagaraStatelessModule_ScaleMeshSize.ush b/Engine/Plugins/FX/Niagara/Shaders/Private/Stateless/Modules/NiagaraStatelessModule_ScaleMeshSize.ush
--- a/Engine/Plugins/FX/Niagara/Shaders/Private/Stateless/Modules/NiagaraStatelessModule_ScaleMeshSize.ush
+++ b/Engine/Plugins/FX/Niagara/Shaders/Private/Stateless/Modules/NiagaraStatelessModule_ScaleMeshSize.ush
@@ -1,7 +1,5 @@
 // Copyright Epic Games, Inc. All Rights Reserved.
 
-#include "../NiagaraStatelessCommon.ush"
-
 FNiagaraStatelessBuiltDistributionType	ScaleMeshSize_Distribution;
 float3	ScaleMeshSize_CurveScale;
 int		ScaleMeshSize_CurveScaleOffset;
diff --git a/Engine/Plugins/FX/Niagara/Shaders/Private/Stateless/Modules/NiagaraStatelessModule_ScaleMeshSizeBySpeed.ush b/Engine/Plugins/FX/Niagara/Shaders/Private/Stateless/Modules/NiagaraStatelessModule_ScaleMeshSizeBySpeed.ush
--- a/Engine/Plugins/FX/Niagara/Shaders/Private/Stateless/Modules/NiagaraStatelessModule_ScaleMeshSizeBySpeed.ush
+++ b/Engine/Plugins/FX/Niagara/Shaders/Private/Stateless/Modules/NiagaraStatelessModule_ScaleMeshSizeBySpeed.ush
@@ -1,7 +1,5 @@
 // Copyright Epic Games, Inc. All Rights Reserved.
 
-#include "../NiagaraStatelessCommon.ush"
-
 float3	ScaleMeshSizeBySpeed_ScaleFactorBias;
 float3	ScaleMeshSizeBySpeed_ScaleFactorScale;
 uint2	ScaleMeshSizeBySpeed_ScaleDistribution;
diff --git a/Engine/Plugins/FX/Niagara/Shaders/Private/Stateless/Modules/NiagaraStatelessModule_ScaleRibbonWidth.ush b/Engine/Plugins/FX/Niagara/Shaders/Private/Stateless/Modules/NiagaraStatelessModule_ScaleRibbonWidth.ush
--- a/Engine/Plugins/FX/Niagara/Shaders/Private/Stateless/Modules/NiagaraStatelessModule_ScaleRibbonWidth.ush
+++ b/Engine/Plugins/FX/Niagara/Shaders/Private/Stateless/Modules/NiagaraStatelessModule_ScaleRibbonWidth.ush
@@ -1,7 +1,5 @@
 // Copyright Epic Games, Inc. All Rights Reserved.
 
-#include "../NiagaraStatelessCommon.ush"
-
 FNiagaraStatelessBuiltDistributionType	ScaleRibbonWidth_Distribution;
 float	ScaleRibbonWidth_CurveScale;
 
diff --git a/Engine/Plugins/FX/Niagara/Shaders/Private/Stateless/Modules/NiagaraStatelessModule_ScaleSpriteSize.ush b/Engine/Plugins/FX/Niagara/Shaders/Private/Stateless/Modules/NiagaraStatelessModule_ScaleSpriteSize.ush
--- a/Engine/Plugins/FX/Niagara/Shaders/Private/Stateless/Modules/NiagaraStatelessModule_ScaleSpriteSize.ush
+++ b/Engine/Plugins/FX/Niagara/Shaders/Private/Stateless/Modules/NiagaraStatelessModule_ScaleSpriteSize.ush
@@ -1,7 +1,5 @@
 // Copyright Epic Games, Inc. All Rights Reserved.
 
-#include "../NiagaraStatelessCommon.ush"
-
 FNiagaraStatelessBuiltDistributionType	ScaleSpriteSize_Distribution;
 float2	ScaleSpriteSize_CurveScale;
 
diff --git a/Engine/Plugins/FX/Niagara/Shaders/Private/Stateless/Modules/NiagaraStatelessModule_ScaleSpriteSizeBySpeed.ush b/Engine/Plugins/FX/Niagara/Shaders/Private/Stateless/Modules/NiagaraStatelessModule_ScaleSpriteSizeBySpeed.ush
--- a/Engine/Plugins/FX/Niagara/Shaders/Private/Stateless/Modules/NiagaraStatelessModule_ScaleSpriteSizeBySpeed.ush
+++ b/Engine/Plugins/FX/Niagara/Shaders/Private/Stateless/Modules/NiagaraStatelessModule_ScaleSpriteSizeBySpeed.ush
@@ -1,7 +1,5 @@
 // Copyright Epic Games, Inc. All Rights Reserved.
 
-#include "../NiagaraStatelessCommon.ush"
-
 float2	ScaleSpriteSizeBySpeed_ScaleFactorBias;
 float2	ScaleSpriteSizeBySpeed_ScaleFactorScale;
 uint2	ScaleSpriteSizeBySpeed_ScaleDistribution;
diff --git a/Engine/Plugins/FX/Niagara/Shaders/Private/Stateless/Modules/NiagaraStatelessModule_ShapeLocation.ush b/Engine/Plugins/FX/Niagara/Shaders/Private/Stateless/Modules/NiagaraStatelessModule_ShapeLocation.ush
--- a/Engine/Plugins/FX/Niagara/Shaders/Private/Stateless/Modules/NiagaraStatelessModule_ShapeLocation.ush
+++ b/Engine/Plugins/FX/Niagara/Shaders/Private/Stateless/Modules/NiagaraStatelessModule_ShapeLocation.ush
@@ -1,8 +1,5 @@
 // Copyright Epic Games, Inc. All Rights Reserved.
 
-#include "../NiagaraStatelessCommon.ush"
-#include "/Plugin/FX/Niagara/Private/NiagaraQuaternionUtils.ush"
-
 uint4	ShapeLocation_Mode;
 float4	ShapeLocation_Parameters0;
 float4	ShapeLocation_Parameters1;
diff --git a/Engine/Plugins/FX/Niagara/Shaders/Private/Stateless/Modules/NiagaraStatelessModule_SolveVelocitiesAndForces.ush b/Engine/Plugins/FX/Niagara/Shaders/Private/Stateless/Modules/NiagaraStatelessModule_SolveVelocitiesAndForces.ush
--- a/Engine/Plugins/FX/Niagara/Shaders/Private/Stateless/Modules/NiagaraStatelessModule_SolveVelocitiesAndForces.ush
+++ b/Engine/Plugins/FX/Niagara/Shaders/Private/Stateless/Modules/NiagaraStatelessModule_SolveVelocitiesAndForces.ush
@@ -1,7 +1,5 @@
 // Copyright Epic Games, Inc. All Rights Reserved.
 
-#include "../NiagaraStatelessCommon.ush"
-
 float	SolveVelocitiesAndForces_MassScale;
 float	SolveVelocitiesAndForces_MassBias;
 float	SolveVelocitiesAndForces_DragScale;
diff --git a/Engine/Plugins/FX/Niagara/Shaders/Private/Stateless/Modules/NiagaraStatelessModule_SpriteFacingAndAlignment.ush b/Engine/Plugins/FX/Niagara/Shaders/Private/Stateless/Modules/NiagaraStatelessModule_SpriteFacingAndAlignment.ush
--- a/Engine/Plugins/FX/Niagara/Shaders/Private/Stateless/Modules/NiagaraStatelessModule_SpriteFacingAndAlignment.ush
+++ b/Engine/Plugins/FX/Niagara/Shaders/Private/Stateless/Modules/NiagaraStatelessModule_SpriteFacingAndAlignment.ush
@@ -1,7 +1,5 @@
 // Copyright Epic Games, Inc. All Rights Reserved.
 
-#include "../NiagaraStatelessCommon.ush"
-
 float3	SpriteFacingAndAlignment_SpriteFacing;
 float3	SpriteFacingAndAlignment_SpriteAlignment;
 
diff --git a/Engine/Plugins/FX/Niagara/Shaders/Private/Stateless/Modules/NiagaraStatelessModule_SpriteRotationRate.ush b/Engine/Plugins/FX/Niagara/Shaders/Private/Stateless/Modules/NiagaraStatelessModule_SpriteRotationRate.ush
--- a/Engine/Plugins/FX/Niagara/Shaders/Private/Stateless/Modules/NiagaraStatelessModule_SpriteRotationRate.ush
+++ b/Engine/Plugins/FX/Niagara/Shaders/Private/Stateless/Modules/NiagaraStatelessModule_SpriteRotationRate.ush
@@ -1,7 +1,5 @@
 // Copyright Epic Games, Inc. All Rights Reserved.
 
-#include "../NiagaraStatelessCommon.ush"
-
 int										SpriteRotationRate_Enabled;
 float									SpriteRotationRate_Scale;
 float									SpriteRotationRate_Bias;
diff --git a/Engine/Plugins/FX/Niagara/Shaders/Private/Stateless/Modules/NiagaraStatelessModule_SubUVAnimation.ush b/Engine/Plugins/FX/Niagara/Shaders/Private/Stateless/Modules/NiagaraStatelessModule_SubUVAnimation.ush
--- a/Engine/Plugins/FX/Niagara/Shaders/Private/Stateless/Modules/NiagaraStatelessModule_SubUVAnimation.ush
+++ b/Engine/Plugins/FX/Niagara/Shaders/Private/Stateless/Modules/NiagaraStatelessModule_SubUVAnimation.ush
@@ -1,7 +1,5 @@
 // Copyright Epic Games, Inc. All Rights Reserved.
 
-#include "../NiagaraStatelessCommon.ush"
-
 int		SubUVAnimation_Mode;
 float	SubUVAnimation_NumFrames;
 float	SubUVAnimation_InitialFrameScale;
diff --git a/Engine/Plugins/FX/Niagara/Shaders/Private/Stateless/NiagaraStatelessCommon.ush b/Engine/Plugins/FX/Niagara/Shaders/Private/Stateless/NiagaraStatelessCommon.ush
--- a/Engine/Plugins/FX/Niagara/Shaders/Private/Stateless/NiagaraStatelessCommon.ush
+++ b/Engine/Plugins/FX/Niagara/Shaders/Private/Stateless/NiagaraStatelessCommon.ush
@@ -70,34 +70,6 @@ float3 SafeNormalize(float3 v)
 	return SafeNormalize(v, float3(1.0f, 0.0f, 0.0f));
 }
 
-////////////////////////////////////////////////////////////////////////////////////
-
-struct FStatelessParticle
-{
-	bool bAlive;					// If the particle should remain alive or not (only evaluated as specific points)
-	uint UniqueIndex;				// Unique particle index, sequential based on previous spawn info accumulation
-	float MaterialRandom;			// Random float passed to materials
-	float Lifetime;					// Overall lifetime for the particle
-	float Age;						// Current age for the particle
-	float NormalizedAge;			// Current normalized age for the particle
-	float PreviousAge;				// Previous age for the particle
-	float PreviousNormalizedAge;	// Previous normalized age for the particle
-
-	float DeltaTime;				// Simulation DT
-	float InvDeltaTime;				// Simulation Inverse DT
-	
-	/////////////////////////////////////////////////////////////////////////////////////////////////////////////////////////
-	// TBD this might produce sub optimal code, needs investigate and is temporary
-	#define PARTICLE_ATTRIBUTE_OUTPUT(TYPE, NAME) TYPE NAME##;
-	#define PARTICLE_ATTRIBUTE_OUTPUT_IMPLICIT(TYPE, NAME)
-	#define PARTICLE_ATTRIBUTE_TRANSIENT(TYPE, NAME) TYPE NAME##;
-		PARTICLE_ATTRIBUTES
-	#undef PARTICLE_ATTRIBUTE
-	#undef PARTICLE_ATTRIBUTE_OUTPUT_IMPLICIT
-	#undef PARTICLE_ATTRIBUTE_TRANSIENT
-	/////////////////////////////////////////////////////////////////////////////////////////////////////////////////////////
-};
-
 static uint4 GRandomSeedInternal;
 
 uint4 RandomUInt4(uint RandomSeedOffset)
diff --git a/Engine/Plugins/FX/Niagara/Shaders/Private/Stateless/NiagaraStatelessSimulationDefault.usf b/Engine/Plugins/FX/Niagara/Shaders/Private/Stateless/NiagaraStatelessSimulationDefault.usf
--- a/Engine/Plugins/FX/Niagara/Shaders/Private/Stateless/NiagaraStatelessSimulationDefault.usf
+++ b/Engine/Plugins/FX/Niagara/Shaders/Private/Stateless/NiagaraStatelessSimulationDefault.usf
@@ -1,16 +1,19 @@
 // Copyright Epic Games, Inc. All Rights Reserved.
 
+#include "/Plugin/FX/Niagara/Private/NiagaraQuaternionUtils.ush"
+#include "/Plugin/FX/Niagara/Private/Stateless/NiagaraStatelessCommon.ush"
+
 // Defines all paricles attributes
 // - PARTICLE_ATTRIBUTE_OUTPUT - Defines an output variable
 // - PARTICLE_ATTRIBUTE_OUTPUT_IMPLICIT - Defines an output variable which is already part of the main particle payload (i.e. always present)
 // - PARTICLE_ATTRIBUTE_TRANSIENT - Defines an output that is temporary and not output
 #define PARTICLE_ATTRIBUTES \
-	PARTICLE_ATTRIBUTE_OUTPUT_IMPLICIT(int,		UniqueIndex) \
+	PARTICLE_ATTRIBUTE_OUTPUT_IMPLICIT(int,		UniqueID) \
 	PARTICLE_ATTRIBUTE_OUTPUT_IMPLICIT(float,	MaterialRandom) \
 	PARTICLE_ATTRIBUTE_OUTPUT(float3,	Position) \
 	PARTICLE_ATTRIBUTE_OUTPUT(float,	CameraOffset) \
 	PARTICLE_ATTRIBUTE_OUTPUT(float4,	Color) \
-	PARTICLE_ATTRIBUTE_OUTPUT(float4,	DynamicMaterialParameter0) \
+	PARTICLE_ATTRIBUTE_OUTPUT(float4,	DynamicMaterialParameter) \
 	PARTICLE_ATTRIBUTE_OUTPUT(float4,	DynamicMaterialParameter1) \
 	PARTICLE_ATTRIBUTE_OUTPUT(float4,	DynamicMaterialParameter2) \
 	PARTICLE_ATTRIBUTE_OUTPUT(float4,	DynamicMaterialParameter3) \
@@ -65,4 +68,26 @@
 	PARTICLE_MODULE(SpriteRotationRate) \
 	PARTICLE_MODULE(SubUVAnimation) \
 
-#include "NiagaraStatelessSimulationTemplate.ush"
+#include "/Plugin/FX/Niagara/Private/Stateless/Template/GenerateCS_PreModules.ush"
+
+//#include "/Plugin/FX/Niagara/Private/Stateless/Modules/NiagaraStatelessModule_CalculateAccurateVelocity.ush"
+#include "/Plugin/FX/Niagara/Private/Stateless/Modules/NiagaraStatelessModule_CameraOffset.ush"
+#include "/Plugin/FX/Niagara/Private/Stateless/Modules/NiagaraStatelessModule_DynamicMaterialParameters.ush"
+#include "/Plugin/FX/Niagara/Private/Stateless/Modules/NiagaraStatelessModule_InitializeParticle.ush"
+#include "/Plugin/FX/Niagara/Private/Stateless/Modules/NiagaraStatelessModule_InitialMeshOrientation.ush"
+#include "/Plugin/FX/Niagara/Private/Stateless/Modules/NiagaraStatelessModule_MeshIndex.ush"
+#include "/Plugin/FX/Niagara/Private/Stateless/Modules/NiagaraStatelessModule_MeshRotationRate.ush"
+//#include "/Plugin/FX/Niagara/Private/Stateless/Modules/NiagaraStatelessModule_RotateAroundPoint.ush"
+#include "/Plugin/FX/Niagara/Private/Stateless/Modules/NiagaraStatelessModule_ShapeLocation.ush"
+#include "/Plugin/FX/Niagara/Private/Stateless/Modules/NiagaraStatelessModule_ScaleColor.ush"
+#include "/Plugin/FX/Niagara/Private/Stateless/Modules/NiagaraStatelessModule_ScaleRibbonWidth.ush"
+#include "/Plugin/FX/Niagara/Private/Stateless/Modules/NiagaraStatelessModule_ScaleSpriteSize.ush"
+#include "/Plugin/FX/Niagara/Private/Stateless/Modules/NiagaraStatelessModule_ScaleSpriteSizeBySpeed.ush"
+#include "/Plugin/FX/Niagara/Private/Stateless/Modules/NiagaraStatelessModule_ScaleMeshSize.ush"
+#include "/Plugin/FX/Niagara/Private/Stateless/Modules/NiagaraStatelessModule_ScaleMeshSizebySpeed.ush"
+#include "/Plugin/FX/Niagara/Private/Stateless/Modules/NiagaraStatelessModule_SolveVelocitiesAndForces.ush"
+#include "/Plugin/FX/Niagara/Private/Stateless/Modules/NiagaraStatelessModule_SpriteFacingAndAlignment.ush"
+#include "/Plugin/FX/Niagara/Private/Stateless/Modules/NiagaraStatelessModule_SpriteRotationRate.ush"
+#include "/Plugin/FX/Niagara/Private/Stateless/Modules/NiagaraStatelessModule_SubUVAnimation.ush"
+
+#include "/Plugin/FX/Niagara/Private/Stateless/Template/GenerateCS_PostModules.ush"
diff --git a/Engine/Plugins/FX/Niagara/Shaders/Private/Stateless/NiagaraStatelessSimulationExample1.usf b/Engine/Plugins/FX/Niagara/Shaders/Private/Stateless/NiagaraStatelessSimulationExample1.usf
deleted file mode 100644
diff --git a/Engine/Plugins/FX/Niagara/Shaders/Private/Stateless/NiagaraStatelessSimulationTemplate.ush b/Engine/Plugins/FX/Niagara/Shaders/Private/Stateless/NiagaraStatelessSimulationTemplate.ush
deleted file mode 100644
diff --git a/Engine/Plugins/FX/Niagara/Shaders/Private/Stateless/Template/GenerateCS_PostModules.ush b/Engine/Plugins/FX/Niagara/Shaders/Private/Stateless/Template/GenerateCS_PostModules.ush
new file mode 100644
--- /dev/null
+++ b/Engine/Plugins/FX/Niagara/Shaders/Private/Stateless/Template/GenerateCS_PostModules.ush
@@ -0,0 +1,109 @@
+// Copyright Epic Games, Inc. All Rights Reserved.
+
+#include "/Engine/Private/ComputeShaderUtils.ush"
+
+uint	Common_RandomSeed;
+float	Common_SimulationTime;
+float	Common_SimulationDeltaTime;
+float	Common_SimulationInvDeltaTime;
+
+uint				Common_GPUCountBufferOffset;
+RWBuffer<int>		Common_GPUCountBuffer;
+
+DECLARE_SCALAR_ARRAY(uint,	SpawnInfo_NumActive, NIAGARA_MAX_GPU_SPAWN_INFOS);
+DECLARE_SCALAR_ARRAY(uint,	SpawnInfo_ParticleOffset, NIAGARA_MAX_GPU_SPAWN_INFOS);
+DECLARE_SCALAR_ARRAY(uint,	SpawnInfo_UniqueOffset, NIAGARA_MAX_GPU_SPAWN_INFOS);
+DECLARE_SCALAR_ARRAY(float,	SpawnInfo_Time, NIAGARA_MAX_GPU_SPAWN_INFOS);
+DECLARE_SCALAR_ARRAY(float,	SpawnInfo_Rate, NIAGARA_MAX_GPU_SPAWN_INFOS);
+DECLARE_SCALAR_ARRAY(float,	SpawnInfo_LifetimeScale, NIAGARA_MAX_GPU_SPAWN_INFOS);
+DECLARE_SCALAR_ARRAY(float,	SpawnInfo_LifetimeBias, NIAGARA_MAX_GPU_SPAWN_INFOS);
+
+/////////////////////////////////////////////////////////////////////////////////////////////////////////////////////////
+
+[numthreads(THREADGROUP_SIZE, 1, 1)]
+void StatelessMain(uint3 GroupId : SV_GroupID, int GroupThreadIndex : SV_GroupIndex)
+{
+	// Convert thread ID into linear thread
+	uint LinearThreadId = GetUnWrappedDispatchThreadId(GroupId, GroupThreadIndex, THREADGROUP_SIZE);
+
+	// Initialize our particle data
+	FStatelessParticle Particle = (FStatelessParticle)0;
+	Particle.Age = -1.0f;	
+
+	// Find which spawn info we belong to and fill out details
+	uint ModuleSeedOffset = 0;
+	{
+		uint SpawnInfoIndex = LinearThreadId;
+		for ( int i=0; i < NIAGARA_MAX_GPU_SPAWN_INFOS; ++i )
+		{
+			const uint SpawnInfoNumActive		= GET_SCALAR_ARRAY_ELEMENT(SpawnInfo_NumActive, i);
+			const uint SpawnInfoParticleOffset	= GET_SCALAR_ARRAY_ELEMENT(SpawnInfo_ParticleOffset, i);
+			const uint SpawnInfoUniqueOffset	= GET_SCALAR_ARRAY_ELEMENT(SpawnInfo_UniqueOffset, i);
+			const float SpawnInfoTime			= GET_SCALAR_ARRAY_ELEMENT(SpawnInfo_Time, i);
+			const float SpawnInfoRate			= GET_SCALAR_ARRAY_ELEMENT(SpawnInfo_Rate, i);
+			const float SpawnInfoLifetimeScale	= GET_SCALAR_ARRAY_ELEMENT(SpawnInfo_LifetimeScale, i);
+			const float SpawnInfoLifetimeBias	= GET_SCALAR_ARRAY_ELEMENT(SpawnInfo_LifetimeBias, i);
+			if ( SpawnInfoIndex < SpawnInfoNumActive )
+			{
+				const uint SpawnParticleIndex = SpawnInfoIndex + SpawnInfoParticleOffset;
+				Particle.UniqueID		= SpawnInfoIndex + SpawnInfoUniqueOffset + SpawnInfoParticleOffset;
+				Particle.Age			= Common_SimulationTime - (SpawnInfoTime + float(SpawnParticleIndex) * SpawnInfoRate);
+				GRandomSeedInternal		= FNiagaraStatelessDefinitions::MakeRandomSeed(Common_RandomSeed, Particle.UniqueID, ModuleSeedOffset++, 0);
+				Particle.Lifetime		= RandomScaleBiasFloat(0, SpawnInfoLifetimeScale, SpawnInfoLifetimeBias);
+				break;
+			}
+
+			SpawnInfoIndex -= SpawnInfoNumActive;
+		}
+		if ( Particle.Age < 0.0f )
+		{
+			return;
+		}
+	}
+	
+	// Initialize variables / determine lifetime / etc
+	if ( Particle.Lifetime <= 0.0f || Particle.Age >= Particle.Lifetime )
+	{
+		return;
+	}
+
+	Particle.NormalizedAge	= Particle.Age / Particle.Lifetime;
+
+	Particle.PreviousAge			= max(Particle.Age - Common_SimulationDeltaTime, 0.0f);
+	Particle.PreviousNormalizedAge	= Particle.PreviousAge / Particle.Lifetime;
+
+	Particle.DeltaTime		= Common_SimulationDeltaTime;
+	Particle.InvDeltaTime	= Common_SimulationInvDeltaTime;
+	
+	Particle.bAlive = true;
+
+	Particle.MaterialRandom	= RandomFloat(1);
+	
+	// Run Particle Simulate
+	#define PARTICLE_MODULE(NAME) \
+		GRandomSeedInternal = FNiagaraStatelessDefinitions::MakeRandomSeed(Common_RandomSeed, Particle.UniqueID, ModuleSeedOffset++, 0); \
+		##NAME##_Simulate(Particle);
+
+		PARTICLE_MODULES
+
+	#undef PARTICLE_MODULE
+
+	if ( !Particle.bAlive )
+	{
+		return;
+	}
+	
+	// Output Particles that are still alive
+	uint OutputIndex;
+	
+	//WaveInterlockedAddScalar_(Common_GPUCountBuffer[Common_GPUCountBufferOffset], 1, Particle.OutputIndex);
+	InterlockedAdd(Common_GPUCountBuffer[Common_GPUCountBufferOffset], 1, OutputIndex);
+
+	#define PARTICLE_ATTRIBUTE_OUTPUT(TYPE, NAME) if ( IsValidComponent(Permutation_##NAME##Component) ) { OutputComponentData(OutputIndex, Permutation_##NAME##Component, Particle.##NAME); }
+	#define PARTICLE_ATTRIBUTE_OUTPUT_IMPLICIT(TYPE, NAME) if ( IsValidComponent(Permutation_##NAME##Component) ) { OutputComponentData(OutputIndex, Permutation_##NAME##Component, Particle.##NAME); }
+	#define PARTICLE_ATTRIBUTE_TRANSIENT(TYPE, NAME)
+		PARTICLE_ATTRIBUTES
+	#undef PARTICLE_ATTRIBUTE
+	#undef PARTICLE_ATTRIBUTE_OUTPUT_IMPLICIT
+	#undef PARTICLE_ATTRIBUTE_TRANSIENT
+}
diff --git a/Engine/Plugins/FX/Niagara/Shaders/Private/Stateless/Template/GenerateCS_PreModules.ush b/Engine/Plugins/FX/Niagara/Shaders/Private/Stateless/Template/GenerateCS_PreModules.ush
new file mode 100644
--- /dev/null
+++ b/Engine/Plugins/FX/Niagara/Shaders/Private/Stateless/Template/GenerateCS_PreModules.ush
@@ -0,0 +1,47 @@
+// Copyright Epic Games, Inc. All Rights Reserved.
+
+/////////////////////////////////////////////////////////////////////////////////////////////////////////////////////////
+
+struct FStatelessParticle
+{
+	bool bAlive;					// If the particle should remain alive or not (only evaluated as specific points)
+	uint UniqueID;					// Unique particle index, sequential based on previous spawn info accumulation
+	float MaterialRandom;			// Random float passed to materials
+	float Lifetime;					// Overall lifetime for the particle
+	float Age;						// Current age for the particle
+	float NormalizedAge;			// Current normalized age for the particle
+	float PreviousAge;				// Previous age for the particle
+	float PreviousNormalizedAge;	// Previous normalized age for the particle
+
+	float DeltaTime;				// Simulation DT
+	float InvDeltaTime;				// Simulation Inverse DT
+	
+	/////////////////////////////////////////////////////////////////////////////////////////////////////////////////////////
+	// TBD this might produce sub optimal code, needs investigate and is temporary
+	#define PARTICLE_ATTRIBUTE_OUTPUT(TYPE, NAME) TYPE NAME##;
+	#define PARTICLE_ATTRIBUTE_OUTPUT_IMPLICIT(TYPE, NAME)
+	#define PARTICLE_ATTRIBUTE_TRANSIENT(TYPE, NAME) TYPE NAME##;
+		PARTICLE_ATTRIBUTES
+	#undef PARTICLE_ATTRIBUTE
+	#undef PARTICLE_ATTRIBUTE_OUTPUT_IMPLICIT
+	#undef PARTICLE_ATTRIBUTE_TRANSIENT
+	/////////////////////////////////////////////////////////////////////////////////////////////////////////////////////////
+};
+
+/////////////////////////////////////////////////////////////////////////////////////////////////////////////////////////
+
+#define PARTICLE_ATTRIBUTE_OUTPUT(TYPE, NAME) int Permutation_##NAME##Component;
+#define PARTICLE_ATTRIBUTE_OUTPUT_IMPLICIT(TYPE, NAME) int Permutation_##NAME##Component;
+#define PARTICLE_ATTRIBUTE_TRANSIENT(TYPE, NAME)
+	PARTICLE_ATTRIBUTES
+#undef PARTICLE_ATTRIBUTE
+#undef PARTICLE_ATTRIBUTE_OUTPUT_IMPLICIT
+#undef PARTICLE_ATTRIBUTE_TRANSIENT
+
+/////////////////////////////////////////////////////////////////////////////////////////////////////////////////////////
+
+void Null_Simulate(inout FStatelessParticle Particle)
+{
+}
+
+/////////////////////////////////////////////////////////////////////////////////////////////////////////////////////////
diff --git a/Engine/Plugins/FX/Niagara/Shaders/Shared/NiagaraStatelessDefinitions.h b/Engine/Plugins/FX/Niagara/Shaders/Shared/NiagaraStatelessDefinitions.h
--- a/Engine/Plugins/FX/Niagara/Shaders/Shared/NiagaraStatelessDefinitions.h
+++ b/Engine/Plugins/FX/Niagara/Shaders/Shared/NiagaraStatelessDefinitions.h
@@ -9,7 +9,7 @@
 // Core random functionality
 struct FNiagaraStatelessDefinitions
 {
-	static uint4 MakeRandomSeed(uint EmitterSeed, uint UniqueIndex, uint ModuleOffset, uint CallOffset) { return uint4(7123u + CallOffset, EmitterSeed, 3581u + UniqueIndex, 6405u + ModuleOffset); }
+	static uint4 MakeRandomSeed(uint EmitterSeed, uint UniqueID, uint ModuleOffset, uint CallOffset) { return uint4(7123u + CallOffset, EmitterSeed, 3581u + UniqueID, 6405u + ModuleOffset); }
 	static uint4 OffsetRandomSeedForCall(uint4 Seed, uint CallOffset) { Seed[0] += CallOffset; return Seed; }
 };
 
diff --git a/Engine/Plugins/FX/Niagara/Source/Niagara/Internal/Stateless/Modules/NiagaraStatelessModule_CalculateAccurateVelocity.h b/Engine/Plugins/FX/Niagara/Source/Niagara/Internal/Stateless/Modules/NiagaraStatelessModule_CalculateAccurateVelocity.h
--- a/Engine/Plugins/FX/Niagara/Source/Niagara/Internal/Stateless/Modules/NiagaraStatelessModule_CalculateAccurateVelocity.h
+++ b/Engine/Plugins/FX/Niagara/Source/Niagara/Internal/Stateless/Modules/NiagaraStatelessModule_CalculateAccurateVelocity.h
@@ -17,6 +17,7 @@ public:
 	virtual void SetShaderParameters(const FNiagaraStatelessSetShaderParameterContext& SetShaderParameterContext) const override;
 
 #if WITH_EDITORONLY_DATA
-	virtual void GetOutputVariables(TArray<FNiagaraVariableBase>& OutVariables) const override;
+	virtual const TCHAR* GetShaderTemplatePath() const override;
+	virtual void GetOutputVariables(TArray<FNiagaraVariableBase>& OutVariables, EVariableFilter Filter) const override;
 #endif
 };
diff --git a/Engine/Plugins/FX/Niagara/Source/Niagara/Internal/Stateless/Modules/NiagaraStatelessModule_CameraOffset.h b/Engine/Plugins/FX/Niagara/Source/Niagara/Internal/Stateless/Modules/NiagaraStatelessModule_CameraOffset.h
--- a/Engine/Plugins/FX/Niagara/Source/Niagara/Internal/Stateless/Modules/NiagaraStatelessModule_CameraOffset.h
+++ b/Engine/Plugins/FX/Niagara/Source/Niagara/Internal/Stateless/Modules/NiagaraStatelessModule_CameraOffset.h
@@ -86,7 +86,12 @@ public:
 	virtual bool CanDisableModule() const override { return true; }
 #endif
 #if WITH_EDITORONLY_DATA
-	virtual void GetOutputVariables(TArray<FNiagaraVariableBase>& OutVariables) const override
+	virtual const TCHAR* GetShaderTemplatePath() const override
+	{
+		return TEXT("/Plugin/FX/Niagara/Private/Stateless/Modules/NiagaraStatelessModule_CameraOffset.ush");
+	}
+
+	virtual void GetOutputVariables(TArray<FNiagaraVariableBase>& OutVariables, EVariableFilter Filter) const override
 	{
 		const FNiagaraStatelessGlobals& StatelessGlobals = FNiagaraStatelessGlobals::Get();
 		OutVariables.AddUnique(StatelessGlobals.CameraOffsetVariable);
diff --git a/Engine/Plugins/FX/Niagara/Source/Niagara/Internal/Stateless/Modules/NiagaraStatelessModule_DecalAttributes.h b/Engine/Plugins/FX/Niagara/Source/Niagara/Internal/Stateless/Modules/NiagaraStatelessModule_DecalAttributes.h
--- a/Engine/Plugins/FX/Niagara/Source/Niagara/Internal/Stateless/Modules/NiagaraStatelessModule_DecalAttributes.h
+++ b/Engine/Plugins/FX/Niagara/Source/Niagara/Internal/Stateless/Modules/NiagaraStatelessModule_DecalAttributes.h
@@ -45,6 +45,6 @@ public:
 	virtual bool CanDisableModule() const override { return true; }
 #endif
 #if WITH_EDITORONLY_DATA
-	virtual void GetOutputVariables(TArray<FNiagaraVariableBase>& OutVariables) const override;
+	virtual void GetOutputVariables(TArray<FNiagaraVariableBase>& OutVariables, EVariableFilter Filter) const override;
 #endif
 };
diff --git a/Engine/Plugins/FX/Niagara/Source/Niagara/Internal/Stateless/Modules/NiagaraStatelessModule_DynamicMaterialParameters.h b/Engine/Plugins/FX/Niagara/Source/Niagara/Internal/Stateless/Modules/NiagaraStatelessModule_DynamicMaterialParameters.h
--- a/Engine/Plugins/FX/Niagara/Source/Niagara/Internal/Stateless/Modules/NiagaraStatelessModule_DynamicMaterialParameters.h
+++ b/Engine/Plugins/FX/Niagara/Source/Niagara/Internal/Stateless/Modules/NiagaraStatelessModule_DynamicMaterialParameters.h
@@ -267,13 +267,18 @@ public:
 	virtual bool CanDisableModule() const override { return true; }
 #endif
 #if WITH_EDITORONLY_DATA
-	virtual void GetOutputVariables(TArray<FNiagaraVariableBase>& OutVariables) const override
+	virtual const TCHAR* GetShaderTemplatePath() const override
+	{
+		return TEXT("/Plugin/FX/Niagara/Private/Stateless/Modules/NiagaraStatelessModule_DynamicMaterialParameters.ush");
+	}
+
+	virtual void GetOutputVariables(TArray<FNiagaraVariableBase>& OutVariables, EVariableFilter Filter) const override
 	{
 		const FNiagaraStatelessGlobals& StatelessGlobals = FNiagaraStatelessGlobals::Get();
 
 		for (int32 i = 0; i < NumParameters; ++i)
 		{
-			if (GetParameterChannelMask(i) != 0)
+			if (GetParameterChannelMask(i) != 0 || Filter == EVariableFilter::None)
 			{
 				OutVariables.AddUnique(GetParameterVariable(i));
 			}
diff --git a/Engine/Plugins/FX/Niagara/Source/Niagara/Internal/Stateless/Modules/NiagaraStatelessModule_InitialMeshOrientation.h b/Engine/Plugins/FX/Niagara/Source/Niagara/Internal/Stateless/Modules/NiagaraStatelessModule_InitialMeshOrientation.h
--- a/Engine/Plugins/FX/Niagara/Source/Niagara/Internal/Stateless/Modules/NiagaraStatelessModule_InitialMeshOrientation.h
+++ b/Engine/Plugins/FX/Niagara/Source/Niagara/Internal/Stateless/Modules/NiagaraStatelessModule_InitialMeshOrientation.h
@@ -53,7 +53,8 @@ public:
 	virtual bool CanDisableModule() const override { return true; }
 #endif
 #if WITH_EDITORONLY_DATA
-	virtual void GetOutputVariables(TArray<FNiagaraVariableBase>& OutVariables) const override;
+	virtual const TCHAR* GetShaderTemplatePath() const override;
+	virtual void GetOutputVariables(TArray<FNiagaraVariableBase>& OutVariables, EVariableFilter Filter) const override;
 	virtual void PostLoad() override;
 
 private:
diff --git a/Engine/Plugins/FX/Niagara/Source/Niagara/Internal/Stateless/Modules/NiagaraStatelessModule_InitializeParticle.h b/Engine/Plugins/FX/Niagara/Source/Niagara/Internal/Stateless/Modules/NiagaraStatelessModule_InitializeParticle.h
--- a/Engine/Plugins/FX/Niagara/Source/Niagara/Internal/Stateless/Modules/NiagaraStatelessModule_InitializeParticle.h
+++ b/Engine/Plugins/FX/Niagara/Source/Niagara/Internal/Stateless/Modules/NiagaraStatelessModule_InitializeParticle.h
@@ -168,7 +168,12 @@ public:
 	}
 
 #if WITH_EDITORONLY_DATA
-	virtual void GetOutputVariables(TArray<FNiagaraVariableBase>& OutVariables) const override
+	virtual const TCHAR* GetShaderTemplatePath() const override
+	{
+		return TEXT("/Plugin/FX/Niagara/Private/Stateless/Modules/NiagaraStatelessModule_InitializeParticle.ush");
+	}
+
+	virtual void GetOutputVariables(TArray<FNiagaraVariableBase>& OutVariables, EVariableFilter Filter) const override
 	{
 		const FNiagaraStatelessGlobals& StatelessGlobals = FNiagaraStatelessGlobals::Get();
 		OutVariables.AddUnique(StatelessGlobals.PositionVariable);
@@ -182,7 +187,7 @@ public:
 		OutVariables.AddUnique(StatelessGlobals.PreviousSpriteRotationVariable);
 		OutVariables.AddUnique(StatelessGlobals.PreviousScaleVariable);
 
-		if (bWriteRibbonWidth)
+		if (bWriteRibbonWidth || Filter == EVariableFilter::None)
 		{
 			OutVariables.AddUnique(StatelessGlobals.RibbonWidthVariable);
 			OutVariables.AddUnique(StatelessGlobals.PreviousRibbonWidthVariable);
diff --git a/Engine/Plugins/FX/Niagara/Source/Niagara/Internal/Stateless/Modules/NiagaraStatelessModule_LightAttributes.h b/Engine/Plugins/FX/Niagara/Source/Niagara/Internal/Stateless/Modules/NiagaraStatelessModule_LightAttributes.h
--- a/Engine/Plugins/FX/Niagara/Source/Niagara/Internal/Stateless/Modules/NiagaraStatelessModule_LightAttributes.h
+++ b/Engine/Plugins/FX/Niagara/Source/Niagara/Internal/Stateless/Modules/NiagaraStatelessModule_LightAttributes.h
@@ -51,6 +51,6 @@ public:
 	virtual bool CanDisableModule() const override { return true; }
 #endif
 #if WITH_EDITORONLY_DATA
-	virtual void GetOutputVariables(TArray<FNiagaraVariableBase>& OutVariables) const override;
+	virtual void GetOutputVariables(TArray<FNiagaraVariableBase>& OutVariables, EVariableFilter Filter) const override;
 #endif
 };
diff --git a/Engine/Plugins/FX/Niagara/Source/Niagara/Internal/Stateless/Modules/NiagaraStatelessModule_MeshIndex.h b/Engine/Plugins/FX/Niagara/Source/Niagara/Internal/Stateless/Modules/NiagaraStatelessModule_MeshIndex.h
--- a/Engine/Plugins/FX/Niagara/Source/Niagara/Internal/Stateless/Modules/NiagaraStatelessModule_MeshIndex.h
+++ b/Engine/Plugins/FX/Niagara/Source/Niagara/Internal/Stateless/Modules/NiagaraStatelessModule_MeshIndex.h
@@ -35,10 +35,7 @@ public:
 	virtual bool CanDisableModule() const override { return true; }
 #endif
 #if WITH_EDITORONLY_DATA
-	virtual void GetOutputVariables(TArray<FNiagaraVariableBase>& OutVariables) const override
-	{
-		const FNiagaraStatelessGlobals& StatelessGlobals = FNiagaraStatelessGlobals::Get();
-		OutVariables.AddUnique(StatelessGlobals.MeshIndexVariable);
-	}
+	virtual const TCHAR* GetShaderTemplatePath() const override;
+	virtual void GetOutputVariables(TArray<FNiagaraVariableBase>& OutVariables, EVariableFilter Filter) const override;
 #endif
 };
diff --git a/Engine/Plugins/FX/Niagara/Source/Niagara/Internal/Stateless/Modules/NiagaraStatelessModule_MeshRotationRate.h b/Engine/Plugins/FX/Niagara/Source/Niagara/Internal/Stateless/Modules/NiagaraStatelessModule_MeshRotationRate.h
--- a/Engine/Plugins/FX/Niagara/Source/Niagara/Internal/Stateless/Modules/NiagaraStatelessModule_MeshRotationRate.h
+++ b/Engine/Plugins/FX/Niagara/Source/Niagara/Internal/Stateless/Modules/NiagaraStatelessModule_MeshRotationRate.h
@@ -113,7 +113,12 @@ public:
 	virtual bool CanDisableModule() const override { return true; }
 #endif
 #if WITH_EDITORONLY_DATA
-	virtual void GetOutputVariables(TArray<FNiagaraVariableBase>& OutVariables) const override
+	virtual const TCHAR* GetShaderTemplatePath() const override
+	{
+		return TEXT("/Plugin/FX/Niagara/Private/Stateless/Modules/NiagaraStatelessModule_MeshRotationRate.ush");
+	}
+
+	virtual void GetOutputVariables(TArray<FNiagaraVariableBase>& OutVariables, EVariableFilter Filter) const override
 	{
 		const FNiagaraStatelessGlobals& StatelessGlobals = FNiagaraStatelessGlobals::Get();
 		OutVariables.AddUnique(StatelessGlobals.MeshOrientationVariable);
diff --git a/Engine/Plugins/FX/Niagara/Source/Niagara/Internal/Stateless/Modules/NiagaraStatelessModule_RotateAroundPoint.h b/Engine/Plugins/FX/Niagara/Source/Niagara/Internal/Stateless/Modules/NiagaraStatelessModule_RotateAroundPoint.h
--- a/Engine/Plugins/FX/Niagara/Source/Niagara/Internal/Stateless/Modules/NiagaraStatelessModule_RotateAroundPoint.h
+++ b/Engine/Plugins/FX/Niagara/Source/Niagara/Internal/Stateless/Modules/NiagaraStatelessModule_RotateAroundPoint.h
@@ -43,6 +43,7 @@ public:
 	virtual void DrawDebug(const FNiagaraStatelessDrawDebugContext& DrawDebugContext) const override;
 #endif
 #if WITH_EDITORONLY_DATA
-	virtual void GetOutputVariables(TArray<FNiagaraVariableBase>& OutVariables) const override;
+	virtual const TCHAR* GetShaderTemplatePath() const override;
+	virtual void GetOutputVariables(TArray<FNiagaraVariableBase>& OutVariables, EVariableFilter Filter) const override;
 #endif
 };
diff --git a/Engine/Plugins/FX/Niagara/Source/Niagara/Internal/Stateless/Modules/NiagaraStatelessModule_ScaleColor.h b/Engine/Plugins/FX/Niagara/Source/Niagara/Internal/Stateless/Modules/NiagaraStatelessModule_ScaleColor.h
--- a/Engine/Plugins/FX/Niagara/Source/Niagara/Internal/Stateless/Modules/NiagaraStatelessModule_ScaleColor.h
+++ b/Engine/Plugins/FX/Niagara/Source/Niagara/Internal/Stateless/Modules/NiagaraStatelessModule_ScaleColor.h
@@ -82,7 +82,12 @@ public:
 	virtual bool CanDisableModule() const override { return true; }
 #endif
 #if WITH_EDITORONLY_DATA
-	virtual void GetOutputVariables(TArray<FNiagaraVariableBase>& OutVariables) const override
+	virtual const TCHAR* GetShaderTemplatePath() const override
+	{
+		return TEXT("/Plugin/FX/Niagara/Private/Stateless/Modules/NiagaraStatelessModule_ScaleColor.ush");
+	}
+
+	virtual void GetOutputVariables(TArray<FNiagaraVariableBase>& OutVariables, EVariableFilter Filter) const override
 	{
 		const FNiagaraStatelessGlobals& StatelessGlobals = FNiagaraStatelessGlobals::Get();
 		OutVariables.AddUnique(StatelessGlobals.ColorVariable);
diff --git a/Engine/Plugins/FX/Niagara/Source/Niagara/Internal/Stateless/Modules/NiagaraStatelessModule_ScaleMeshSize.h b/Engine/Plugins/FX/Niagara/Source/Niagara/Internal/Stateless/Modules/NiagaraStatelessModule_ScaleMeshSize.h
--- a/Engine/Plugins/FX/Niagara/Source/Niagara/Internal/Stateless/Modules/NiagaraStatelessModule_ScaleMeshSize.h
+++ b/Engine/Plugins/FX/Niagara/Source/Niagara/Internal/Stateless/Modules/NiagaraStatelessModule_ScaleMeshSize.h
@@ -124,7 +124,12 @@ public:
 	virtual bool CanDisableModule() const override { return true; }
 #endif
 #if WITH_EDITORONLY_DATA
-	virtual void GetOutputVariables(TArray<FNiagaraVariableBase>& OutVariables) const override
+	virtual const TCHAR* GetShaderTemplatePath() const override
+	{
+		return TEXT("/Plugin/FX/Niagara/Private/Stateless/Modules/NiagaraStatelessModule_ScaleMeshSize.ush");
+	}
+
+	virtual void GetOutputVariables(TArray<FNiagaraVariableBase>& OutVariables, EVariableFilter Filter) const override
 	{
 		const FNiagaraStatelessGlobals& StatelessGlobals = FNiagaraStatelessGlobals::Get();
 		OutVariables.AddUnique(StatelessGlobals.ScaleVariable);
diff --git a/Engine/Plugins/FX/Niagara/Source/Niagara/Internal/Stateless/Modules/NiagaraStatelessModule_ScaleMeshSizeBySpeed.h b/Engine/Plugins/FX/Niagara/Source/Niagara/Internal/Stateless/Modules/NiagaraStatelessModule_ScaleMeshSizeBySpeed.h
--- a/Engine/Plugins/FX/Niagara/Source/Niagara/Internal/Stateless/Modules/NiagaraStatelessModule_ScaleMeshSizeBySpeed.h
+++ b/Engine/Plugins/FX/Niagara/Source/Niagara/Internal/Stateless/Modules/NiagaraStatelessModule_ScaleMeshSizeBySpeed.h
@@ -42,6 +42,7 @@ public:
 	virtual bool CanDisableModule() const override { return true; }
 #endif
 #if WITH_EDITORONLY_DATA
-	virtual void GetOutputVariables(TArray<FNiagaraVariableBase>& OutVariables) const override;
+	virtual const TCHAR* GetShaderTemplatePath() const override;
+	virtual void GetOutputVariables(TArray<FNiagaraVariableBase>& OutVariables, EVariableFilter Filter) const override;
 #endif
 };
diff --git a/Engine/Plugins/FX/Niagara/Source/Niagara/Internal/Stateless/Modules/NiagaraStatelessModule_ScaleRibbonWidth.h b/Engine/Plugins/FX/Niagara/Source/Niagara/Internal/Stateless/Modules/NiagaraStatelessModule_ScaleRibbonWidth.h
--- a/Engine/Plugins/FX/Niagara/Source/Niagara/Internal/Stateless/Modules/NiagaraStatelessModule_ScaleRibbonWidth.h
+++ b/Engine/Plugins/FX/Niagara/Source/Niagara/Internal/Stateless/Modules/NiagaraStatelessModule_ScaleRibbonWidth.h
@@ -121,7 +121,12 @@ public:
 	virtual bool CanDisableModule() const override { return true; }
 #endif
 #if WITH_EDITORONLY_DATA
-	virtual void GetOutputVariables(TArray<FNiagaraVariableBase>& OutVariables) const override
+	virtual const TCHAR* GetShaderTemplatePath() const override
+	{
+		return TEXT("/Plugin/FX/Niagara/Private/Stateless/Modules/NiagaraStatelessModule_ScaleRibbonWidth.ush");
+	}
+
+	virtual void GetOutputVariables(TArray<FNiagaraVariableBase>& OutVariables, EVariableFilter Filter) const override
 	{
 		const FNiagaraStatelessGlobals& StatelessGlobals = FNiagaraStatelessGlobals::Get();
 		OutVariables.AddUnique(StatelessGlobals.RibbonWidthVariable);
diff --git a/Engine/Plugins/FX/Niagara/Source/Niagara/Internal/Stateless/Modules/NiagaraStatelessModule_ScaleSpriteSize.h b/Engine/Plugins/FX/Niagara/Source/Niagara/Internal/Stateless/Modules/NiagaraStatelessModule_ScaleSpriteSize.h
--- a/Engine/Plugins/FX/Niagara/Source/Niagara/Internal/Stateless/Modules/NiagaraStatelessModule_ScaleSpriteSize.h
+++ b/Engine/Plugins/FX/Niagara/Source/Niagara/Internal/Stateless/Modules/NiagaraStatelessModule_ScaleSpriteSize.h
@@ -121,7 +121,12 @@ public:
 	virtual bool CanDisableModule() const override { return true; }
 #endif
 #if WITH_EDITORONLY_DATA
-	virtual void GetOutputVariables(TArray<FNiagaraVariableBase>& OutVariables) const override
+	virtual const TCHAR* GetShaderTemplatePath() const override
+	{
+		return TEXT("/Plugin/FX/Niagara/Private/Stateless/Modules/NiagaraStatelessModule_ScaleSpriteSize.ush");
+	}
+
+	virtual void GetOutputVariables(TArray<FNiagaraVariableBase>& OutVariables, EVariableFilter Filter) const override
 	{
 		const FNiagaraStatelessGlobals& StatelessGlobals = FNiagaraStatelessGlobals::Get();
 		OutVariables.AddUnique(StatelessGlobals.SpriteSizeVariable);
diff --git a/Engine/Plugins/FX/Niagara/Source/Niagara/Internal/Stateless/Modules/NiagaraStatelessModule_ScaleSpriteSizeBySpeed.h b/Engine/Plugins/FX/Niagara/Source/Niagara/Internal/Stateless/Modules/NiagaraStatelessModule_ScaleSpriteSizeBySpeed.h
--- a/Engine/Plugins/FX/Niagara/Source/Niagara/Internal/Stateless/Modules/NiagaraStatelessModule_ScaleSpriteSizeBySpeed.h
+++ b/Engine/Plugins/FX/Niagara/Source/Niagara/Internal/Stateless/Modules/NiagaraStatelessModule_ScaleSpriteSizeBySpeed.h
@@ -55,6 +55,7 @@ public:
 	virtual bool CanDisableModule() const override { return true; }
 #endif
 #if WITH_EDITORONLY_DATA
-	virtual void GetOutputVariables(TArray<FNiagaraVariableBase>& OutVariables) const override;
+	virtual const TCHAR* GetShaderTemplatePath() const override;
+	virtual void GetOutputVariables(TArray<FNiagaraVariableBase>& OutVariables, EVariableFilter Filter) const override;
 #endif
 };
diff --git a/Engine/Plugins/FX/Niagara/Source/Niagara/Internal/Stateless/Modules/NiagaraStatelessModule_ShapeLocation.h b/Engine/Plugins/FX/Niagara/Source/Niagara/Internal/Stateless/Modules/NiagaraStatelessModule_ShapeLocation.h
--- a/Engine/Plugins/FX/Niagara/Source/Niagara/Internal/Stateless/Modules/NiagaraStatelessModule_ShapeLocation.h
+++ b/Engine/Plugins/FX/Niagara/Source/Niagara/Internal/Stateless/Modules/NiagaraStatelessModule_ShapeLocation.h
@@ -99,12 +99,8 @@ public:
 	virtual void DrawDebug(const FNiagaraStatelessDrawDebugContext& DrawDebugContext) const override;
 #endif
 #if WITH_EDITORONLY_DATA
-	virtual void GetOutputVariables(TArray<FNiagaraVariableBase>& OutVariables) const override
-	{
-		const FNiagaraStatelessGlobals& StatelessGlobals = FNiagaraStatelessGlobals::Get();
-		OutVariables.AddUnique(StatelessGlobals.PositionVariable);
-		OutVariables.AddUnique(StatelessGlobals.PreviousPositionVariable);
-	}
+	virtual const TCHAR* GetShaderTemplatePath() const override;
+	virtual void GetOutputVariables(TArray<FNiagaraVariableBase>& OutVariables, EVariableFilter Filter) const override;
 
 	virtual void PostLoad() override;
 #endif
diff --git a/Engine/Plugins/FX/Niagara/Source/Niagara/Internal/Stateless/Modules/NiagaraStatelessModule_SolveVelocitiesAndForces.h b/Engine/Plugins/FX/Niagara/Source/Niagara/Internal/Stateless/Modules/NiagaraStatelessModule_SolveVelocitiesAndForces.h
--- a/Engine/Plugins/FX/Niagara/Source/Niagara/Internal/Stateless/Modules/NiagaraStatelessModule_SolveVelocitiesAndForces.h
+++ b/Engine/Plugins/FX/Niagara/Source/Niagara/Internal/Stateless/Modules/NiagaraStatelessModule_SolveVelocitiesAndForces.h
@@ -33,6 +33,7 @@ public:
 	virtual void SetShaderParameters(const FNiagaraStatelessSetShaderParameterContext& SetShaderParameterContext) const override;
 
 #if WITH_EDITORONLY_DATA
-	virtual void GetOutputVariables(TArray<FNiagaraVariableBase>& OutVariables) const override;
+	virtual const TCHAR* GetShaderTemplatePath() const override;
+	virtual void GetOutputVariables(TArray<FNiagaraVariableBase>& OutVariables, EVariableFilter Filter) const override;
 #endif
 };
diff --git a/Engine/Plugins/FX/Niagara/Source/Niagara/Internal/Stateless/Modules/NiagaraStatelessModule_SpriteFacingAndAlignment.h b/Engine/Plugins/FX/Niagara/Source/Niagara/Internal/Stateless/Modules/NiagaraStatelessModule_SpriteFacingAndAlignment.h
--- a/Engine/Plugins/FX/Niagara/Source/Niagara/Internal/Stateless/Modules/NiagaraStatelessModule_SpriteFacingAndAlignment.h
+++ b/Engine/Plugins/FX/Niagara/Source/Niagara/Internal/Stateless/Modules/NiagaraStatelessModule_SpriteFacingAndAlignment.h
@@ -114,15 +114,20 @@ public:
 	virtual bool CanDisableModule() const override { return true; }
 #endif
 #if WITH_EDITORONLY_DATA
-	virtual void GetOutputVariables(TArray<FNiagaraVariableBase>& OutVariables) const override
+	virtual const TCHAR* GetShaderTemplatePath() const override
+	{
+		return TEXT("/Plugin/FX/Niagara/Private/Stateless/Modules/NiagaraStatelessModule_SpriteFacingAndAlignment.ush");
+	}
+
+	virtual void GetOutputVariables(TArray<FNiagaraVariableBase>& OutVariables, EVariableFilter Filter) const override
 	{
 		const FNiagaraStatelessGlobals& StatelessGlobals = FNiagaraStatelessGlobals::Get();
-		if (bSpriteFacingEnabled)
+		if (bSpriteFacingEnabled || Filter == EVariableFilter::None)
 		{
 			OutVariables.AddUnique(StatelessGlobals.SpriteFacingVariable);
 			OutVariables.AddUnique(StatelessGlobals.PreviousSpriteFacingVariable);
 		}
-		if (bSpriteAlignmentEnabled)
+		if (bSpriteAlignmentEnabled || Filter == EVariableFilter::None)
 		{
 			OutVariables.AddUnique(StatelessGlobals.SpriteAlignmentVariable);
 			OutVariables.AddUnique(StatelessGlobals.PreviousSpriteAlignmentVariable);
diff --git a/Engine/Plugins/FX/Niagara/Source/Niagara/Internal/Stateless/Modules/NiagaraStatelessModule_SpriteRotationRate.h b/Engine/Plugins/FX/Niagara/Source/Niagara/Internal/Stateless/Modules/NiagaraStatelessModule_SpriteRotationRate.h
--- a/Engine/Plugins/FX/Niagara/Source/Niagara/Internal/Stateless/Modules/NiagaraStatelessModule_SpriteRotationRate.h
+++ b/Engine/Plugins/FX/Niagara/Source/Niagara/Internal/Stateless/Modules/NiagaraStatelessModule_SpriteRotationRate.h
@@ -114,7 +114,12 @@ public:
 	virtual bool CanDisableModule() const override { return true; }
 #endif
 #if WITH_EDITORONLY_DATA
-	virtual void GetOutputVariables(TArray<FNiagaraVariableBase>& OutVariables) const override
+	virtual const TCHAR* GetShaderTemplatePath() const override
+	{
+		return TEXT("/Plugin/FX/Niagara/Private/Stateless/Modules/NiagaraStatelessModule_SpriteRotationRate.ush");
+	}
+
+	virtual void GetOutputVariables(TArray<FNiagaraVariableBase>& OutVariables, EVariableFilter Filter) const override
 	{
 		const FNiagaraStatelessGlobals& StatelessGlobals = FNiagaraStatelessGlobals::Get();
 		OutVariables.AddUnique(StatelessGlobals.SpriteRotationVariable);
diff --git a/Engine/Plugins/FX/Niagara/Source/Niagara/Internal/Stateless/Modules/NiagaraStatelessModule_SubUVAnimation.h b/Engine/Plugins/FX/Niagara/Source/Niagara/Internal/Stateless/Modules/NiagaraStatelessModule_SubUVAnimation.h
--- a/Engine/Plugins/FX/Niagara/Source/Niagara/Internal/Stateless/Modules/NiagaraStatelessModule_SubUVAnimation.h
+++ b/Engine/Plugins/FX/Niagara/Source/Niagara/Internal/Stateless/Modules/NiagaraStatelessModule_SubUVAnimation.h
@@ -66,6 +66,7 @@ public:
 #endif
 
 #if WITH_EDITORONLY_DATA
-	virtual void GetOutputVariables(TArray<FNiagaraVariableBase>& OutVariables) const override;
+	virtual const TCHAR* GetShaderTemplatePath() const override;
+	virtual void GetOutputVariables(TArray<FNiagaraVariableBase>& OutVariables, EVariableFilter Filter) const override;
 #endif
 };
diff --git a/Engine/Plugins/FX/Niagara/Source/Niagara/Internal/Stateless/NiagaraStatelessEmitter.h b/Engine/Plugins/FX/Niagara/Source/Niagara/Internal/Stateless/NiagaraStatelessEmitter.h
--- a/Engine/Plugins/FX/Niagara/Source/Niagara/Internal/Stateless/NiagaraStatelessEmitter.h
+++ b/Engine/Plugins/FX/Niagara/Source/Niagara/Internal/Stateless/NiagaraStatelessEmitter.h
@@ -209,7 +209,7 @@ protected:
 	FNiagaraDataSetCompiledData ParticleDataSetCompiledData;
 
 	UPROPERTY()
-	TArray<int32> ComponentOffsets;
+	TArray<int32> ShaderOutputVariableOffsets;
 
 	UPROPERTY()
 	TArray<TObjectPtr<UNiagaraParameterCollection>> CachedParameterCollectionReferences;
diff --git a/Engine/Plugins/FX/Niagara/Source/Niagara/Internal/Stateless/NiagaraStatelessEmitterData.h b/Engine/Plugins/FX/Niagara/Source/Niagara/Internal/Stateless/NiagaraStatelessEmitterData.h
--- a/Engine/Plugins/FX/Niagara/Source/Niagara/Internal/Stateless/NiagaraStatelessEmitterData.h
+++ b/Engine/Plugins/FX/Niagara/Source/Niagara/Internal/Stateless/NiagaraStatelessEmitterData.h
@@ -34,7 +34,7 @@ struct FNiagaraStatelessEmitterData
 #endif
 
 	TSharedPtr<FNiagaraDataSetCompiledData>			ParticleDataSetCompiledData;
-	TArray<int32>									ComponentOffsets;
+	TArray<int32>									ShaderOutputVariableOffsets;
 
 	bool											bCanEverExecute = false;
 	bool											bDeterministic = false;
diff --git a/Engine/Plugins/FX/Niagara/Source/Niagara/Internal/Stateless/NiagaraStatelessEmitterTemplate.h b/Engine/Plugins/FX/Niagara/Source/Niagara/Internal/Stateless/NiagaraStatelessEmitterTemplate.h
--- a/Engine/Plugins/FX/Niagara/Source/Niagara/Internal/Stateless/NiagaraStatelessEmitterTemplate.h
+++ b/Engine/Plugins/FX/Niagara/Source/Niagara/Internal/Stateless/NiagaraStatelessEmitterTemplate.h
@@ -36,22 +36,30 @@ public:
 
 	TConstArrayView<TObjectPtr<UClass>> GetModules() const { return Modules; }
 #if WITH_EDITORONLY_DATA
-	TConstArrayView<FNiagaraVariableBase> GetOututputComponents() const { return OutputComponents; }
+	TConstArrayView<FNiagaraVariableBase> GetImplicitVariables() const { return ImplicitVariables; }
+	TConstArrayView<FNiagaraVariableBase> GetModuleVariables() const { return ModuleVariables; }
+	TConstArrayView<FNiagaraVariableBase> GetShaderOutputVariables() const { return ShaderOutputVariables; }
 #endif
 
 	virtual const FShaderParametersMetadata* GetShaderParametersMetadata() const { check(ShaderParametersMetadata.IsValid()); return ShaderParametersMetadata.Get(); }
 	virtual TShaderRef<NiagaraStateless::FSimulationShader> GetSimulationShader() const { return TShaderRef<NiagaraStateless::FSimulationShader>(); }
-	virtual void SetShaderParameters(uint8* ShaderParametersBase, TConstArrayView<int32> ComponentOffsets) const;
+	virtual void SetShaderParameters(uint8* ShaderParametersBase, TConstArrayView<int32> ShaderOutputVariableOffsets) const;
 
 protected:
-	UPROPERTY(EditAnywhere, Category="Modules", meta = (DisallowCreateNew, AllowedClasses = "/Script/Niagara.NiagaraStatelessModule"))
+	NIAGARA_API void BuildVariables();
+	NIAGARA_API void BuildShaderParameters();
+
+protected:
+	UPROPERTY(EditAnywhere, Category = "Modules", meta = (DisallowCreateNew, AllowedClasses = "/Script/Niagara.NiagaraStatelessModule"))
 	TArray<TObjectPtr<UClass>> Modules;
 
 #if WITH_EDITORONLY_DATA
 	UPROPERTY(EditAnywhere, Category="Modules")
 	bool bExposedToLibrary = true;
 
-	TArray<FNiagaraVariableBase> OutputComponents;
+	TArray<FNiagaraVariableBase> ImplicitVariables;				// Variables that are implicit to this template (i.e. are not from modules)
+	TArray<FNiagaraVariableBase> ModuleVariables;				// All available variables from modules
+	TArray<FNiagaraVariableBase> ShaderOutputVariables;			// Variables that can be output from the shader, this is tied into the mappings sent into the shader.  A subset or exact match of ImplicitVariables & ModuleVariables.
 #endif
 	TSharedPtr<FShaderParametersMetadata> ShaderParametersMetadata;
 };
@@ -67,6 +75,6 @@ class UNiagaraStatelessEmitterDefault : public UNiagaraStatelessEmitterTemplate
 	virtual TShaderRef<NiagaraStateless::FSimulationShader> GetSimulationShader() const override;
 
 	//-TODO: Add a way to set them directly, we should know that the final struct is a series of ints in the order of the provided variables
-	virtual void SetShaderParameters(uint8* ShaderParametersBase, TConstArrayView<int32> ComponentOffsets) const override;
+	virtual void SetShaderParameters(uint8* ShaderParametersBase, TConstArrayView<int32> ShaderOutputVariableOffsets) const override;
 	//-TODO: Add a way to set them directly, we should know that the final struct is a series of ints in the order of the provided variables
 };
diff --git a/Engine/Plugins/FX/Niagara/Source/Niagara/Internal/Stateless/NiagaraStatelessModule.h b/Engine/Plugins/FX/Niagara/Source/Niagara/Internal/Stateless/NiagaraStatelessModule.h
--- a/Engine/Plugins/FX/Niagara/Source/Niagara/Internal/Stateless/NiagaraStatelessModule.h
+++ b/Engine/Plugins/FX/Niagara/Source/Niagara/Internal/Stateless/NiagaraStatelessModule.h
@@ -20,6 +20,13 @@ class UNiagaraStatelessModule : public UNiagaraMergeable
 {
 	GENERATED_BODY()
 
+public:
+	enum class EVariableFilter
+	{
+		None,
+		Used,
+	};
+
 protected:
 	UPROPERTY(EditAnywhere, Category = "Parameters", meta = (DisplayPriority = 0, HideInStack))
 	uint32 bModuleEnabled : 1 = true;
@@ -56,7 +63,13 @@ public:
 #endif
 
 #if WITH_EDITORONLY_DATA
-	virtual void GetOutputVariables(TArray<FNiagaraVariableBase>& OutVariables) const {}
+	// Optional virtual path to shader template
+	// Not required for CPU modules or GPU modules that aggregate into another module and build time
+	virtual const TCHAR* GetShaderTemplatePath() const { return nullptr; }
+
+	// Gather the output variables for this module
+	// EVariableFilter specifies if we should filter to used variables or not
+	virtual void GetOutputVariables(TArray<FNiagaraVariableBase>& OutVariables, EVariableFilter Filter) const {}
 #endif
 
 	//~UObject interface Begin
diff --git a/Engine/Plugins/FX/Niagara/Source/Niagara/Internal/Stateless/NiagaraStatelessParticleSimContext.h b/Engine/Plugins/FX/Niagara/Source/Niagara/Internal/Stateless/NiagaraStatelessParticleSimContext.h
--- a/Engine/Plugins/FX/Niagara/Source/Niagara/Internal/Stateless/NiagaraStatelessParticleSimContext.h
+++ b/Engine/Plugins/FX/Niagara/Source/Niagara/Internal/Stateless/NiagaraStatelessParticleSimContext.h
@@ -25,7 +25,7 @@ enum class EParticleComponent
 	NormalizedAge,
 	PreviousAge,
 	PreviousNormalizedAge,
-	UniqueIndex,
+	UniqueID,
 	MaterialRandom,
 	Num
 };
@@ -63,7 +63,7 @@ public:
 	float* GetParticleNormalizedAge() { return static_cast<float*>(RequiredComponents[int32(EParticleComponent::NormalizedAge)]); }
 	float* GetParticlePreviousAge() { return static_cast<float*>(RequiredComponents[int32(EParticleComponent::PreviousAge)]); }
 	float* GetParticlePreviousNormalizedAge() { return static_cast<float*>(RequiredComponents[int32(EParticleComponent::PreviousNormalizedAge)]); }
-	int32* GetParticleUniqueIndex() { return static_cast<int32*>(RequiredComponents[int32(EParticleComponent::UniqueIndex)]); }
+	int32* GetParticleUniqueID() { return static_cast<int32*>(RequiredComponents[int32(EParticleComponent::UniqueID)]); }
 	float* GetParticleMaterialRandom() { return static_cast<float*>(RequiredComponents[int32(EParticleComponent::MaterialRandom)]); }
 
 	const float* GetParticleLifetime() const { return static_cast<float*>(RequiredComponents[int32(EParticleComponent::Lifetime)]); }
@@ -71,7 +71,7 @@ public:
 	const float* GetParticleNormalizedAge() const { return static_cast<float*>(RequiredComponents[int32(EParticleComponent::NormalizedAge)]); }
 	const float* GetParticlePreviousAge() const { return static_cast<float*>(RequiredComponents[int32(EParticleComponent::PreviousAge)]); }
 	const float* GetParticlePreviousNormalizedAge() const { return static_cast<float*>(RequiredComponents[int32(EParticleComponent::PreviousNormalizedAge)]); }
-	const int32* GetParticleUniqueIndex() const { return static_cast<int32*>(RequiredComponents[int32(EParticleComponent::UniqueIndex)]); }
+	const int32* GetParticleUniqueID() const { return static_cast<int32*>(RequiredComponents[int32(EParticleComponent::UniqueID)]); }
 
 	// Get a pointer to an optional particle component (i.e. one not required in the data set)
 	float* GetParticleVariableFloat(int32 iVariable) const { return iVariable >= 0 ? reinterpret_cast<float*>(VariableComponents[iVariable]) : nullptr; }
diff --git a/Engine/Plugins/FX/Niagara/Source/Niagara/Private/Stateless/Modules/NiagaraStatelessModule_CalculateAccurateVelocity.cpp b/Engine/Plugins/FX/Niagara/Source/Niagara/Private/Stateless/Modules/NiagaraStatelessModule_CalculateAccurateVelocity.cpp
--- a/Engine/Plugins/FX/Niagara/Source/Niagara/Private/Stateless/Modules/NiagaraStatelessModule_CalculateAccurateVelocity.cpp
+++ b/Engine/Plugins/FX/Niagara/Source/Niagara/Private/Stateless/Modules/NiagaraStatelessModule_CalculateAccurateVelocity.cpp
@@ -62,7 +62,12 @@ void UNiagaraStatelessModule_CalculateAccurateVelocity::SetShaderParameters(cons
 }
 
 #if WITH_EDITORONLY_DATA
-void UNiagaraStatelessModule_CalculateAccurateVelocity::GetOutputVariables(TArray<FNiagaraVariableBase>& OutVariables) const
+const TCHAR* UNiagaraStatelessModule_CalculateAccurateVelocity::GetShaderTemplatePath() const
+{
+	return TEXT("/Plugin/FX/Niagara/Private/Stateless/Modules/NiagaraStatelessModule_CalculateAccurateVelocity.ush");
+}
+
+void UNiagaraStatelessModule_CalculateAccurateVelocity::GetOutputVariables(TArray<FNiagaraVariableBase>& OutVariables, EVariableFilter Filter) const
 {
 	const FNiagaraStatelessGlobals& StatelessGlobals = FNiagaraStatelessGlobals::Get();
 	OutVariables.AddUnique(StatelessGlobals.VelocityVariable);
diff --git a/Engine/Plugins/FX/Niagara/Source/Niagara/Private/Stateless/Modules/NiagaraStatelessModule_DecalAttributes.cpp b/Engine/Plugins/FX/Niagara/Source/Niagara/Private/Stateless/Modules/NiagaraStatelessModule_DecalAttributes.cpp
--- a/Engine/Plugins/FX/Niagara/Source/Niagara/Private/Stateless/Modules/NiagaraStatelessModule_DecalAttributes.cpp
+++ b/Engine/Plugins/FX/Niagara/Source/Niagara/Private/Stateless/Modules/NiagaraStatelessModule_DecalAttributes.cpp
@@ -20,17 +20,18 @@ namespace NSMDecalAttributesPrivate
 		Num
 	};
 
-	const FNiagaraVariableBase* GetAttributeVariable(EModuleAttribute Attribute, const UNiagaraStatelessModule_DecalAttributes* Module)
+	const FNiagaraVariableBase* GetAttributeVariable(EModuleAttribute Attribute, const UNiagaraStatelessModule_DecalAttributes* Module, UNiagaraStatelessModule::EVariableFilter Filter)
 	{
 		static const FNiagaraVariableBase DecalOrientationVariable = NiagaraStatelessCommon::ConvertParticleVariableToStateless(UNiagaraDecalRendererProperties::GetDecalOrientationVariable());
 		static const FNiagaraVariableBase DecalSizeVariable = NiagaraStatelessCommon::ConvertParticleVariableToStateless(UNiagaraDecalRendererProperties::GetDecalSizeVariable());
 		static const FNiagaraVariableBase DecalFadeVariable = NiagaraStatelessCommon::ConvertParticleVariableToStateless(UNiagaraDecalRendererProperties::GetDecalFadeVariable());
 
+		const bool bIgnoreFilter = Filter == UNiagaraStatelessModule::EVariableFilter::None;
 		switch (Attribute)
 		{
-			case EModuleAttribute::Orientation:		return Module->bApplyOrientation ? &DecalOrientationVariable : nullptr;
-			case EModuleAttribute::Size:			return Module->bApplySize ? &DecalSizeVariable : nullptr;
-			case EModuleAttribute::Fade:			return Module->bApplyFade ? &DecalFadeVariable : nullptr;
+			case EModuleAttribute::Orientation:		return bIgnoreFilter || Module->bApplyOrientation ? &DecalOrientationVariable : nullptr;
+			case EModuleAttribute::Size:			return bIgnoreFilter || Module->bApplySize ? &DecalSizeVariable : nullptr;
+			case EModuleAttribute::Fade:			return bIgnoreFilter || Module->bApplyFade ? &DecalFadeVariable : nullptr;
 		}
 		return nullptr;
 	}
@@ -106,7 +107,7 @@ void UNiagaraStatelessModule_DecalAttributes::BuildEmitterData(const FNiagaraSta
 		bool bAnyValidAttributes = false;
 		for (int i = 0; i < int(EModuleAttribute::Num); ++i)
 		{
-			if (const FNiagaraVariableBase* Variable = GetAttributeVariable(EModuleAttribute(i), this))
+			if (const FNiagaraVariableBase* Variable = GetAttributeVariable(EModuleAttribute(i), this, EVariableFilter::Used))
 			{
 				BuiltData->AttributeOffset[i] = BuildContext.FindParticleVariableIndex(*Variable);
 				bAnyValidAttributes |= BuiltData->AttributeOffset[i] != INDEX_NONE;
@@ -145,13 +146,13 @@ void UNiagaraStatelessModule_DecalAttributes::SetShaderParameters(const FNiagara
 }
 
 #if WITH_EDITORONLY_DATA
-void UNiagaraStatelessModule_DecalAttributes::GetOutputVariables(TArray<FNiagaraVariableBase>& OutVariables) const
+void UNiagaraStatelessModule_DecalAttributes::GetOutputVariables(TArray<FNiagaraVariableBase>& OutVariables, EVariableFilter Filter) const
 {
 	using namespace NSMDecalAttributesPrivate;
 
 	for (int i = 0; i < int(EModuleAttribute::Num); ++i)
 	{
-		if (const FNiagaraVariableBase* Variable = GetAttributeVariable(EModuleAttribute(i), this))
+		if (const FNiagaraVariableBase* Variable = GetAttributeVariable(EModuleAttribute(i), this, Filter))
 		{
 			OutVariables.AddUnique(*Variable);
 		}
diff --git a/Engine/Plugins/FX/Niagara/Source/Niagara/Private/Stateless/Modules/NiagaraStatelessModule_InitialMeshOrientation.cpp b/Engine/Plugins/FX/Niagara/Source/Niagara/Private/Stateless/Modules/NiagaraStatelessModule_InitialMeshOrientation.cpp
--- a/Engine/Plugins/FX/Niagara/Source/Niagara/Private/Stateless/Modules/NiagaraStatelessModule_InitialMeshOrientation.cpp
+++ b/Engine/Plugins/FX/Niagara/Source/Niagara/Private/Stateless/Modules/NiagaraStatelessModule_InitialMeshOrientation.cpp
@@ -121,7 +121,12 @@ void UNiagaraStatelessModule_InitialMeshOrientation::SetShaderParameters(const F
 }
 
 #if WITH_EDITORONLY_DATA
-void UNiagaraStatelessModule_InitialMeshOrientation::GetOutputVariables(TArray<FNiagaraVariableBase>& OutVariables) const
+const TCHAR* UNiagaraStatelessModule_InitialMeshOrientation::GetShaderTemplatePath() const
+{
+	return TEXT("/Plugin/FX/Niagara/Private/Stateless/Modules/NiagaraStatelessModule_InitialMeshOrientation.ush");
+}
+
+void UNiagaraStatelessModule_InitialMeshOrientation::GetOutputVariables(TArray<FNiagaraVariableBase>& OutVariables, EVariableFilter Filter) const
 {
 	const FNiagaraStatelessGlobals& StatelessGlobals = FNiagaraStatelessGlobals::Get();
 	OutVariables.AddUnique(StatelessGlobals.MeshOrientationVariable);
diff --git a/Engine/Plugins/FX/Niagara/Source/Niagara/Private/Stateless/Modules/NiagaraStatelessModule_LightAttributes.cpp b/Engine/Plugins/FX/Niagara/Source/Niagara/Private/Stateless/Modules/NiagaraStatelessModule_LightAttributes.cpp
--- a/Engine/Plugins/FX/Niagara/Source/Niagara/Private/Stateless/Modules/NiagaraStatelessModule_LightAttributes.cpp
+++ b/Engine/Plugins/FX/Niagara/Source/Niagara/Private/Stateless/Modules/NiagaraStatelessModule_LightAttributes.cpp
@@ -21,7 +21,7 @@ namespace NSMLightAttributesPrivate
 		Num
 	};
 
-	const FNiagaraVariableBase* GetAttributeVariable(EModuleAttribute Attribute, const UNiagaraStatelessModule_LightAttributes* Module)
+	const FNiagaraVariableBase* GetAttributeVariable(EModuleAttribute Attribute, const UNiagaraStatelessModule_LightAttributes* Module, UNiagaraStatelessModule::EVariableFilter Filter)
 	{
 		static const FNiagaraVariableBase LightRadiusVariable = NiagaraStatelessCommon::ConvertParticleVariableToStateless(SYS_PARAM_PARTICLES_LIGHT_RADIUS);
 		static const FNiagaraVariableBase LightFalloffExponent = NiagaraStatelessCommon::ConvertParticleVariableToStateless(SYS_PARAM_PARTICLES_LIGHT_EXPONENT);
@@ -29,13 +29,14 @@ namespace NSMLightAttributesPrivate
 		static const FNiagaraVariableBase LightSpecularScale = NiagaraStatelessCommon::ConvertParticleVariableToStateless(SYS_PARAM_PARTICLES_LIGHT_SPECULAR_SCALE);
 		static const FNiagaraVariableBase LightVolumetricScattering = NiagaraStatelessCommon::ConvertParticleVariableToStateless(SYS_PARAM_PARTICLES_LIGHT_VOLUMETRIC_SCATTERING);
 
+		const bool bIgnoreFilter = Filter == UNiagaraStatelessModule::EVariableFilter::None;
 		switch (Attribute)
 		{
-			case EModuleAttribute::Radius:			return Module->bApplyRadius ? &LightRadiusVariable : nullptr;
-			case EModuleAttribute::Falloff:			return Module->bApplyFalloffExponent ? &LightFalloffExponent : nullptr;
-			case EModuleAttribute::DiffuseScale:	return Module->bApplyDiffuseScale ? &LightDiffuseScale : nullptr;
-			case EModuleAttribute::SpecularScale:	return Module->bApplySpecularScale ? &LightSpecularScale : nullptr;
-			case EModuleAttribute::VolumetricScale:	return Module->bApplyVolumetricScattering ? &LightVolumetricScattering : nullptr;
+			case EModuleAttribute::Radius:			return bIgnoreFilter || Module->bApplyRadius ? &LightRadiusVariable : nullptr;
+			case EModuleAttribute::Falloff:			return bIgnoreFilter || Module->bApplyFalloffExponent ? &LightFalloffExponent : nullptr;
+			case EModuleAttribute::DiffuseScale:	return bIgnoreFilter || Module->bApplyDiffuseScale ? &LightDiffuseScale : nullptr;
+			case EModuleAttribute::SpecularScale:	return bIgnoreFilter || Module->bApplySpecularScale ? &LightSpecularScale : nullptr;
+			case EModuleAttribute::VolumetricScale:	return bIgnoreFilter || Module->bApplyVolumetricScattering ? &LightVolumetricScattering : nullptr;
 		}
 		return nullptr;
 	}
@@ -110,7 +111,7 @@ void UNiagaraStatelessModule_LightAttributes::BuildEmitterData(const FNiagaraSta
 		bool bAnyValidAttributes = false;
 		for (int i = 0; i < int(EModuleAttribute::Num); ++i)
 		{
-			if (const FNiagaraVariableBase* Variable = GetAttributeVariable(EModuleAttribute(i), this))
+			if (const FNiagaraVariableBase* Variable = GetAttributeVariable(EModuleAttribute(i), this, EVariableFilter::Used))
 			{
 				BuiltData->AttributeOffset[i] = BuildContext.FindParticleVariableIndex(*Variable);
 				bAnyValidAttributes |= BuiltData->AttributeOffset[i] != INDEX_NONE;
@@ -144,13 +145,13 @@ void UNiagaraStatelessModule_LightAttributes::SetShaderParameters(const FNiagara
 }
 
 #if WITH_EDITORONLY_DATA
-void UNiagaraStatelessModule_LightAttributes::GetOutputVariables(TArray<FNiagaraVariableBase>& OutVariables) const
+void UNiagaraStatelessModule_LightAttributes::GetOutputVariables(TArray<FNiagaraVariableBase>& OutVariables, EVariableFilter Filter) const
 {
 	using namespace NSMLightAttributesPrivate;
 
 	for (int i = 0; i < int(EModuleAttribute::Num); ++i)
 	{
-		if (const FNiagaraVariableBase* Variable = GetAttributeVariable(EModuleAttribute(i), this))
+		if (const FNiagaraVariableBase* Variable = GetAttributeVariable(EModuleAttribute(i), this, Filter))
 		{
 			OutVariables.AddUnique(*Variable);
 		}
diff --git a/Engine/Plugins/FX/Niagara/Source/Niagara/Private/Stateless/Modules/NiagaraStatelessModule_MeshIndex.cpp b/Engine/Plugins/FX/Niagara/Source/Niagara/Private/Stateless/Modules/NiagaraStatelessModule_MeshIndex.cpp
--- a/Engine/Plugins/FX/Niagara/Source/Niagara/Private/Stateless/Modules/NiagaraStatelessModule_MeshIndex.cpp
+++ b/Engine/Plugins/FX/Niagara/Source/Niagara/Private/Stateless/Modules/NiagaraStatelessModule_MeshIndex.cpp
@@ -148,3 +148,16 @@ void UNiagaraStatelessModule_MeshIndex::SetShaderParameters(const FNiagaraStatel
 	Parameters->MeshIndex_TableOffset		= ModuleBuiltData->TableOffset;
 	Parameters->MeshIndex_TableNumElements	= ModuleBuiltData->TableNumElements;
 }
+
+#if WITH_EDITORONLY_DATA
+const TCHAR* UNiagaraStatelessModule_MeshIndex::GetShaderTemplatePath() const
+{
+	return TEXT("/Plugin/FX/Niagara/Private/Stateless/Modules/NiagaraStatelessModule_MeshIndex.ush");
+}
+
+void UNiagaraStatelessModule_MeshIndex::GetOutputVariables(TArray<FNiagaraVariableBase>& OutVariables, EVariableFilter Filter) const
+{
+	const FNiagaraStatelessGlobals& StatelessGlobals = FNiagaraStatelessGlobals::Get();
+	OutVariables.AddUnique(StatelessGlobals.MeshIndexVariable);
+}
+#endif
diff --git a/Engine/Plugins/FX/Niagara/Source/Niagara/Private/Stateless/Modules/NiagaraStatelessModule_RotateAroundPoint.cpp b/Engine/Plugins/FX/Niagara/Source/Niagara/Private/Stateless/Modules/NiagaraStatelessModule_RotateAroundPoint.cpp
--- a/Engine/Plugins/FX/Niagara/Source/Niagara/Private/Stateless/Modules/NiagaraStatelessModule_RotateAroundPoint.cpp
+++ b/Engine/Plugins/FX/Niagara/Source/Niagara/Private/Stateless/Modules/NiagaraStatelessModule_RotateAroundPoint.cpp
@@ -153,7 +153,12 @@ void UNiagaraStatelessModule_RotateAroundPoint::DrawDebug(const FNiagaraStateles
 #endif
 
 #if WITH_EDITORONLY_DATA
-void UNiagaraStatelessModule_RotateAroundPoint::GetOutputVariables(TArray<FNiagaraVariableBase>& OutVariables) const
+const TCHAR* UNiagaraStatelessModule_RotateAroundPoint::GetShaderTemplatePath() const
+{
+	return TEXT("/Plugin/FX/Niagara/Private/Stateless/Modules/NiagaraStatelessModule_RotateAroundPoint.ush");
+}
+
+void UNiagaraStatelessModule_RotateAroundPoint::GetOutputVariables(TArray<FNiagaraVariableBase>& OutVariables, EVariableFilter Filter) const
 {
 	using namespace NSMRotateAroundPointPrivate;
 
diff --git a/Engine/Plugins/FX/Niagara/Source/Niagara/Private/Stateless/Modules/NiagaraStatelessModule_ScaleMeshSizeBySpeed.cpp b/Engine/Plugins/FX/Niagara/Source/Niagara/Private/Stateless/Modules/NiagaraStatelessModule_ScaleMeshSizeBySpeed.cpp
--- a/Engine/Plugins/FX/Niagara/Source/Niagara/Private/Stateless/Modules/NiagaraStatelessModule_ScaleMeshSizeBySpeed.cpp
+++ b/Engine/Plugins/FX/Niagara/Source/Niagara/Private/Stateless/Modules/NiagaraStatelessModule_ScaleMeshSizeBySpeed.cpp
@@ -120,7 +120,12 @@ void UNiagaraStatelessModule_ScaleMeshSizeBySpeed::SetShaderParameters(const FNi
 }
 
 #if WITH_EDITORONLY_DATA
-void UNiagaraStatelessModule_ScaleMeshSizeBySpeed::GetOutputVariables(TArray<FNiagaraVariableBase>& OutVariables) const
+const TCHAR* UNiagaraStatelessModule_ScaleMeshSizeBySpeed::GetShaderTemplatePath() const
+{
+	return TEXT("/Plugin/FX/Niagara/Private/Stateless/Modules/NiagaraStatelessModule_ScaleMeshSizebySpeed.ush");
+}
+
+void UNiagaraStatelessModule_ScaleMeshSizeBySpeed::GetOutputVariables(TArray<FNiagaraVariableBase>& OutVariables, EVariableFilter Filter) const
 {
 	const FNiagaraStatelessGlobals& StatelessGlobals = FNiagaraStatelessGlobals::Get();
 	OutVariables.AddUnique(StatelessGlobals.ScaleVariable);
diff --git a/Engine/Plugins/FX/Niagara/Source/Niagara/Private/Stateless/Modules/NiagaraStatelessModule_ScaleSpriteSizeBySpeed.cpp b/Engine/Plugins/FX/Niagara/Source/Niagara/Private/Stateless/Modules/NiagaraStatelessModule_ScaleSpriteSizeBySpeed.cpp
--- a/Engine/Plugins/FX/Niagara/Source/Niagara/Private/Stateless/Modules/NiagaraStatelessModule_ScaleSpriteSizeBySpeed.cpp
+++ b/Engine/Plugins/FX/Niagara/Source/Niagara/Private/Stateless/Modules/NiagaraStatelessModule_ScaleSpriteSizeBySpeed.cpp
@@ -118,7 +118,12 @@ void UNiagaraStatelessModule_ScaleSpriteSizeBySpeed::SetShaderParameters(const F
 }
 
 #if WITH_EDITORONLY_DATA
-void UNiagaraStatelessModule_ScaleSpriteSizeBySpeed::GetOutputVariables(TArray<FNiagaraVariableBase>& OutVariables) const
+const TCHAR* UNiagaraStatelessModule_ScaleSpriteSizeBySpeed::GetShaderTemplatePath() const
+{
+	return TEXT("/Plugin/FX/Niagara/Private/Stateless/Modules/NiagaraStatelessModule_ScaleSpriteSizeBySpeed.ush");
+}
+
+void UNiagaraStatelessModule_ScaleSpriteSizeBySpeed::GetOutputVariables(TArray<FNiagaraVariableBase>& OutVariables, EVariableFilter Filter) const
 {
 	const FNiagaraStatelessGlobals& StatelessGlobals = FNiagaraStatelessGlobals::Get();
 	OutVariables.AddUnique(StatelessGlobals.SpriteSizeVariable);
diff --git a/Engine/Plugins/FX/Niagara/Source/Niagara/Private/Stateless/Modules/NiagaraStatelessModule_ShapeLocation.cpp b/Engine/Plugins/FX/Niagara/Source/Niagara/Private/Stateless/Modules/NiagaraStatelessModule_ShapeLocation.cpp
--- a/Engine/Plugins/FX/Niagara/Source/Niagara/Private/Stateless/Modules/NiagaraStatelessModule_ShapeLocation.cpp
+++ b/Engine/Plugins/FX/Niagara/Source/Niagara/Private/Stateless/Modules/NiagaraStatelessModule_ShapeLocation.cpp
@@ -394,6 +394,20 @@ void UNiagaraStatelessModule_ShapeLocation::DrawDebug(const FNiagaraStatelessDra
 		}
 	}
 }
+#endif
+
+#if WITH_EDITORONLY_DATA
+const TCHAR* UNiagaraStatelessModule_ShapeLocation::GetShaderTemplatePath() const
+{
+	return TEXT("/Plugin/FX/Niagara/Private/Stateless/Modules/NiagaraStatelessModule_ShapeLocation.ush");
+}
+
+void UNiagaraStatelessModule_ShapeLocation::GetOutputVariables(TArray<FNiagaraVariableBase>& OutVariables, EVariableFilter Filter) const
+{
+	const FNiagaraStatelessGlobals& StatelessGlobals = FNiagaraStatelessGlobals::Get();
+	OutVariables.AddUnique(StatelessGlobals.PositionVariable);
+	OutVariables.AddUnique(StatelessGlobals.PreviousPositionVariable);
+}
 
 void UNiagaraStatelessModule_ShapeLocation::PostLoad()
 {
diff --git a/Engine/Plugins/FX/Niagara/Source/Niagara/Private/Stateless/Modules/NiagaraStatelessModule_SolveVelocitiesAndForces.cpp b/Engine/Plugins/FX/Niagara/Source/Niagara/Private/Stateless/Modules/NiagaraStatelessModule_SolveVelocitiesAndForces.cpp
--- a/Engine/Plugins/FX/Niagara/Source/Niagara/Private/Stateless/Modules/NiagaraStatelessModule_SolveVelocitiesAndForces.cpp
+++ b/Engine/Plugins/FX/Niagara/Source/Niagara/Private/Stateless/Modules/NiagaraStatelessModule_SolveVelocitiesAndForces.cpp
@@ -240,7 +240,12 @@ void UNiagaraStatelessModule_SolveVelocitiesAndForces::SetShaderParameters(const
 }
 
 #if WITH_EDITORONLY_DATA
-void UNiagaraStatelessModule_SolveVelocitiesAndForces::GetOutputVariables(TArray<FNiagaraVariableBase>& OutVariables) const
+const TCHAR* UNiagaraStatelessModule_SolveVelocitiesAndForces::GetShaderTemplatePath() const
+{
+	return TEXT("/Plugin/FX/Niagara/Private/Stateless/Modules/NiagaraStatelessModule_SolveVelocitiesAndForces.ush");
+}
+
+void UNiagaraStatelessModule_SolveVelocitiesAndForces::GetOutputVariables(TArray<FNiagaraVariableBase>& OutVariables, EVariableFilter Filter) const
 {
 	const FNiagaraStatelessGlobals& StatelessGlobals = FNiagaraStatelessGlobals::Get();
 	OutVariables.AddUnique(StatelessGlobals.PositionVariable);
diff --git a/Engine/Plugins/FX/Niagara/Source/Niagara/Private/Stateless/Modules/NiagaraStatelessModule_SubUVAnimation.cpp b/Engine/Plugins/FX/Niagara/Source/Niagara/Private/Stateless/Modules/NiagaraStatelessModule_SubUVAnimation.cpp
--- a/Engine/Plugins/FX/Niagara/Source/Niagara/Private/Stateless/Modules/NiagaraStatelessModule_SubUVAnimation.cpp
+++ b/Engine/Plugins/FX/Niagara/Source/Niagara/Private/Stateless/Modules/NiagaraStatelessModule_SubUVAnimation.cpp
@@ -162,7 +162,12 @@ void UNiagaraStatelessModule_SubUVAnimation::SetShaderParameters(const FNiagaraS
 
 
 #if WITH_EDITORONLY_DATA
-void UNiagaraStatelessModule_SubUVAnimation::GetOutputVariables(TArray<FNiagaraVariableBase>& OutVariables) const
+const TCHAR* UNiagaraStatelessModule_SubUVAnimation::GetShaderTemplatePath() const
+{
+	return TEXT("/Plugin/FX/Niagara/Private/Stateless/Modules/NiagaraStatelessModule_SubUVAnimation.ush");
+}
+
+void UNiagaraStatelessModule_SubUVAnimation::GetOutputVariables(TArray<FNiagaraVariableBase>& OutVariables, EVariableFilter Filter) const
 {
 	const FNiagaraStatelessGlobals& StatelessGlobals = FNiagaraStatelessGlobals::Get();
 	OutVariables.AddUnique(StatelessGlobals.SubImageIndexVariable);
diff --git a/Engine/Plugins/FX/Niagara/Source/Niagara/Private/Stateless/NiagaraStatelessEmitter.cpp b/Engine/Plugins/FX/Niagara/Source/Niagara/Private/Stateless/NiagaraStatelessEmitter.cpp
--- a/Engine/Plugins/FX/Niagara/Source/Niagara/Private/Stateless/NiagaraStatelessEmitter.cpp
+++ b/Engine/Plugins/FX/Niagara/Source/Niagara/Private/Stateless/NiagaraStatelessEmitter.cpp
@@ -513,35 +513,23 @@ void UNiagaraStatelessEmitter::BuildCompiledDataSet()
 	{
 		// Gather a list of all the output variables from the modules, this can change based on what is enabled / disabled
 		TArray<FNiagaraVariableBase> AvailableVariables;
+		AvailableVariables.Append(EmitterTemplate->GetImplicitVariables());
 
-		// Implicitly defined variables
-		const FNiagaraStatelessGlobals& StatelessGlobals = FNiagaraStatelessGlobals::Get();
-		if ( EmitterTemplate->GetOututputComponents().Contains(StatelessGlobals.UniqueIDVariable) )
-		{
-			AvailableVariables.AddUnique(StatelessGlobals.UniqueIDVariable);
-		}
-		if (EmitterTemplate->GetOututputComponents().Contains(StatelessGlobals.MaterialRandomVariable))
-		{
-			AvailableVariables.AddUnique(StatelessGlobals.MaterialRandomVariable);
-		}
-
-		// Module based variables
 		for (UNiagaraStatelessModule* Module : Modules)
 		{
 			if (Module->IsModuleEnabled())
 			{
-				Module->GetOutputVariables(AvailableVariables);
+				Module->GetOutputVariables(AvailableVariables, UNiagaraStatelessModule::EVariableFilter::Used);
 			}
 		}
 
-		// Remove any variables we don't output
-		// Note: We only need to do this for GPU path as it can't vary the outputs
-		TConstArrayView<FNiagaraVariableBase> OutputComponents = EmitterTemplate->GetOututputComponents();
+		// Validate and remove any variables that are not compatible with the GPU path (if enabled)
+		TConstArrayView<FNiagaraVariableBase> ShaderOutputVariables = EmitterTemplate->GetShaderOutputVariables();
 		if (EnumHasAnyFlags(StatelessEmitterData->FeatureMask, ENiagaraStatelessFeatureMask::ExecuteGPU))
 		{
-			for (auto it = AvailableVariables.CreateIterator(); it; ++it)
+			for (auto it=AvailableVariables.CreateIterator(); it; ++it)
 			{
-				if (!OutputComponents.Contains(*it))
+				if (!ShaderOutputVariables.Contains(*it))
 				{
 					UE_LOG(LogNiagara, Log, TEXT("Removed variable '%s' for emitter '%s' as it's not part of the output components"), *it->GetName().ToString(), *GetFullName());
 					it.RemoveCurrent();
@@ -600,23 +588,23 @@ void UNiagaraStatelessEmitter::BuildCompiledDataSet()
 		ParticleDataSetCompiledData.BuildLayout();
 
 		// Create mapping from output components to data set for the shader to output
-		ComponentOffsets.Empty(OutputComponents.Num());
-		for (const FNiagaraVariableBase& OutputComponent : OutputComponents)
+		ShaderOutputVariableOffsets.Empty(ShaderOutputVariables.Num());
+		for (const FNiagaraVariableBase& OutputVariable : ShaderOutputVariables)
 		{
-			if (OutputComponent.GetType() == FNiagaraTypeDefinition::GetIntDef())
+			if (OutputVariable.GetType() == FNiagaraTypeDefinition::GetIntDef())
 			{
-				ComponentOffsets.Add(NiagaraStatelessInternal::GetDataSetIntOffset(ParticleDataSetCompiledData, OutputComponent));
+				ShaderOutputVariableOffsets.Add(NiagaraStatelessInternal::GetDataSetIntOffset(ParticleDataSetCompiledData, OutputVariable));
 			}
 			else
 			{
-				ComponentOffsets.Add(NiagaraStatelessInternal::GetDataSetFloatOffset(ParticleDataSetCompiledData, OutputComponent));
+				ShaderOutputVariableOffsets.Add(NiagaraStatelessInternal::GetDataSetFloatOffset(ParticleDataSetCompiledData, OutputVariable));
 			}
 		}
-		ComponentOffsets.Shrink();
+		ShaderOutputVariableOffsets.Shrink();
 	}
 #endif
 	StatelessEmitterData->ParticleDataSetCompiledData = MakeShared<FNiagaraDataSetCompiledData>(ParticleDataSetCompiledData);
-	StatelessEmitterData->ComponentOffsets = ComponentOffsets;
+	StatelessEmitterData->ShaderOutputVariableOffsets = ShaderOutputVariableOffsets;
 }
 
 void UNiagaraStatelessEmitter::ResolveScalabilitySettings()
@@ -700,7 +688,7 @@ NiagaraStateless::FCommonShaderParameters* UNiagaraStatelessEmitter::AllocateSha
 	}
 
 	//-TODO: Add a way to set them directly, we should know that the final struct is a series of ints in the order of the provided variables
-	GetEmitterTemplate()->SetShaderParameters(static_cast<uint8*>(UntypedShaderParameters), StatelessEmitterData->ComponentOffsets);
+	GetEmitterTemplate()->SetShaderParameters(static_cast<uint8*>(UntypedShaderParameters), StatelessEmitterData->ShaderOutputVariableOffsets);
 
 	return CommonParameters;
 }
diff --git a/Engine/Plugins/FX/Niagara/Source/Niagara/Private/Stateless/NiagaraStatelessEmitterTemplate.cpp b/Engine/Plugins/FX/Niagara/Source/Niagara/Private/Stateless/NiagaraStatelessEmitterTemplate.cpp
--- a/Engine/Plugins/FX/Niagara/Source/Niagara/Private/Stateless/NiagaraStatelessEmitterTemplate.cpp
+++ b/Engine/Plugins/FX/Niagara/Source/Niagara/Private/Stateless/NiagaraStatelessEmitterTemplate.cpp
@@ -108,29 +108,8 @@ void UNiagaraStatelessEmitterTemplate::InitModulesAndAttributes()
 	// Remove any invalid modules
 	Modules.RemoveAll([](UClass* ModuleClass) { return ModuleClass == nullptr || !ModuleClass->IsChildOf(UNiagaraStatelessModule::StaticClass()); });
 
-#if WITH_EDITORONLY_DATA
-	// Find All Output Components
-	OutputComponents.Empty();
-	for (UClass* ModuleClass : Modules)
-	{
-		UNiagaraStatelessModule* Module = CastChecked<UNiagaraStatelessModule>(ModuleClass->GetDefaultObject());
-		Module->GetOutputVariables(OutputComponents);
-	}
-#endif
-
-	// Build Shader Parameters
-	FShaderParametersMetadataBuilder ShaderMetadataBuilder;
-	FNiagaraStatelessShaderParametersBuilder ShaderParametersBuilder(&ShaderMetadataBuilder);
-
-	ShaderParametersBuilder.AddParameterNestedStruct<NiagaraStateless::FCommonShaderParameters>();
-	for (UClass* ModuleClass : Modules)
-	{
-		UNiagaraStatelessModule* Module = CastChecked<UNiagaraStatelessModule>(ModuleClass->GetDefaultObject());
-		Module->BuildShaderParameters(ShaderParametersBuilder);
-	}
-	//-TODO: Reserve space for component outputs when compiling GPU ShaderParametersBuilder.AddParameterNestedStruct<NiagaraStateless::FCommonShaderParameters>();
-
-	ShaderParametersMetadata = MakeShareable<FShaderParametersMetadata>(ShaderMetadataBuilder.Build(FShaderParametersMetadata::EUseCase::ShaderParameterStruct, TEXT("NiagaraStatelessEmitterTemplate")));
+	BuildVariables();
+	BuildShaderParameters();
 }
 
 #if WITH_EDITOR
@@ -182,11 +161,65 @@ bool UNiagaraStatelessEmitterTemplate::IsExposedToLibrary(const FAssetData& Asse
 }
 #endif
 
-void UNiagaraStatelessEmitterTemplate::SetShaderParameters(uint8* ShaderParametersBase, TConstArrayView<int32> ComponentOffsets) const
+void UNiagaraStatelessEmitterTemplate::SetShaderParameters(uint8* ShaderParametersBase, TConstArrayView<int32> ShaderOutputVariableOffsets) const
 {
 	//-TODO: We don't store component offsets right now
 }
 
+void UNiagaraStatelessEmitterTemplate::BuildVariables()
+{
+#if WITH_EDITORONLY_DATA
+	// Clear Variable List
+	ImplicitVariables.Empty();
+	ModuleVariables.Empty();
+	ShaderOutputVariables.Empty();
+
+	// Implicit Variables
+	const FNiagaraStatelessGlobals& StatelessGlobals = FNiagaraStatelessGlobals::Get();
+	ImplicitVariables.Add(StatelessGlobals.UniqueIDVariable);
+	ImplicitVariables.Add(StatelessGlobals.MaterialRandomVariable);
+
+	ShaderOutputVariables = ImplicitVariables;
+
+	// Extract variables from modules
+	{
+		TArray<FNiagaraVariableBase> TempVariables;
+		for (UClass* ModuleClass : Modules)
+		{
+			UNiagaraStatelessModule* Module = CastChecked<UNiagaraStatelessModule>(ModuleClass->GetDefaultObject());
+			Module->GetOutputVariables(TempVariables, UNiagaraStatelessModule::EVariableFilter::None);
+
+			const bool bShaderOutput = EnumHasAnyFlags(Module->GetFeatureMask(), ENiagaraStatelessFeatureMask::ExecuteGPU);
+			for (const FNiagaraVariableBase& Variable : TempVariables)
+			{
+				ModuleVariables.AddUnique(Variable);
+				if (bShaderOutput)
+				{
+					ShaderOutputVariables.AddUnique(Variable);
+				}
+			}
+			TempVariables.Reset();
+		}
+	}
+#endif
+}
+
+void UNiagaraStatelessEmitterTemplate::BuildShaderParameters()
+{
+	FShaderParametersMetadataBuilder ShaderMetadataBuilder;
+	FNiagaraStatelessShaderParametersBuilder ShaderParametersBuilder(&ShaderMetadataBuilder);
+
+	ShaderParametersBuilder.AddParameterNestedStruct<NiagaraStateless::FCommonShaderParameters>();
+	for (UClass* ModuleClass : Modules)
+	{
+		UNiagaraStatelessModule* Module = CastChecked<UNiagaraStatelessModule>(ModuleClass->GetDefaultObject());
+		Module->BuildShaderParameters(ShaderParametersBuilder);
+	}
+	//-TODO: Reserve space for component outputs when compiling GPU ShaderParametersBuilder.AddParameterNestedStruct<NiagaraStateless::FCommonShaderParameters>();
+
+	ShaderParametersMetadata = MakeShareable<FShaderParametersMetadata>(ShaderMetadataBuilder.Build(FShaderParametersMetadata::EUseCase::ShaderParameterStruct, TEXT("NiagaraStatelessEmitterTemplate")));
+}
+
 //////////////////////////////////////////////////////////////////////////////////////////////////////////////////////////////////////////////////////////
 
 void UNiagaraStatelessEmitterDefault::InitModulesAndAttributes()
@@ -222,41 +255,7 @@ void UNiagaraStatelessEmitterDefault::InitModulesAndAttributes()
 		UNiagaraStatelessModule_SubUVAnimation::StaticClass(),
 	};
 
-#if WITH_EDITORONLY_DATA
-	const FNiagaraStatelessGlobals& StatelessGlobals = FNiagaraStatelessGlobals::Get();
-	OutputComponents =
-	{
-		StatelessGlobals.UniqueIDVariable,
-		StatelessGlobals.MaterialRandomVariable,
-		StatelessGlobals.PositionVariable,
-		StatelessGlobals.CameraOffsetVariable,
-		StatelessGlobals.ColorVariable,
-		StatelessGlobals.DynamicMaterialParameters0Variable,
-		StatelessGlobals.DynamicMaterialParameters1Variable,
-		StatelessGlobals.DynamicMaterialParameters2Variable,
-		StatelessGlobals.DynamicMaterialParameters3Variable,
-		StatelessGlobals.MeshIndexVariable,
-		StatelessGlobals.MeshOrientationVariable,
-		StatelessGlobals.RibbonWidthVariable,
-		StatelessGlobals.ScaleVariable,
-		StatelessGlobals.SpriteSizeVariable,
-		StatelessGlobals.SpriteFacingVariable,
-		StatelessGlobals.SpriteAlignmentVariable,
-		StatelessGlobals.SpriteRotationVariable,
-		StatelessGlobals.SubImageIndexVariable,
-		StatelessGlobals.VelocityVariable,
-		StatelessGlobals.PreviousPositionVariable,
-		StatelessGlobals.PreviousCameraOffsetVariable,
-		StatelessGlobals.PreviousMeshOrientationVariable,
-		StatelessGlobals.PreviousRibbonWidthVariable,
-		StatelessGlobals.PreviousScaleVariable,
-		StatelessGlobals.PreviousSpriteSizeVariable,
-		StatelessGlobals.PreviousSpriteFacingVariable,
-		StatelessGlobals.PreviousSpriteAlignmentVariable,
-		StatelessGlobals.PreviousSpriteRotationVariable,
-		StatelessGlobals.PreviousVelocityVariable,
-	};
-#endif //WITH_EDITORONLY_DATA
+	BuildVariables();
 }
 
 const FShaderParametersMetadata* UNiagaraStatelessEmitterDefault::GetShaderParametersMetadata() const
@@ -272,40 +271,40 @@ TShaderRef<NiagaraStateless::FSimulationShader> UNiagaraStatelessEmitterDefault:
 }
 
 //-TODO: Add a way to set them directly, we should know that the final struct is a series of ints in the order of the provided variables
-void UNiagaraStatelessEmitterDefault::SetShaderParameters(uint8* ShaderParametersBase, TConstArrayView<int32> ComponentOffsets) const
+void UNiagaraStatelessEmitterDefault::SetShaderParameters(uint8* ShaderParametersBase, TConstArrayView<int32> ShaderOutputVariableOffsets) const
 {
 	using namespace NiagaraStateless;
 	FSimulationShaderDefaultCS::FParameters* ShaderParameters	= reinterpret_cast<FSimulationShaderDefaultCS::FParameters*>(ShaderParametersBase);
 	int iComponent = 0;
-	ShaderParameters->Permutation_UniqueIndexComponent				= ComponentOffsets[iComponent++];
-	ShaderParameters->Permutation_MaterialRandomComponent			= ComponentOffsets[iComponent++];
-	ShaderParameters->Permutation_PositionComponent					= ComponentOffsets[iComponent++];
-	ShaderParameters->Permutation_CameraOffsetComponent				= ComponentOffsets[iComponent++];
-	ShaderParameters->Permutation_ColorComponent					= ComponentOffsets[iComponent++];
-	ShaderParameters->Permutation_DynamicMaterialParameter0Component= ComponentOffsets[iComponent++];
-	ShaderParameters->Permutation_DynamicMaterialParameter1Component= ComponentOffsets[iComponent++];
-	ShaderParameters->Permutation_DynamicMaterialParameter2Component= ComponentOffsets[iComponent++];
-	ShaderParameters->Permutation_DynamicMaterialParameter3Component= ComponentOffsets[iComponent++];
-	ShaderParameters->Permutation_MeshIndexComponent				= ComponentOffsets[iComponent++];
-	ShaderParameters->Permutation_MeshOrientationComponent			= ComponentOffsets[iComponent++];
-	ShaderParameters->Permutation_RibbonWidthComponent				= ComponentOffsets[iComponent++];
-	ShaderParameters->Permutation_ScaleComponent					= ComponentOffsets[iComponent++];
-	ShaderParameters->Permutation_SpriteSizeComponent				= ComponentOffsets[iComponent++];
-	ShaderParameters->Permutation_SpriteFacingComponent				= ComponentOffsets[iComponent++];
-	ShaderParameters->Permutation_SpriteAlignmentComponent			= ComponentOffsets[iComponent++];
-	ShaderParameters->Permutation_SpriteRotationComponent			= ComponentOffsets[iComponent++];
-	ShaderParameters->Permutation_SubImageIndexComponent			= ComponentOffsets[iComponent++];
-	ShaderParameters->Permutation_VelocityComponent					= ComponentOffsets[iComponent++];
-	ShaderParameters->Permutation_PreviousPositionComponent			= ComponentOffsets[iComponent++];
-	ShaderParameters->Permutation_PreviousCameraOffsetComponent		= ComponentOffsets[iComponent++];
-	ShaderParameters->Permutation_PreviousMeshOrientationComponent	= ComponentOffsets[iComponent++];
-	ShaderParameters->Permutation_PreviousRibbonWidthComponent		= ComponentOffsets[iComponent++];
-	ShaderParameters->Permutation_PreviousScaleComponent			= ComponentOffsets[iComponent++];
-	ShaderParameters->Permutation_PreviousSpriteSizeComponent		= ComponentOffsets[iComponent++];
-	ShaderParameters->Permutation_PreviousSpriteFacingComponent		= ComponentOffsets[iComponent++];
-	ShaderParameters->Permutation_PreviousSpriteAlignmentComponent	= ComponentOffsets[iComponent++];
-	ShaderParameters->Permutation_PreviousSpriteRotationComponent	= ComponentOffsets[iComponent++];
-	ShaderParameters->Permutation_PreviousVelocityComponent			= ComponentOffsets[iComponent++];
+	ShaderParameters->Permutation_UniqueIDComponent					= ShaderOutputVariableOffsets[iComponent++];
+	ShaderParameters->Permutation_MaterialRandomComponent			= ShaderOutputVariableOffsets[iComponent++];
+	ShaderParameters->Permutation_PositionComponent					= ShaderOutputVariableOffsets[iComponent++];
+	ShaderParameters->Permutation_ColorComponent					= ShaderOutputVariableOffsets[iComponent++];
+	ShaderParameters->Permutation_SpriteSizeComponent				= ShaderOutputVariableOffsets[iComponent++];
+	ShaderParameters->Permutation_SpriteRotationComponent			= ShaderOutputVariableOffsets[iComponent++];
+	ShaderParameters->Permutation_ScaleComponent					= ShaderOutputVariableOffsets[iComponent++];
+	ShaderParameters->Permutation_PreviousPositionComponent			= ShaderOutputVariableOffsets[iComponent++];
+	ShaderParameters->Permutation_PreviousSpriteSizeComponent		= ShaderOutputVariableOffsets[iComponent++];
+	ShaderParameters->Permutation_PreviousSpriteRotationComponent	= ShaderOutputVariableOffsets[iComponent++];
+	ShaderParameters->Permutation_PreviousScaleComponent			= ShaderOutputVariableOffsets[iComponent++];
+	ShaderParameters->Permutation_RibbonWidthComponent				= ShaderOutputVariableOffsets[iComponent++];
+	ShaderParameters->Permutation_PreviousRibbonWidthComponent		= ShaderOutputVariableOffsets[iComponent++];
+	ShaderParameters->Permutation_MeshOrientationComponent			= ShaderOutputVariableOffsets[iComponent++];
+	ShaderParameters->Permutation_PreviousMeshOrientationComponent	= ShaderOutputVariableOffsets[iComponent++];
+	ShaderParameters->Permutation_VelocityComponent					= ShaderOutputVariableOffsets[iComponent++];
+	ShaderParameters->Permutation_PreviousVelocityComponent			= ShaderOutputVariableOffsets[iComponent++];
+	ShaderParameters->Permutation_CameraOffsetComponent				= ShaderOutputVariableOffsets[iComponent++];
+	ShaderParameters->Permutation_PreviousCameraOffsetComponent		= ShaderOutputVariableOffsets[iComponent++];
+	ShaderParameters->Permutation_DynamicMaterialParameterComponent	= ShaderOutputVariableOffsets[iComponent++];
+	ShaderParameters->Permutation_DynamicMaterialParameter1Component= ShaderOutputVariableOffsets[iComponent++];
+	ShaderParameters->Permutation_DynamicMaterialParameter2Component= ShaderOutputVariableOffsets[iComponent++];
+	ShaderParameters->Permutation_DynamicMaterialParameter3Component= ShaderOutputVariableOffsets[iComponent++];
+	ShaderParameters->Permutation_MeshIndexComponent				= ShaderOutputVariableOffsets[iComponent++];
+	ShaderParameters->Permutation_SpriteFacingComponent				= ShaderOutputVariableOffsets[iComponent++];
+	ShaderParameters->Permutation_PreviousSpriteFacingComponent		= ShaderOutputVariableOffsets[iComponent++];
+	ShaderParameters->Permutation_SpriteAlignmentComponent			= ShaderOutputVariableOffsets[iComponent++];
+	ShaderParameters->Permutation_PreviousSpriteAlignmentComponent	= ShaderOutputVariableOffsets[iComponent++];
+	ShaderParameters->Permutation_SubImageIndexComponent			= ShaderOutputVariableOffsets[iComponent++];
 }
 //-TODO: Add a way to set them directly, we should know that the final struct is a series of ints in the order of the provided variables
 
diff --git a/Engine/Plugins/FX/Niagara/Source/Niagara/Private/Stateless/NiagaraStatelessParticleSimContext.cpp b/Engine/Plugins/FX/Niagara/Source/Niagara/Private/Stateless/NiagaraStatelessParticleSimContext.cpp
--- a/Engine/Plugins/FX/Niagara/Source/Niagara/Private/Stateless/NiagaraStatelessParticleSimContext.cpp
+++ b/Engine/Plugins/FX/Niagara/Source/Niagara/Private/Stateless/NiagaraStatelessParticleSimContext.cpp
@@ -57,7 +57,7 @@ namespace NiagaraStateless
 			{},														// float		- NormalizedAge
 			{},														// float		- PreviousAge
 			{},														// float		- PreviousNormalizedAge
-			FNiagaraStatelessGlobals::Get().UniqueIDVariable,		// int32		- UniqueIndex
+			FNiagaraStatelessGlobals::Get().UniqueIDVariable,		// int32		- UniqueID
 			FNiagaraStatelessGlobals::Get().MaterialRandomVariable,	// float		- MaterialRandom
 		};
 		return MakeArrayView(RequireComponents);
@@ -166,7 +166,7 @@ namespace NiagaraStateless
 		// Setup particles
 		for (uint32 iParticle = 0; iParticle < ActiveParticles; ++iParticle)
 		{
-			uint32 Particle_UniqueIndex = 0;
+			uint32 Particle_UniqueID = 0;
 			float Particle_Age = -1.0f;
 			float Particle_Lifetime = 0.0f;
 			{
@@ -183,11 +183,11 @@ namespace NiagaraStateless
 					if (SpawnInfoIndex < SpawnInfoNumActive)
 					{
 						const uint32 SpawnParticleIndex = SpawnInfoIndex + SpawnInfoParticleOffset;
-						Particle_UniqueIndex	= SpawnInfoIndex + SpawnInfoUniqueOffset + SpawnInfoParticleOffset;
-						Particle_Age			= EmitterAge - (SpawnInfoTime + float(SpawnParticleIndex) * SpawnInfoRate);
+						Particle_UniqueID	= SpawnInfoIndex + SpawnInfoUniqueOffset + SpawnInfoParticleOffset;
+						Particle_Age		= EmitterAge - (SpawnInfoTime + float(SpawnParticleIndex) * SpawnInfoRate);
 
 						// Write unique index as it's needed for the random operation
-						GetParticleUniqueIndex()[NumInstances] = Particle_UniqueIndex;
+						GetParticleUniqueID()[NumInstances] = Particle_UniqueID;
 
 						Particle_Lifetime		= RandomScaleBiasFloat(NumInstances, 0, SpawnInfoLifetimeScale, SpawnInfoLifetimeBias);
 						break;
@@ -288,32 +288,32 @@ namespace NiagaraStateless
 
 	uint32 FParticleSimulationContext::RandomUInt(uint32 iInstance, uint32 RandomSeedOffset) const
 	{
-		const uint32 UniqueIndex = uint32(GetParticleUniqueIndex()[iInstance]);
-		const FUintVector4 RandomSeed = FNiagaraStatelessDefinitions::MakeRandomSeed(EmitterRandomSeed, UniqueIndex, ModuleRandomSeed, RandomSeedOffset);
+		const uint32 UniqueID = uint32(GetParticleUniqueID()[iInstance]);
+		const FUintVector4 RandomSeed = FNiagaraStatelessDefinitions::MakeRandomSeed(EmitterRandomSeed, UniqueID, ModuleRandomSeed, RandomSeedOffset);
 		const FUintVector4 RandomValue = Rand4DPCG32(RandomSeed);
 		return RandomValue.X;
 	}
 
 	FUintVector2 FParticleSimulationContext::RandomUInt2(uint32 iInstance, uint32 RandomSeedOffset) const
 	{
-		const uint32 UniqueIndex = uint32(GetParticleUniqueIndex()[iInstance]);
-		const FUintVector4 RandomSeed = FNiagaraStatelessDefinitions::MakeRandomSeed(EmitterRandomSeed, UniqueIndex, ModuleRandomSeed, RandomSeedOffset);
+		const uint32 UniqueID = uint32(GetParticleUniqueID()[iInstance]);
+		const FUintVector4 RandomSeed = FNiagaraStatelessDefinitions::MakeRandomSeed(EmitterRandomSeed, UniqueID, ModuleRandomSeed, RandomSeedOffset);
 		const FUintVector4 RandomValue = Rand4DPCG32(RandomSeed);
 		return FUintVector2(RandomValue.X, RandomValue.Y);
 	}
 
 	FUintVector3 FParticleSimulationContext::RandomUInt3(uint32 iInstance, uint32 RandomSeedOffset) const
 	{
-		const uint32 UniqueIndex = uint32(GetParticleUniqueIndex()[iInstance]);
-		const FUintVector4 RandomSeed = FNiagaraStatelessDefinitions::MakeRandomSeed(EmitterRandomSeed, UniqueIndex, ModuleRandomSeed, RandomSeedOffset);
+		const uint32 UniqueID = uint32(GetParticleUniqueID()[iInstance]);
+		const FUintVector4 RandomSeed = FNiagaraStatelessDefinitions::MakeRandomSeed(EmitterRandomSeed, UniqueID, ModuleRandomSeed, RandomSeedOffset);
 		const FUintVector4 RandomValue = Rand4DPCG32(RandomSeed);
 		return FUintVector3(RandomValue.X, RandomValue.Y, RandomValue.Z);
 	}
 
 	FUintVector4 FParticleSimulationContext::RandomUInt4(uint32 iInstance, uint32 RandomSeedOffset) const
 	{
-		const uint32 UniqueIndex = uint32(GetParticleUniqueIndex()[iInstance]);
-		const FUintVector4 RandomSeed = FNiagaraStatelessDefinitions::MakeRandomSeed(EmitterRandomSeed, UniqueIndex, ModuleRandomSeed, RandomSeedOffset);
+		const uint32 UniqueID = uint32(GetParticleUniqueID()[iInstance]);
+		const FUintVector4 RandomSeed = FNiagaraStatelessDefinitions::MakeRandomSeed(EmitterRandomSeed, UniqueID, ModuleRandomSeed, RandomSeedOffset);
 		const FUintVector4 RandomValue = Rand4DPCG32(RandomSeed);
 		return RandomValue;
 	}
diff --git a/Engine/Plugins/FX/Niagara/Source/NiagaraEditor/Internal/ViewModels/NiagaraStatelessEmitterTemplateViewModel.h b/Engine/Plugins/FX/Niagara/Source/NiagaraEditor/Internal/ViewModels/NiagaraStatelessEmitterTemplateViewModel.h
--- a/Engine/Plugins/FX/Niagara/Source/NiagaraEditor/Internal/ViewModels/NiagaraStatelessEmitterTemplateViewModel.h
+++ b/Engine/Plugins/FX/Niagara/Source/NiagaraEditor/Internal/ViewModels/NiagaraStatelessEmitterTemplateViewModel.h
@@ -12,6 +12,14 @@ class UNiagaraStatelessEmitterTemplate;
 class FNiagaraStatelessEmitterTemplateViewModel : public TSharedFromThis<FNiagaraStatelessEmitterTemplateViewModel>, public FNotifyHook
 {
 public:
+	struct FOutputVariable
+	{
+		FNiagaraVariableBase	Variable;
+		bool					bIsImplicit = false;
+		bool					bIsModule = false;
+		bool					bIsShader = false;
+	};
+
 	DECLARE_MULTICAST_DELEGATE(FOnTemplateChanged)
 
 	explicit FNiagaraStatelessEmitterTemplateViewModel(UNiagaraStatelessEmitterTemplate* EmitterTemplate);
@@ -25,13 +33,19 @@ public:
 	//- End: FNotifyHook
 
 	UNiagaraStatelessEmitterTemplate* GetTemplateForEdit() { return EmitterTemplateForEdit.Get(); }
-	TConstArrayView<FNiagaraVariableBase> GetOututputComponents() const;
+	TConstArrayView<FOutputVariable> GetOututputVariables() const;
 	TPair<ENiagaraStatelessFeatureMask, ENiagaraStatelessFeatureMask> GetFeatureMaskRange() const;
 
+	FString GenerateComputeTemplateHLSL() const;
+
 	FOnTemplateChanged& OnTemplateChanged() { return OnTemplateChangedDelegate; }
 
+private:
+	void BuildOutputVariables();
+
 private:
 	TStrongObjectPtr<UNiagaraStatelessEmitterTemplate>	EmitterTemplate;
 	TStrongObjectPtr<UNiagaraStatelessEmitterTemplate>	EmitterTemplateForEdit;
+	TArray<FOutputVariable>								OutputVariables;
 	FOnTemplateChanged									OnTemplateChangedDelegate;
 };
diff --git a/Engine/Plugins/FX/Niagara/Source/NiagaraEditor/Private/Stateless/NiagaraStatelessEmitterTranslator.cpp b/Engine/Plugins/FX/Niagara/Source/NiagaraEditor/Private/Stateless/NiagaraStatelessEmitterTranslator.cpp
new file mode 100644
--- /dev/null
+++ b/Engine/Plugins/FX/Niagara/Source/NiagaraEditor/Private/Stateless/NiagaraStatelessEmitterTranslator.cpp
@@ -0,0 +1,299 @@
+// Copyright Epic Games, Inc. All Rights Reserved.
+
+#include "Stateless/NiagaraStatelessEmitterTranslator.h"
+#include "Stateless/NiagaraStatelessEmitterTemplate.h"
+#include "Stateless/NiagaraStatelessEmitter.h"
+#include "Stateless/NiagaraStatelessModule.h"
+
+#include "ShaderCore.h"
+
+#define NIAGARA_STATELESS_TRANSLATION_OLD_MODE 1
+
+namespace FNiagaraStatelessEmitterTranslator
+{
+	static const TCHAR* DefaultIncludes[] =
+	{
+		TEXT("/Plugin/FX/Niagara/Private/NiagaraQuaternionUtils.ush"),
+		TEXT("/Plugin/FX/Niagara/Private/Stateless/NiagaraStatelessCommon.ush"),
+	};
+
+	static TConstArrayView<FNiagaraVariableBase> GetImplicitVariables()
+	{
+		static TArray<FNiagaraVariableBase> ImplicitVariables =
+		{
+			FNiagaraStatelessGlobals::Get().MaterialRandomVariable,
+			FNiagaraStatelessGlobals::Get().UniqueIDVariable,
+		};
+		return ImplicitVariables;
+	}
+
+	static FString GetModuleHlslName(const UClass* ModuleClass, int ModuleIndex)
+	{
+		//return FString::Printf(TEXT("%s_%d"), *ModuleClass->GetName(), ModuleIndex);
+		return ModuleClass->GetName();
+	}
+
+	static FString SanitizeVariableName(const FName VariableName)
+	{
+		FString SanitizedName(VariableName.ToString());
+		SanitizedName.ReplaceInline(TEXT("."), TEXT(""));
+		return SanitizedName;
+	}
+
+	static const TCHAR* GetVariableHlslType(const FNiagaraTypeDefinition& TypeDef)
+	{
+		if (TypeDef == FNiagaraTypeDefinition::GetFloatDef())
+		{
+			return TEXT("float");
+		}
+		else if (TypeDef == FNiagaraTypeDefinition::GetVec2Def())
+		{
+			return TEXT("float2");
+		}
+		else if (TypeDef == FNiagaraTypeDefinition::GetVec3Def() || TypeDef == FNiagaraTypeDefinition::GetPositionDef())
+		{
+			return TEXT("float3");
+		}
+		else if (TypeDef == FNiagaraTypeDefinition::GetVec4Def() || TypeDef == FNiagaraTypeDefinition::GetColorDef() || TypeDef == FNiagaraTypeDefinition::GetColorDef() || TypeDef == FNiagaraTypeDefinition::GetQuatDef())
+		{
+			return TEXT("float4");
+		}
+		else if (TypeDef == FNiagaraTypeDefinition::GetIntDef())
+		{
+			return TEXT("int");
+		}
+		checkNoEntry();
+		return TEXT("");
+	}
+
+	void WriteHeader(FString& HlslOutput)
+	{
+		HlslOutput.Append(TEXT("// Copyright Epic Games, Inc. All Rights Reserved.\n"));
+		HlslOutput.Append(TEXT("\n"));
+		for (const TCHAR* Include : DefaultIncludes)
+		{
+			HlslOutput.Appendf(TEXT("#include \"%s\"\n"), Include);
+		}
+	}
+
+	void WriteParticleStruct(FString& HlslOutput, UNiagaraStatelessEmitterTemplate* EmitterTemplate)
+	{
+	#if NIAGARA_STATELESS_TRANSLATION_OLD_MODE
+		HlslOutput.Append(TEXT("\n"));
+		HlslOutput.Append(TEXT("#define PARTICLE_ATTRIBUTES \\\n"));
+
+		TConstArrayView<FNiagaraVariableBase> ModuleVariables = EmitterTemplate->GetModuleVariables();
+		if (ModuleVariables.Num() > 0)
+		{
+			TConstArrayView<FNiagaraVariableBase> ImplicitVariables = EmitterTemplate->GetImplicitVariables();
+			for (const FNiagaraVariableBase& ParticleVariable : ImplicitVariables)
+			{
+				const TCHAR* HlslType = GetVariableHlslType(ParticleVariable.GetType());
+				const FString HlslName = SanitizeVariableName(ParticleVariable.GetName());
+				HlslOutput.Appendf(TEXT("\tPARTICLE_ATTRIBUTE_OUTPUT_IMPLICIT(%s, %s) \\\n"), HlslType, *HlslName);
+			}
+
+			for (const FNiagaraVariableBase& ParticleVariable : ModuleVariables)
+			{
+				if (ImplicitVariables.Contains(ParticleVariable))
+				{
+					continue;
+				}
+
+				const TCHAR* HlslType = GetVariableHlslType(ParticleVariable.GetType());
+				const FString HlslName = SanitizeVariableName(ParticleVariable.GetName());
+				HlslOutput.Appendf(TEXT("\tPARTICLE_ATTRIBUTE_OUTPUT(%s, %s) \\\n"), HlslType, *HlslName);
+			}
+		}
+	#else
+		HlslOutput.Append(TEXT("\n"));
+		HlslOutput.Append(TEXT("struct FStatelessParticle\n"));
+		HlslOutput.Append(TEXT("{\n"));
+		HlslOutput.Append(TEXT("\tbool bAlive;					// If the particle should remain alive or not (only evaluated at specific points)\n"));
+		HlslOutput.Append(TEXT("\tuint UniqueID;				// Unique particle index, sequential based on previous spawn info accumulation\n"));
+		HlslOutput.Append(TEXT("\tfloat MaterialRandom;			// Random float passed to materials\n"));
+		HlslOutput.Append(TEXT("\tfloat Lifetime;					// Overall lifetime for the particle\n"));
+		HlslOutput.Append(TEXT("\tfloat Age;						// Current age for the particle\n"));
+		HlslOutput.Append(TEXT("\tfloat NormalizedAge;			// Current normalized age for the particle\n"));
+		HlslOutput.Append(TEXT("\tloat PreviousAge;				// Previous age for the particle\n"));
+		HlslOutput.Append(TEXT("\tfloat PreviousNormalizedAge;	// Previous normalized age for the particle\n"));
+		HlslOutput.Append(TEXT("\tfloat DeltaTime;				// Simulation DT\n"));
+		HlslOutput.Append(TEXT("\tfloat InvDeltaTime;				// Simulation Inverse DT\n"));
+
+		TConstArrayView<FNiagaraVariableBase> ModuleVariables = EmitterTemplate->GetModuleVariables();
+		if (ModuleVariables.Num() > 0)
+		{
+			TConstArrayView<FNiagaraVariableBase> ImplicitVariables = EmitterTemplate->GetImplicitVariables();
+
+			HlslOutput.Append(TEXT("\n"));
+			HlslOutput.Append(TEXT("\t// Module Variables\n"));
+
+			for (const FNiagaraVariableBase& ParticleVariable : ModuleVariables)
+			{
+				if (ImplicitVariables.Contains(ParticleVariable))
+				{
+					continue;
+				}
+
+				const TCHAR* HlslType = GetVariableHlslType(ParticleVariable.GetType());
+				const FString HlslName = SanitizeVariableName(ParticleVariable.GetName());
+				HlslOutput.Appendf(TEXT("\t%s %s;\n"), HlslType, *HlslName);
+			}
+		}
+		HlslOutput.Append(TEXT("};\n"));
+	#endif
+	}
+
+	void WriteModules(FString& HlslOutput, UNiagaraStatelessEmitterTemplate* EmitterTemplate)
+	{
+	#if NIAGARA_STATELESS_TRANSLATION_OLD_MODE
+		HlslOutput.Append(TEXT("\n"));
+		HlslOutput.Append(TEXT("#define PARTICLE_MODULES \\\n"));
+
+		TConstArrayView<UClass*> Modules = EmitterTemplate->GetModules();
+		for (int iModule = 0; iModule < Modules.Num(); ++iModule)
+		{
+			const UClass* ModuleClass = Modules[iModule];
+			const UNiagaraStatelessModule* Module = ModuleClass->GetDefaultObject<UNiagaraStatelessModule>();
+			const TCHAR* ModuleShaderTemplatePath = Module ? Module->GetShaderTemplatePath() : nullptr;
+
+			FString ShortClassName = ModuleClass->GetName();
+			ShortClassName.RemoveFromStart(TEXT("NiagaraStatelessModule_"));
+
+			if (ModuleShaderTemplatePath == nullptr)
+			{
+				HlslOutput.Appendf(TEXT("\tPARTICLE_MODULE(Null) /* %s */ \\\n"), *ShortClassName);
+			}
+			else
+			{
+				HlslOutput.Appendf(TEXT("\tPARTICLE_MODULE(%s) \\\n"), *ShortClassName);
+			}
+		}
+
+		HlslOutput.Append(TEXT("\n"));
+		HlslOutput.Append(TEXT("#include \"/Plugin/FX/Niagara/Private/Stateless/Template/GenerateCS_PreModules.ush\"\n"));
+		HlslOutput.Append(TEXT("\n"));
+
+		for (int iModule = 0; iModule < Modules.Num(); ++iModule)
+		{
+			const UClass* ModuleClass = Modules[iModule];
+			const UNiagaraStatelessModule* Module = ModuleClass->GetDefaultObject<UNiagaraStatelessModule>();
+			const TCHAR* ModuleShaderTemplatePath = Module ? Module->GetShaderTemplatePath() : nullptr;
+			if (ModuleShaderTemplatePath)
+			{
+				HlslOutput.Appendf(TEXT("#include \"%s\"\n"), ModuleShaderTemplatePath);
+			}
+		}
+
+		HlslOutput.Append(TEXT("\n"));
+		HlslOutput.Append(TEXT("#include \"/Plugin/FX/Niagara/Private/Stateless/Template/GenerateCS_PostModules.ush\"\n"));
+	#else
+		TConstArrayView<UClass*> Modules = EmitterTemplate->GetModules();
+		if (Modules.Num() == 0)
+		{
+			return;
+		}
+
+		HlslOutput.Append(TEXT("\n"));
+		for (int iModule = 0; iModule < Modules.Num(); ++iModule)
+		{
+			const UClass* ModuleClass = Modules[iModule];
+			const UNiagaraStatelessModule* Module = ModuleClass->GetDefaultObject<UNiagaraStatelessModule>();
+			const TCHAR* ModuleShaderTemplatePath = Module ? Module->GetShaderTemplatePath() : nullptr;
+			if (ModuleShaderTemplatePath == nullptr)
+			{
+				continue;
+			}
+
+			FString TemplateFile;
+			LoadShaderSourceFileChecked(ModuleShaderTemplatePath, EShaderPlatform::SP_PCD3D_SM5, TemplateFile);
+
+			//const FString ModuleHLSLName = GetModuleHlslName(ModuleClass, iModule);
+			//TemplateFile.ReplaceInline(Thing, ModuleHLSLName, ESearchCase::CaseSensitive);
+
+			HlslOutput.Append(TemplateFile);
+
+			//HlslOutput.Appendf(TEXT("#define MODULE_NAME %s\n"), *ModuleHLSLName);
+			//HlslOutput.Appendf(TEXT("#include \"%s\"\n"), ModuleShaderTemplatePath);
+			//HlslOutput.Append(TEXT("#undef MODULE_NAME\n"));
+		}
+	#endif
+	}
+
+	void WriteSimulateFunction(FString& HlslOutput, UNiagaraStatelessEmitterTemplate* EmitterTemplate)
+	{
+	#if NIAGARA_STATELESS_TRANSLATION_OLD_MODE == 0
+		TConstArrayView<UClass*> Modules = EmitterTemplate->GetModules();
+		if (Modules.Num() == 0)
+		{
+			return;
+		}
+
+		HlslOutput.Append(TEXT("\n"));
+		HlslOutput.Append(TEXT("void SimulateParticle(inout FStatelessParticle Particle, int RandomSeedOffset)\n"));
+		HlslOutput.Append(TEXT("{\n"));
+		for (int iModule = 0; iModule < Modules.Num(); ++iModule)
+		{
+			const UClass* ModuleClass = Modules[iModule];
+			const UNiagaraStatelessModule* Module = ModuleClass->GetDefaultObject<UNiagaraStatelessModule>();
+			const TCHAR* ModuleShaderTemplatePath = Module ? Module->GetShaderTemplatePath() : nullptr;
+			if (ModuleShaderTemplatePath == nullptr)
+			{
+				continue;
+			}
+
+			const FString ModuleHLSLName = GetModuleHlslName(ModuleClass, iModule);
+			HlslOutput.Append(TEXT("\tGRandomSeedInternal = FNiagaraStatelessDefinitions::MakeRandomSeed(Common_RandomSeed, Particle.UniqueID, RandomSeedOffset + iModule, 0);\n"));
+			HlslOutput.Appendf(TEXT("\t%s_Simulate(Particle);\n"), *ModuleHLSLName);
+		}
+		HlslOutput.Append(TEXT("}\n"));
+	#endif
+	}
+
+	void WriteOutputFunction(FString& HlslOutput, UNiagaraStatelessEmitterTemplate* EmitterTemplate)
+	{
+	#if NIAGARA_STATELESS_TRANSLATION_OLD_MODE == 0
+		TConstArrayView<FNiagaraVariableBase> ShaderOutputVariables = EmitterTemplate->GetShaderOutputVariables();
+		if (ShaderOutputVariables.Num() == 0)
+		{
+			return;
+		}
+
+		HlslOutput.Append(TEXT("\n"));
+		HlslOutput.Appendf(TEXT("DECLARE_SCALAR_ARRAY(uint, ShaderOutputOffsets, %d);\n"), ShaderOutputVariables.Num());
+
+		HlslOutput.Append(TEXT("\n"));
+		HlslOutput.Append(TEXT("void OutputParticle(inout FStatelessParticle Particle, uint ParticleIndex)\n"));
+		HlslOutput.Append(TEXT("{\n"));
+		for (int i = 0; i < ShaderOutputVariables.Num(); ++i)
+		{
+			const FNiagaraVariableBase& ParticleVariable = ShaderOutputVariables[i];
+			const FString HlslName = ParticleVariable.GetName().ToString();
+
+			HlslOutput.Appendf(TEXT("\tif ( IsValidComponent(GET_SCALAR_ARRAY_ELEMENT(ShaderOutputOffsets, %d)) )\n"), i);
+			HlslOutput.Append(TEXT("\t{\n"));
+			HlslOutput.Appendf(TEXT("\t\tOutputComponentData(OutputIndex, GET_SCALAR_ARRAY_ELEMENT(ShaderOutputOffsets, %d), Particle.%s);\n"), i, *HlslName);
+			HlslOutput.Append(TEXT("\t}\n"));
+		}
+		HlslOutput.Append(TEXT("}\n"));
+	#endif
+	}
+
+	void TranslateToCompute(FString& HlslOutput, UNiagaraStatelessEmitterTemplate* EmitterTemplate)
+	{
+		TConstArrayView<TObjectPtr<UClass>> Modules = EmitterTemplate->GetModules();
+
+		WriteHeader(HlslOutput);
+		WriteParticleStruct(HlslOutput, EmitterTemplate);
+		WriteModules(HlslOutput, EmitterTemplate);
+		WriteSimulateFunction(HlslOutput, EmitterTemplate);
+		WriteOutputFunction(HlslOutput, EmitterTemplate);
+
+	#if NIAGARA_STATELESS_TRANSLATION_OLD_MODE == 0
+		HlslOutput.Append(TEXT("\n"));
+		HlslOutput.Append(TEXT("#include \"NiagaraStatelessSimulationTemplate.ush\"\n"));
+	#endif
+	}
+}
+
+#undef NIAGARA_STATELESS_TRANSLATION_OLD_MODE
diff --git a/Engine/Plugins/FX/Niagara/Source/NiagaraEditor/Private/Stateless/NiagaraStatelessEmitterTranslator.h b/Engine/Plugins/FX/Niagara/Source/NiagaraEditor/Private/Stateless/NiagaraStatelessEmitterTranslator.h
new file mode 100644
--- /dev/null
+++ b/Engine/Plugins/FX/Niagara/Source/NiagaraEditor/Private/Stateless/NiagaraStatelessEmitterTranslator.h
@@ -0,0 +1,12 @@
+// Copyright Epic Games, Inc. All Rights Reserved.
+
+#pragma once
+
+#include "Stateless/NiagaraStatelessCommon.h"
+
+class UNiagaraStatelessEmitterTemplate;
+
+namespace FNiagaraStatelessEmitterTranslator
+{
+	void TranslateToCompute(FString& HlslOutput, UNiagaraStatelessEmitterTemplate* EmitterTemplate);
+}
diff --git a/Engine/Plugins/FX/Niagara/Source/NiagaraEditor/Private/Toolkits/NiagaraStatelessEmitterTemplateToolkit.cpp b/Engine/Plugins/FX/Niagara/Source/NiagaraEditor/Private/Toolkits/NiagaraStatelessEmitterTemplateToolkit.cpp
--- a/Engine/Plugins/FX/Niagara/Source/NiagaraEditor/Private/Toolkits/NiagaraStatelessEmitterTemplateToolkit.cpp
+++ b/Engine/Plugins/FX/Niagara/Source/NiagaraEditor/Private/Toolkits/NiagaraStatelessEmitterTemplateToolkit.cpp
@@ -7,6 +7,7 @@
 #include "Widgets/SNiagaraStatelessEmitterTemplateWidgets.h"
 
 #include "Framework/Docking/TabManager.h"
+#include "Framework/MultiBox/MultiBoxBuilder.h"
 #include "Widgets/Docking/SDockTab.h"
 
 #define LOCTEXT_NAMESPACE "NiagaraStatelessEmitterTemplateToolkit"
@@ -16,6 +17,7 @@ namespace NiagaraStatelessEmitterTemplateToolkitPrivate
 	const FName FeaturesTabID(TEXT("NiagaraStatelessEmitterTemplateEditor_Features"));
 	const FName ModulesTabID(TEXT("NiagaraStatelessEmitterTemplateEditor_Modules"));
 	const FName OutputVariablesTabID(TEXT("NiagaraStatelessEmitterTemplateEditor_OutputVariables"));
+	const FName CodeViewTabID(TEXT("NiagaraStatelessEmitterTemplateEditor_CodeView"));
 }
 
 FNiagaraStatelessEmitterTemplateToolkit::FNiagaraStatelessEmitterTemplateToolkit()
@@ -39,7 +41,7 @@ void FNiagaraStatelessEmitterTemplateToolkit::Initialize(const EToolkitMode::Typ
 	ViewModel = MakeShared<FNiagaraStatelessEmitterTemplateViewModel>(EmitterTemplate);
 		
 	const TSharedRef<FTabManager::FLayout> StandaloneDefaultLayout =
-		FTabManager::NewLayout("Standalone_Niagara_StatelessEmitterTemplateLayout_V1_Dev")
+		FTabManager::NewLayout("Standalone_Niagara_StatelessEmitterTemplateLayout_V2_Dev")
 		->AddArea
 		(
 			FTabManager::NewPrimaryArea()->SetOrientation(Orient_Horizontal)
@@ -62,6 +64,7 @@ void FNiagaraStatelessEmitterTemplateToolkit::Initialize(const EToolkitMode::Typ
 			(
 				FTabManager::NewStack()
 				->SetSizeCoefficient(0.8f)
+				->AddTab(CodeViewTabID, ETabState::OpenedTab)
 				->AddTab(ModulesTabID, ETabState::OpenedTab)
 			)
 		);
@@ -100,6 +103,12 @@ void FNiagaraStatelessEmitterTemplateToolkit::RegisterTabSpawners(const TSharedR
 		.SetDisplayName(LOCTEXT("TemplateOutputVariables", "Output Variables"))
 		.SetGroup(WorkspaceMenuCategory.ToSharedRef()
 	);
+	InTabManager->RegisterTabSpawner(
+		CodeViewTabID,
+		FOnSpawnTab::CreateSP(this, &FNiagaraStatelessEmitterTemplateToolkit::SpawnTab_CodeView))
+		.SetDisplayName(LOCTEXT("TemplateCodeView", "Code View"))
+		.SetGroup(WorkspaceMenuCategory.ToSharedRef()
+	);
 }
 
 void FNiagaraStatelessEmitterTemplateToolkit::UnregisterTabSpawners(const TSharedRef<class FTabManager>& InTabManager)
@@ -135,6 +144,24 @@ FLinearColor FNiagaraStatelessEmitterTemplateToolkit::GetWorldCentricTabColorSca
 
 void FNiagaraStatelessEmitterTemplateToolkit::ExtendToolbar()
 {
+	struct Local
+	{
+		static void FillToolbar(FToolBarBuilder& ToolbarBuilder, TSharedRef<FNiagaraStatelessEmitterTemplateViewModel> ViewModel)
+		{
+		//	ToolbarBuilder.BeginSection("Export");
+		//	ToolbarBuilder.EndSection();
+		}
+	};
+	
+	TSharedPtr<FExtender> ToolbarExtender = MakeShareable(new FExtender);
+	ToolbarExtender->AddToolBarExtension(
+		"Asset",
+		EExtensionHook::After,
+		GetToolkitCommands(),
+		FToolBarExtensionDelegate::CreateStatic(&Local::FillToolbar, ViewModel->AsShared())
+	);
+
+	AddToolbarExtender(ToolbarExtender);
 }
 
 TSharedRef<SDockTab> FNiagaraStatelessEmitterTemplateToolkit::SpawnTab_Features(const FSpawnTabArgs& Args)
@@ -167,5 +194,14 @@ TSharedRef<SDockTab> FNiagaraStatelessEmitterTemplateToolkit::SpawnTab_OutputVar
 		];
 }
 
+TSharedRef<SDockTab> FNiagaraStatelessEmitterTemplateToolkit::SpawnTab_CodeView(const FSpawnTabArgs& Args)
+{
+	return
+		SNew(SDockTab)
+		[
+			SNew(SNiagaraStatelessEmitterTemplateCodeView)
+			.ViewModel(ViewModel)
+		];
+}
 
 #undef LOCTEXT_NAMESPACE
diff --git a/Engine/Plugins/FX/Niagara/Source/NiagaraEditor/Private/Toolkits/NiagaraStatelessEmitterTemplateToolkit.h b/Engine/Plugins/FX/Niagara/Source/NiagaraEditor/Private/Toolkits/NiagaraStatelessEmitterTemplateToolkit.h
--- a/Engine/Plugins/FX/Niagara/Source/NiagaraEditor/Private/Toolkits/NiagaraStatelessEmitterTemplateToolkit.h
+++ b/Engine/Plugins/FX/Niagara/Source/NiagaraEditor/Private/Toolkits/NiagaraStatelessEmitterTemplateToolkit.h
@@ -37,6 +37,7 @@ protected:
 	TSharedRef<SDockTab> SpawnTab_Features(const FSpawnTabArgs& Args);
 	TSharedRef<SDockTab> SpawnTab_Modules(const FSpawnTabArgs& Args);
 	TSharedRef<SDockTab> SpawnTab_OutputVariables(const FSpawnTabArgs& Args);
+	TSharedRef<SDockTab> SpawnTab_CodeView(const FSpawnTabArgs& Args);
 
 private:
 	TWeakObjectPtr<UNiagaraStatelessEmitterTemplate>		WeakEmitterTemplate;
diff --git a/Engine/Plugins/FX/Niagara/Source/NiagaraEditor/Private/ViewModels/NiagaraStatelessEmitterTemplateViewModel.cpp b/Engine/Plugins/FX/Niagara/Source/NiagaraEditor/Private/ViewModels/NiagaraStatelessEmitterTemplateViewModel.cpp
--- a/Engine/Plugins/FX/Niagara/Source/NiagaraEditor/Private/ViewModels/NiagaraStatelessEmitterTemplateViewModel.cpp
+++ b/Engine/Plugins/FX/Niagara/Source/NiagaraEditor/Private/ViewModels/NiagaraStatelessEmitterTemplateViewModel.cpp
@@ -3,6 +3,7 @@
 #include "ViewModels/NiagaraStatelessEmitterTemplateViewModel.h"
 
 #include "Stateless/NiagaraStatelessEmitterTemplate.h"
+#include "Stateless/NiagaraStatelessEmitterTranslator.h"
 #include "Stateless/NiagaraStatelessEmitter.h"
 #include "Stateless/NiagaraStatelessModule.h"
 
@@ -21,6 +22,8 @@ FNiagaraStatelessEmitterTemplateViewModel::FNiagaraStatelessEmitterTemplateViewM
 	EmitterTemplateForEdit.Reset(Cast<UNiagaraStatelessEmitterTemplate>(StaticDuplicateObjectEx(DuplicateParameters)));
 
 	FEditorDelegates::PostUndoRedo.AddRaw(this, &FNiagaraStatelessEmitterTemplateViewModel::PostUndoRedo);
+
+	BuildOutputVariables();
 }
 
 FNiagaraStatelessEmitterTemplateViewModel::~FNiagaraStatelessEmitterTemplateViewModel()
@@ -40,12 +43,15 @@ void FNiagaraStatelessEmitterTemplateViewModel::NotifyPreChange(FProperty* Prope
 void FNiagaraStatelessEmitterTemplateViewModel::NotifyPostChange(const FPropertyChangedEvent& PropertyChangedEvent, FProperty* PropertyThatChanged)
 {
 	EmitterTemplate->ModifyTemplate(EmitterTemplateForEdit.Get());
+
+	BuildOutputVariables();
+
 	OnTemplateChangedDelegate.Broadcast();
 }
 
-TConstArrayView<FNiagaraVariableBase> FNiagaraStatelessEmitterTemplateViewModel::GetOututputComponents() const
+TConstArrayView<FNiagaraStatelessEmitterTemplateViewModel::FOutputVariable> FNiagaraStatelessEmitterTemplateViewModel::GetOututputVariables() const
 {
-	return EmitterTemplate->GetOututputComponents();
+	return OutputVariables;
 }
 
 TPair<ENiagaraStatelessFeatureMask, ENiagaraStatelessFeatureMask> FNiagaraStatelessEmitterTemplateViewModel::GetFeatureMaskRange() const
@@ -78,4 +84,57 @@ TPair<ENiagaraStatelessFeatureMask, ENiagaraStatelessFeatureMask> FNiagaraStatel
 	return {MinMask, MaxMask};
 }
 
+FString FNiagaraStatelessEmitterTemplateViewModel::GenerateComputeTemplateHLSL() const
+{
+	FString HlslOutput;
+	FNiagaraStatelessEmitterTranslator::TranslateToCompute(HlslOutput, EmitterTemplate.Get());
+	return MoveTemp(HlslOutput);
+}
+
+void FNiagaraStatelessEmitterTemplateViewModel::BuildOutputVariables()
+{
+	auto FindOrAddOutputVariable =
+		[this](const FNiagaraVariableBase& Variable) -> FOutputVariable&
+		{
+			for (FOutputVariable& Existing : OutputVariables)
+			{
+				if (Existing.Variable == Variable)
+				{
+					return Existing;
+				}
+			}
+			FOutputVariable& NewOutput = OutputVariables.AddDefaulted_GetRef();
+			NewOutput.Variable = Variable;
+			return NewOutput;
+		};
+
+	// Gather output variables
+	OutputVariables.Reset();
+	for (const FNiagaraVariableBase& Variable : EmitterTemplate->GetImplicitVariables())
+	{
+		FOutputVariable& Output = FindOrAddOutputVariable(Variable);
+		Output.bIsImplicit = true;
+	}
+
+	for (const FNiagaraVariableBase& Variable : EmitterTemplate->GetModuleVariables())
+	{
+		FOutputVariable& Output = FindOrAddOutputVariable(Variable);
+		Output.bIsModule = true;
+	}
+
+	for (const FNiagaraVariableBase& Variable : EmitterTemplate->GetShaderOutputVariables())
+	{
+		FOutputVariable& Output = FindOrAddOutputVariable(Variable);
+		Output.bIsShader = true;
+	}
+
+	Algo::Sort(
+		OutputVariables,
+		[](const FOutputVariable& Lhs, const FOutputVariable& Rhs)
+		{
+			return Lhs.Variable.GetName().LexicalLess(Rhs.Variable.GetName());
+		}
+	);
+}
+
 #undef LOCTEXT_NAMESPACE
diff --git a/Engine/Plugins/FX/Niagara/Source/NiagaraEditor/Private/Widgets/SNiagaraStatelessEmitterTemplateWidgets.cpp b/Engine/Plugins/FX/Niagara/Source/NiagaraEditor/Private/Widgets/SNiagaraStatelessEmitterTemplateWidgets.cpp
--- a/Engine/Plugins/FX/Niagara/Source/NiagaraEditor/Private/Widgets/SNiagaraStatelessEmitterTemplateWidgets.cpp
+++ b/Engine/Plugins/FX/Niagara/Source/NiagaraEditor/Private/Widgets/SNiagaraStatelessEmitterTemplateWidgets.cpp
@@ -4,15 +4,19 @@
 #include "Stateless/NiagaraStatelessCommon.h"
 #include "Stateless/NiagaraStatelessEmitterTemplate.h"
 #include "ViewModels/NiagaraStatelessEmitterTemplateViewModel.h"
+#include "Widgets/NiagaraHLSLSyntaxHighlighter.h"
 #include "Widgets/SNiagaraParameterName.h"
 #include "NiagaraEditorUtilities.h"
 
 #include "Components/VerticalBox.h"
+#include "HAL/PlatformApplicationMisc.h"
 #include "Widgets/Layout/SGridPanel.h"
+#include "Widgets/Text/SMultiLineEditableText.h"
 #include "DetailsViewArgs.h"
 #include "IDetailsView.h"
 #include "PropertyEditorModule.h"
 
+
 #define LOCTEXT_NAMESPACE "NiagaraStatelessEmitterTemplateWidgets"
 
 /////////////////////////////////////////////////////////////////////////////////////////////////////////////// 
@@ -139,7 +143,11 @@ void SNiagaraStatelessEmitterTemplateOutputVariables::Construct(const FArguments
 
 	ChildSlot
 	[
-		SAssignNew(GridPanel, SGridPanel)
+		SNew(SScrollBox)
+		+SScrollBox::Slot().Padding(5)
+		[
+			SAssignNew(GridPanel, SGridPanel)
+		]
 	];
 
 	OnRebuildWidget();
@@ -149,16 +157,119 @@ void SNiagaraStatelessEmitterTemplateOutputVariables::OnRebuildWidget()
 {
 	GridPanel->ClearChildren();
 
-	int32 NumSlots = 0;
+	int32 iRow = 0;
+	{
+		int32 iColumn = 0;
+		GridPanel->AddSlot(iColumn++, iRow)
+		.Padding(5)
+		[
+			SNew(STextBlock)
+			.Text(LOCTEXT("Variable", "Variable"))
+		];
+		GridPanel->AddSlot(iColumn++, iRow)
+		.Padding(5)
+		[
+			SNew(STextBlock)
+			.Text(LOCTEXT("Implicit", "Implicit"))
+		];
+		GridPanel->AddSlot(iColumn++, iRow)
+		.Padding(5)
+		[
+			SNew(STextBlock)
+			.Text(LOCTEXT("Module", "Module"))
+		];
+		GridPanel->AddSlot(iColumn++, iRow)
+		.Padding(5)
+		[
+			SNew(STextBlock)
+			.Text(LOCTEXT("GPU", "GPU"))
+		];
+	}
+	++iRow;
+
+	const FText TextTrue = LOCTEXT("True", "True");
+	const FText TextFalse = LOCTEXT("False", "False");
 
-	for (const FNiagaraVariableBase& Variable : ViewModel->GetOututputComponents())
+	for (const FNiagaraStatelessEmitterTemplateViewModel::FOutputVariable& OutputVariable : ViewModel->GetOututputVariables())
 	{
-		GridPanel->AddSlot(0, NumSlots++)
-		//.Padding(5)
+		int32 iColumn = 0;
+
+		GridPanel->AddSlot(iColumn++, iRow)
 		[
-			FNiagaraParameterUtilities::GetParameterWidget(FNiagaraVariable(Variable), true, false)
+			FNiagaraParameterUtilities::GetParameterWidget(FNiagaraVariable(OutputVariable.Variable), true, false)
 		];
+		GridPanel->AddSlot(iColumn++, iRow)
+		[
+			SNew(STextBlock).Text(OutputVariable.bIsImplicit ? TextTrue : TextFalse)
+		];
+		GridPanel->AddSlot(iColumn++, iRow)
+		[
+			SNew(STextBlock).Text(OutputVariable.bIsModule ? TextTrue : TextFalse)
+		];
+		GridPanel->AddSlot(iColumn++, iRow)
+		[
+			SNew(STextBlock).Text(OutputVariable.bIsShader ? TextTrue : TextFalse)
+		];
+		++iRow;
 	}
 }
 
+/////////////////////////////////////////////////////////////////////////////////////////////////////////////// 
+
+void SNiagaraStatelessEmitterTemplateCodeView::Construct(const FArguments& InArgs)
+{
+	ViewModel = InArgs._ViewModel;
+	ViewModel->OnTemplateChanged().AddSP(this, &SNiagaraStatelessEmitterTemplateCodeView::OnRebuildWidget);
+
+	SyntaxHighlighter = FNiagaraHLSLSyntaxHighlighter::Create();
+
+	ChildSlot
+	[
+		SNew(SVerticalBox)
+		+SVerticalBox::Slot()
+		.AutoHeight()
+		[
+			SNew(SHorizontalBox)
+			+SHorizontalBox::Slot()
+			.AutoWidth()
+			[
+				SNew(SButton)
+				.Text(LOCTEXT("CopyToClipboard", "Copy To Clipboard"))
+				.OnPressed(this, &SNiagaraStatelessEmitterTemplateCodeView::OnCopyToClipboard)
+			]
+		]
+		+ SVerticalBox::Slot()
+		.FillHeight(1)
+		[
+			SNew(SScrollBox)
+			+SScrollBox::Slot().Padding(5)
+			[
+				SNew(SMultiLineEditableText)
+				.Text(this, &SNiagaraStatelessEmitterTemplateCodeView::GetHlslText)
+				.TextStyle(FAppStyle::Get(), "MessageLog")
+				.IsReadOnly(true)
+				.Marshaller(SyntaxHighlighter)
+			]
+		]
+	];
+
+	OnRebuildWidget();
+}
+
+void SNiagaraStatelessEmitterTemplateCodeView::OnRebuildWidget()
+{
+	HlslAsString = ViewModel->GenerateComputeTemplateHLSL();
+	HlslAsText = FText::FromString(HlslAsString);
+}
+
+FText SNiagaraStatelessEmitterTemplateCodeView::GetHlslText() const
+{
+	return HlslAsText;
+}
+
+void SNiagaraStatelessEmitterTemplateCodeView::OnCopyToClipboard()
+{
+	FPlatformApplicationMisc::ClipboardCopy(*HlslAsString);
+}
+
 #undef LOCTEXT_NAMESPACE
diff --git a/Engine/Plugins/FX/Niagara/Source/NiagaraEditor/Private/Widgets/SNiagaraStatelessEmitterTemplateWidgets.h b/Engine/Plugins/FX/Niagara/Source/NiagaraEditor/Private/Widgets/SNiagaraStatelessEmitterTemplateWidgets.h
--- a/Engine/Plugins/FX/Niagara/Source/NiagaraEditor/Private/Widgets/SNiagaraStatelessEmitterTemplateWidgets.h
+++ b/Engine/Plugins/FX/Niagara/Source/NiagaraEditor/Private/Widgets/SNiagaraStatelessEmitterTemplateWidgets.h
@@ -7,6 +7,7 @@
 
 class FNiagaraStatelessEmitterTemplateViewModel;
 class SGridPanel;
+class ITextLayoutMarshaller;
 
 /////////////////////////////////////////////////////////////////////////////////////////////////////////////// 
 
@@ -59,3 +60,27 @@ private:
 	TSharedPtr<SGridPanel>									GridPanel;
 };
 
+/////////////////////////////////////////////////////////////////////////////////////////////////////////////// 
+
+class SNiagaraStatelessEmitterTemplateCodeView : public SCompoundWidget
+{
+public:
+	SLATE_BEGIN_ARGS(SNiagaraStatelessEmitterTemplateCodeView) {}
+		SLATE_ARGUMENT(TSharedPtr<FNiagaraStatelessEmitterTemplateViewModel>, ViewModel)
+	SLATE_END_ARGS()
+
+	void Construct(const FArguments& InArgs);
+
+	void OnRebuildWidget();
+
+	FText GetHlslText() const;
+
+	void OnCopyToClipboard();
+
+private:
+	TSharedPtr<FNiagaraStatelessEmitterTemplateViewModel>	ViewModel;
+	TSharedPtr<ITextLayoutMarshaller>						SyntaxHighlighter;
+	FString													HlslAsString;
+	FText													HlslAsText;
+};
+
diff --git a/Engine/Plugins/FX/Niagara/Source/NiagaraShader/Internal/Stateless/NiagaraStatelessEmitterShaders.h b/Engine/Plugins/FX/Niagara/Source/NiagaraShader/Internal/Stateless/NiagaraStatelessEmitterShaders.h
--- a/Engine/Plugins/FX/Niagara/Source/NiagaraShader/Internal/Stateless/NiagaraStatelessEmitterShaders.h
+++ b/Engine/Plugins/FX/Niagara/Source/NiagaraShader/Internal/Stateless/NiagaraStatelessEmitterShaders.h
@@ -45,12 +45,12 @@ namespace NiagaraStateless
 			ADD_STATELESS_MODULE(FSpriteRotationRateModule_ShaderParameters)
 			ADD_STATELESS_MODULE(FSubUVAnimationModule_ShaderParameters)
 
-			SHADER_PARAMETER(int, Permutation_UniqueIndexComponent)
+			SHADER_PARAMETER(int, Permutation_UniqueIDComponent)
 			SHADER_PARAMETER(int, Permutation_MaterialRandomComponent)
 			SHADER_PARAMETER(int, Permutation_PositionComponent)
 			SHADER_PARAMETER(int, Permutation_CameraOffsetComponent)
 			SHADER_PARAMETER(int, Permutation_ColorComponent)
-			SHADER_PARAMETER(int, Permutation_DynamicMaterialParameter0Component)
+			SHADER_PARAMETER(int, Permutation_DynamicMaterialParameterComponent)
 			SHADER_PARAMETER(int, Permutation_DynamicMaterialParameter1Component)
 			SHADER_PARAMETER(int, Permutation_DynamicMaterialParameter2Component)
 			SHADER_PARAMETER(int, Permutation_DynamicMaterialParameter3Component)
