- Move module implementations from header into cpp files

#tests EngineTests / personal project / etc
#rnx

diff --git a/Engine/Plugins/FX/Niagara/Source/Niagara/Internal/Stateless/Modules/NiagaraStatelessModule_AccelerationForce.h b/Engine/Plugins/FX/Niagara/Source/Niagara/Internal/Stateless/Modules/NiagaraStatelessModule_AccelerationForce.h
--- a/Engine/Plugins/FX/Niagara/Source/Niagara/Internal/Stateless/Modules/NiagaraStatelessModule_AccelerationForce.h
+++ b/Engine/Plugins/FX/Niagara/Source/Niagara/Internal/Stateless/Modules/NiagaraStatelessModule_AccelerationForce.h
@@ -3,8 +3,6 @@
 #pragma once
 
 #include "Stateless/NiagaraStatelessModule.h"
-#include "Stateless/NiagaraStatelessEmitterDataBuildContext.h"
-#include "Stateless/NiagaraStatelessModuleShaderParameters.h"
 #include "Stateless/Modules/NiagaraStatelessModuleCommon.h"
 
 #include "NiagaraStatelessModule_AccelerationForce.generated.h"
@@ -22,20 +20,7 @@ public:
 	UPROPERTY(EditAnywhere, Category = "Parameters", meta = (SegmentedDisplay))
 	ENiagaraCoordinateSpace CoordinateSpace = ENiagaraCoordinateSpace::Local;
 
-	virtual void BuildEmitterData(const FNiagaraStatelessEmitterDataBuildContext& BuildContext) const override
-	{
-		if (!IsModuleEnabled())
-		{
-			return;
-		}
-
-		const FNiagaraStatelessRangeVector3 AccelerationRange = AccelerationDistribution.CalculateRange(FVector3f::ZeroVector);
-
-		NiagaraStateless::FPhysicsBuildData& PhysicsBuildData = BuildContext.GetTransientBuildData<NiagaraStateless::FPhysicsBuildData>();
-		PhysicsBuildData.AccelerationCoordinateSpace = CoordinateSpace;
-		PhysicsBuildData.AccelerationRange.Min += AccelerationRange.Min;
-		PhysicsBuildData.AccelerationRange.Max += AccelerationRange.Max;
-	}
+	virtual void BuildEmitterData(const FNiagaraStatelessEmitterDataBuildContext& BuildContext) const override;
 
 #if WITH_EDITOR
 	virtual bool CanDisableModule() const override { return true; }
diff --git a/Engine/Plugins/FX/Niagara/Source/Niagara/Internal/Stateless/Modules/NiagaraStatelessModule_CameraOffset.h b/Engine/Plugins/FX/Niagara/Source/Niagara/Internal/Stateless/Modules/NiagaraStatelessModule_CameraOffset.h
--- a/Engine/Plugins/FX/Niagara/Source/Niagara/Internal/Stateless/Modules/NiagaraStatelessModule_CameraOffset.h
+++ b/Engine/Plugins/FX/Niagara/Source/Niagara/Internal/Stateless/Modules/NiagaraStatelessModule_CameraOffset.h
@@ -3,9 +3,6 @@
 #pragma once
 
 #include "Stateless/NiagaraStatelessModule.h"
-#include "Stateless/NiagaraStatelessEmitterDataBuildContext.h"
-#include "Stateless/NiagaraStatelessModuleShaderParameters.h"
-#include "Stateless/NiagaraStatelessParticleSimContext.h"
 
 #include "NiagaraStatelessModule_CameraOffset.generated.h"
 
@@ -16,86 +13,19 @@ class UNiagaraStatelessModule_CameraOffset : public UNiagaraStatelessModule
 public:
 	GENERATED_BODY()
 
-	struct FModuleBuiltData
-	{
-		FUintVector3	DistributionParameters			= FUintVector3::ZeroValue;
-		int32			CameraVariableOffset			= INDEX_NONE;
-		int32			PreviousCameraVariableOffset	= INDEX_NONE;
-	};
-
 public:
-	using FParameters = NiagaraStateless::FCameraOffsetModule_ShaderParameters;
-
 	UPROPERTY(EditAnywhere, Category = "Parameters", meta = (DisplayName = "CameraOffset"))
 	FNiagaraDistributionFloat CameraOffsetDistribution = FNiagaraDistributionFloat(0.0f);
 
-	virtual void BuildEmitterData(const FNiagaraStatelessEmitterDataBuildContext& BuildContext) const override
-	{
-		const FNiagaraStatelessGlobals& StatelessGlobals = FNiagaraStatelessGlobals::Get();
-
-		FModuleBuiltData* BuiltData				= BuildContext.AllocateBuiltData<FModuleBuiltData>();
-		BuiltData->CameraVariableOffset			= BuildContext.FindParticleVariableIndex(StatelessGlobals.CameraOffsetVariable);
-		BuiltData->PreviousCameraVariableOffset	= BuildContext.FindParticleVariableIndex(StatelessGlobals.PreviousCameraOffsetVariable);
-
-		const bool bAttributesUsed = (BuiltData->CameraVariableOffset != INDEX_NONE || BuiltData->PreviousCameraVariableOffset != INDEX_NONE);
-		if (IsModuleEnabled() && bAttributesUsed)
-		{
-			BuiltData->DistributionParameters = BuildContext.AddDistribution(CameraOffsetDistribution);
-
-			BuildContext.AddParticleSimulationExecSimulate(&UNiagaraStatelessModule_CameraOffset::ParticleSimulate);
-		}
-	}
-
-	virtual void BuildShaderParameters(FNiagaraStatelessShaderParametersBuilder& ShaderParametersBuilder) const override
-	{
-		ShaderParametersBuilder.AddParameterNestedStruct<FParameters>();
-	}
-
-	virtual void SetShaderParameters(const FNiagaraStatelessSetShaderParameterContext& SetShaderParameterContext) const override
-	{
-		const FModuleBuiltData* ModuleBuiltData = SetShaderParameterContext.ReadBuiltData<FModuleBuiltData>();
-
-		FParameters* Parameters = SetShaderParameterContext.GetParameterNestedStruct<FParameters>();
-		Parameters->CameraOffset_Distribution = ModuleBuiltData->DistributionParameters;
-	}
-
-	static void ParticleSimulate(const NiagaraStateless::FParticleSimulationContext& ParticleSimulationContext)
-	{
-		using namespace NiagaraStateless;
-
-		const FModuleBuiltData* ModuleBuiltData	= ParticleSimulationContext.ReadBuiltData<FModuleBuiltData>();
-		const float* NormalizedAgeData			= ParticleSimulationContext.GetParticleNormalizedAge();
-		const float* PreviousNormalizedAgeData	= ParticleSimulationContext.GetParticlePreviousNormalizedAge();
-
-		for (uint32 i = 0; i < ParticleSimulationContext.GetNumInstances(); ++i)
-		{
-			const FStatelessDistributionSampler<float> CameraOffsetSampler(ParticleSimulationContext, ModuleBuiltData->DistributionParameters, i, 0);
-
-			const float NormalizedAge			= NormalizedAgeData[i];
-			const float PreviousNormalizedAge	= PreviousNormalizedAgeData[i];
-
-			const float CameraOffset			= CameraOffsetSampler.GetValue(ParticleSimulationContext, NormalizedAge);
-			const float PreviousCameraOffset	= CameraOffsetSampler.GetValue(ParticleSimulationContext, PreviousNormalizedAge);
-
-			ParticleSimulationContext.WriteParticleVariable(ModuleBuiltData->CameraVariableOffset, i, CameraOffset);
-			ParticleSimulationContext.WriteParticleVariable(ModuleBuiltData->PreviousCameraVariableOffset, i, PreviousCameraOffset);
-		}
-	}
+	virtual void BuildEmitterData(const FNiagaraStatelessEmitterDataBuildContext& BuildContext) const override;
+	virtual void BuildShaderParameters(FNiagaraStatelessShaderParametersBuilder& ShaderParametersBuilder) const override;
+	virtual void SetShaderParameters(const FNiagaraStatelessSetShaderParameterContext& SetShaderParameterContext) const override;
 
 #if WITH_EDITOR
 	virtual bool CanDisableModule() const override { return true; }
 #endif
 #if WITH_EDITORONLY_DATA
-	virtual const TCHAR* GetShaderTemplatePath() const override
-	{
-		return TEXT("/Plugin/FX/Niagara/Private/Stateless/Modules/NiagaraStatelessModule_CameraOffset.ush");
-	}
-
-	virtual void GetOutputVariables(TArray<FNiagaraVariableBase>& OutVariables, EVariableFilter Filter) const override
-	{
-		const FNiagaraStatelessGlobals& StatelessGlobals = FNiagaraStatelessGlobals::Get();
-		OutVariables.AddUnique(StatelessGlobals.CameraOffsetVariable);
-		OutVariables.AddUnique(StatelessGlobals.PreviousCameraOffsetVariable);
-	}
+	virtual const TCHAR* GetShaderTemplatePath() const override;
+	virtual void GetOutputVariables(TArray<FNiagaraVariableBase>& OutVariables, EVariableFilter Filter) const override;
 #endif
 };
diff --git a/Engine/Plugins/FX/Niagara/Source/Niagara/Internal/Stateless/Modules/NiagaraStatelessModule_Drag.h b/Engine/Plugins/FX/Niagara/Source/Niagara/Internal/Stateless/Modules/NiagaraStatelessModule_Drag.h
--- a/Engine/Plugins/FX/Niagara/Source/Niagara/Internal/Stateless/Modules/NiagaraStatelessModule_Drag.h
+++ b/Engine/Plugins/FX/Niagara/Source/Niagara/Internal/Stateless/Modules/NiagaraStatelessModule_Drag.h
@@ -3,8 +3,6 @@
 #pragma once
 
 #include "Stateless/NiagaraStatelessModule.h"
-#include "Stateless/NiagaraStatelessEmitterDataBuildContext.h"
-#include "Stateless/NiagaraStatelessModuleShaderParameters.h"
 #include "Stateless/Modules/NiagaraStatelessModuleCommon.h"
 
 #include "NiagaraStatelessModule_Drag.generated.h"
@@ -21,16 +19,7 @@ public:
 	UPROPERTY(EditAnywhere, Category = "Parameters", meta = (DisplayName = "Drag", DisableBindingDistribution))
 	FNiagaraDistributionRangeFloat DragDistribution = FNiagaraDistributionRangeFloat(DefaultDrag);
 
-	virtual void BuildEmitterData(const FNiagaraStatelessEmitterDataBuildContext& BuildContext) const override
-	{
-		if (!IsModuleEnabled())
-		{
-			return;
-		}
-
-		NiagaraStateless::FPhysicsBuildData& PhysicsBuildData = BuildContext.GetTransientBuildData<NiagaraStateless::FPhysicsBuildData>();
-		PhysicsBuildData.DragRange = DragDistribution.CalculateRange(DefaultDrag);
-	}
+	virtual void BuildEmitterData(const FNiagaraStatelessEmitterDataBuildContext& BuildContext) const override;
 
 #if WITH_EDITOR
 	virtual bool CanDisableModule() const override { return true; }
diff --git a/Engine/Plugins/FX/Niagara/Source/Niagara/Internal/Stateless/Modules/NiagaraStatelessModule_DynamicMaterialParameters.h b/Engine/Plugins/FX/Niagara/Source/Niagara/Internal/Stateless/Modules/NiagaraStatelessModule_DynamicMaterialParameters.h
--- a/Engine/Plugins/FX/Niagara/Source/Niagara/Internal/Stateless/Modules/NiagaraStatelessModule_DynamicMaterialParameters.h
+++ b/Engine/Plugins/FX/Niagara/Source/Niagara/Internal/Stateless/Modules/NiagaraStatelessModule_DynamicMaterialParameters.h
@@ -3,9 +3,6 @@
 #pragma once
 
 #include "Stateless/NiagaraStatelessModule.h"
-#include "Stateless/NiagaraStatelessEmitterDataBuildContext.h"
-#include "Stateless/NiagaraStatelessModuleShaderParameters.h"
-#include "Stateless/NiagaraStatelessParticleSimContext.h"
 
 #include "NiagaraStatelessModule_DynamicMaterialParameters.generated.h"
 
@@ -45,32 +42,7 @@ class UNiagaraStatelessModule_DynamicMaterialParameters : public UNiagaraStatele
 {
 	GENERATED_BODY()
 
-	static constexpr int32 NumParameters			= 4;
-	static constexpr int32 NumChannelPerParameter	= 4;
-
-	struct FModuleBuiltData
-	{
-		FModuleBuiltData()
-		{
-			for (FUintVector3& v : ParameterDistributions)
-			{
-				v = FUintVector3::ZeroValue;
-			}
-
-			for (int32& ParameterVariableOffset : ParameterVariableOffsets)
-			{
-				ParameterVariableOffset = INDEX_NONE;
-			}
-		}
-
-		uint32			ChannelMask = 0;
-		FUintVector3	ParameterDistributions[NumParameters * NumChannelPerParameter];
-		int32			ParameterVariableOffsets[NumParameters];
-	};
-
 public:
-	using FParameters = NiagaraStateless::FDynamicMaterialParametersModule_ShaderParameters;
-
 	UPROPERTY(EditAnywhere, Category = "Parameters", meta = (InlineEditConditionToggle))
 	uint32 bParameter0Enabled : 1 = true;
 
@@ -95,194 +67,21 @@ public:
 	UPROPERTY(EditAnywhere, Category = "Parameters", meta = (EditCondition = "bParameter3Enabled"))
 	FNiagaraStatelessDynamicParameterSet Parameter3;
 
-	const FNiagaraStatelessDynamicParameterSet& GetParameterSet(int32 ParameterIndex) const
-	{
-		switch (ParameterIndex)
-		{
-			case 0: return Parameter0;
-			case 1: return Parameter1;
-			case 2: return Parameter2;
-			case 3: return Parameter3;
-			default: checkNoEntry(); return Parameter0;
-		}
-	}
-
-	uint32 GetParameterChannelMask(int32 ParameterIndex) const
-	{
-		uint32 Mask = 0;
-		if ((ParameterIndex == 0 && bParameter0Enabled) ||
-			(ParameterIndex == 1 && bParameter1Enabled) ||
-			(ParameterIndex == 2 && bParameter2Enabled) ||
-			(ParameterIndex == 3 && bParameter3Enabled))
-		{
-			const FNiagaraStatelessDynamicParameterSet& ParameterSet = GetParameterSet(ParameterIndex);
-			Mask |= ParameterSet.bXChannelEnabled ? 1 << 0 : 0;
-			Mask |= ParameterSet.bYChannelEnabled ? 1 << 1 : 0;
-			Mask |= ParameterSet.bZChannelEnabled ? 1 << 2 : 0;
-			Mask |= ParameterSet.bWChannelEnabled ? 1 << 3 : 0;
-		}
-		return Mask;
-	}
-
-	const FNiagaraDistributionFloat& GetParameterChannelDistribution(int32 ParameterIndex, int32 ChannelIndex) const
-	{
-		const FNiagaraStatelessDynamicParameterSet& ParameterSet = GetParameterSet(ParameterIndex);
-		switch (ChannelIndex)
-		{
-			case 0: return ParameterSet.XChannelDistribution;
-			case 1: return ParameterSet.YChannelDistribution;
-			case 2: return ParameterSet.ZChannelDistribution;
-			case 3: return ParameterSet.WChannelDistribution;
-			default: checkNoEntry(); return ParameterSet.XChannelDistribution;
-		}
-	}
-
-	const FNiagaraVariableBase& GetParameterVariable(int32 ParameterIndex) const
-	{
-		const FNiagaraStatelessGlobals& StatelessGlobals = FNiagaraStatelessGlobals::Get();
-		switch (ParameterIndex)
-		{
-			default: checkNoEntry(); return StatelessGlobals.DynamicMaterialParameters0Variable;
-			case 0: return StatelessGlobals.DynamicMaterialParameters0Variable;
-			case 1: return StatelessGlobals.DynamicMaterialParameters1Variable;
-			case 2: return StatelessGlobals.DynamicMaterialParameters2Variable;
-			case 3: return StatelessGlobals.DynamicMaterialParameters3Variable;
-		}
-	}
-
-
-	int32 GetRendererChannelMask() const
-	{
-		uint32 ChannelMask = 0;
-		if (IsModuleEnabled())
-		{
-			for (int32 iParameter = 0; iParameter < NumParameters; ++iParameter)
-			{
-				ChannelMask |= GetParameterChannelMask(iParameter) << (iParameter * NumChannelPerParameter);
-			}
-		}
-		return ChannelMask;
-	}
-
-	virtual void BuildEmitterData(const FNiagaraStatelessEmitterDataBuildContext& BuildContext) const override
-	{
-		FModuleBuiltData* BuiltData = BuildContext.AllocateBuiltData<FModuleBuiltData>();
-
-		if (IsModuleEnabled())
-		{
-			for (int32 iParameter=0; iParameter < NumParameters; ++iParameter)
-			{
-				const uint32 ParameterChannelMask = GetParameterChannelMask(iParameter);
-				if (ParameterChannelMask == 0)
-				{
-					continue;
-				}
-
-				BuiltData->ParameterVariableOffsets[iParameter] = BuildContext.FindParticleVariableIndex(GetParameterVariable(iParameter));
-				if (BuiltData->ParameterVariableOffsets[iParameter] == INDEX_NONE)
-				{
-					continue;
-				}
+	const FNiagaraStatelessDynamicParameterSet& GetParameterSet(int32 ParameterIndex) const;
+	uint32 GetParameterChannelMask(int32 ParameterIndex) const;
+	const FNiagaraDistributionFloat& GetParameterChannelDistribution(int32 ParameterIndex, int32 ChannelIndex) const;
+	const FNiagaraVariableBase& GetParameterVariable(int32 ParameterIndex) const;
+	int32 GetRendererChannelMask() const;
 
-				BuiltData->ChannelMask |= ParameterChannelMask << (iParameter * NumChannelPerParameter);
-				for (int32 iChannel = 0; iChannel < NumChannelPerParameter; ++iChannel)
-				{
-					BuiltData->ParameterDistributions[(iParameter * NumChannelPerParameter) + iChannel] = BuildContext.AddDistribution(GetParameterChannelDistribution(iParameter, iChannel));
-				}
-			}
-
-			if (BuiltData->ChannelMask != 0)
-			{
-				BuildContext.AddParticleSimulationExecSimulate(&UNiagaraStatelessModule_DynamicMaterialParameters::ParticleSimulate);
-			}
-		}
-	}
-
-	virtual void BuildShaderParameters(FNiagaraStatelessShaderParametersBuilder& ShaderParametersBuilder) const override
-	{
-		ShaderParametersBuilder.AddParameterNestedStruct<FParameters>();
-	}
-
-	virtual void SetShaderParameters(const FNiagaraStatelessSetShaderParameterContext& SetShaderParameterContext) const override
-	{
-		const FModuleBuiltData* ModuleBuiltData = SetShaderParameterContext.ReadBuiltData<FModuleBuiltData>();
-
-		FParameters* Parameters = SetShaderParameterContext.GetParameterNestedStruct<FParameters>();
-		Parameters->DynamicMaterialParameters_ChannelMask = ModuleBuiltData->ChannelMask;
-		Parameters->DynamicMaterialParameters_Parameter0X = ModuleBuiltData->ParameterDistributions[0];
-		Parameters->DynamicMaterialParameters_Parameter0Y = ModuleBuiltData->ParameterDistributions[1];
-		Parameters->DynamicMaterialParameters_Parameter0Z = ModuleBuiltData->ParameterDistributions[2];
-		Parameters->DynamicMaterialParameters_Parameter0W = ModuleBuiltData->ParameterDistributions[3];
-		Parameters->DynamicMaterialParameters_Parameter1X = ModuleBuiltData->ParameterDistributions[4];
-		Parameters->DynamicMaterialParameters_Parameter1Y = ModuleBuiltData->ParameterDistributions[5];
-		Parameters->DynamicMaterialParameters_Parameter1Z = ModuleBuiltData->ParameterDistributions[6];
-		Parameters->DynamicMaterialParameters_Parameter1W = ModuleBuiltData->ParameterDistributions[7];
-		Parameters->DynamicMaterialParameters_Parameter2X = ModuleBuiltData->ParameterDistributions[8];
-		Parameters->DynamicMaterialParameters_Parameter2Y = ModuleBuiltData->ParameterDistributions[9];
-		Parameters->DynamicMaterialParameters_Parameter2Z = ModuleBuiltData->ParameterDistributions[10];
-		Parameters->DynamicMaterialParameters_Parameter2W = ModuleBuiltData->ParameterDistributions[11];
-		Parameters->DynamicMaterialParameters_Parameter3X = ModuleBuiltData->ParameterDistributions[12];
-		Parameters->DynamicMaterialParameters_Parameter3Y = ModuleBuiltData->ParameterDistributions[13];
-		Parameters->DynamicMaterialParameters_Parameter3Z = ModuleBuiltData->ParameterDistributions[14];
-		Parameters->DynamicMaterialParameters_Parameter3W = ModuleBuiltData->ParameterDistributions[15];
-	}
-
-	static void ParticleSimulate(const NiagaraStateless::FParticleSimulationContext& ParticleSimulationContext)
-	{
-		using namespace NiagaraStateless;
-
-		const FModuleBuiltData* ModuleBuiltData = ParticleSimulationContext.ReadBuiltData<FModuleBuiltData>();
-		const float* NormalizedAgeData = ParticleSimulationContext.GetParticleNormalizedAge();
-
-		for (uint32 i = 0; i < ParticleSimulationContext.GetNumInstances(); ++i)
-		{
-			for (int32 iParameter = 0; iParameter < NumParameters; ++iParameter)
-			{
-				const uint32 ChannelMask = (ModuleBuiltData->ChannelMask >> (iParameter * NumChannelPerParameter)) & 0xf;
-				if (ChannelMask == 0)
-				{
-					continue;
-				}
-
-				const int32 FirstChannel = iParameter * NumChannelPerParameter;
-				const FStatelessDistributionSampler<float> XChannelSampler(ParticleSimulationContext, ModuleBuiltData->ParameterDistributions[FirstChannel + 0], i, FirstChannel + 0);
-				const FStatelessDistributionSampler<float> YChannelSampler(ParticleSimulationContext, ModuleBuiltData->ParameterDistributions[FirstChannel + 1], i, FirstChannel + 1);
-				const FStatelessDistributionSampler<float> ZChannelSampler(ParticleSimulationContext, ModuleBuiltData->ParameterDistributions[FirstChannel + 2], i, FirstChannel + 2);
-				const FStatelessDistributionSampler<float> WChannelSampler(ParticleSimulationContext, ModuleBuiltData->ParameterDistributions[FirstChannel + 3], i, FirstChannel + 3);
-
-				const float NormalizedAge = NormalizedAgeData[i];
-				const FVector4f DynamicParameter(
-					((ChannelMask & 0x1) != 0) ? XChannelSampler.GetValue(ParticleSimulationContext, NormalizedAge) : 0.0f,
-					((ChannelMask & 0x2) != 0) ? YChannelSampler.GetValue(ParticleSimulationContext, NormalizedAge) : 0.0f,
-					((ChannelMask & 0x4) != 0) ? ZChannelSampler.GetValue(ParticleSimulationContext, NormalizedAge) : 0.0f,
-					((ChannelMask & 0x8) != 0) ? WChannelSampler.GetValue(ParticleSimulationContext, NormalizedAge) : 0.0f
-				);
-
-				ParticleSimulationContext.WriteParticleVariable(ModuleBuiltData->ParameterVariableOffsets[iParameter], i, DynamicParameter);
-			}
-		}
-	}
+	virtual void BuildEmitterData(const FNiagaraStatelessEmitterDataBuildContext& BuildContext) const override;
+	virtual void BuildShaderParameters(FNiagaraStatelessShaderParametersBuilder& ShaderParametersBuilder) const override;
+	virtual void SetShaderParameters(const FNiagaraStatelessSetShaderParameterContext& SetShaderParameterContext) const override;
 
 #if WITH_EDITOR
 	virtual bool CanDisableModule() const override { return true; }
 #endif
 #if WITH_EDITORONLY_DATA
-	virtual const TCHAR* GetShaderTemplatePath() const override
-	{
-		return TEXT("/Plugin/FX/Niagara/Private/Stateless/Modules/NiagaraStatelessModule_DynamicMaterialParameters.ush");
-	}
-
-	virtual void GetOutputVariables(TArray<FNiagaraVariableBase>& OutVariables, EVariableFilter Filter) const override
-	{
-		const FNiagaraStatelessGlobals& StatelessGlobals = FNiagaraStatelessGlobals::Get();
-
-		for (int32 i = 0; i < NumParameters; ++i)
-		{
-			if (GetParameterChannelMask(i) != 0 || Filter == EVariableFilter::None)
-			{
-				OutVariables.AddUnique(GetParameterVariable(i));
-			}
-		}
-	}
+	virtual const TCHAR* GetShaderTemplatePath() const override;
+	virtual void GetOutputVariables(TArray<FNiagaraVariableBase>& OutVariables, EVariableFilter Filter) const override;
 #endif
 };
diff --git a/Engine/Plugins/FX/Niagara/Source/Niagara/Internal/Stateless/Modules/NiagaraStatelessModule_GravityForce.h b/Engine/Plugins/FX/Niagara/Source/Niagara/Internal/Stateless/Modules/NiagaraStatelessModule_GravityForce.h
--- a/Engine/Plugins/FX/Niagara/Source/Niagara/Internal/Stateless/Modules/NiagaraStatelessModule_GravityForce.h
+++ b/Engine/Plugins/FX/Niagara/Source/Niagara/Internal/Stateless/Modules/NiagaraStatelessModule_GravityForce.h
@@ -3,9 +3,6 @@
 #pragma once
 
 #include "Stateless/NiagaraStatelessModule.h"
-#include "Stateless/NiagaraStatelessEmitterDataBuildContext.h"
-#include "Stateless/NiagaraStatelessModuleShaderParameters.h"
-#include "Stateless/Modules/NiagaraStatelessModuleCommon.h"
 
 #include "NiagaraStatelessModule_GravityForce.generated.h"
 
@@ -19,18 +16,7 @@ public:
 	UPROPERTY(EditAnywhere, Category = "Parameters", meta = (DisplayName = "Gravity", DisableUniformDistribution, DisableBindingDistribution))
 	FNiagaraDistributionRangeVector3 GravityDistribution = FNiagaraDistributionRangeVector3(GetDefaultValue());
 
-	virtual void BuildEmitterData(const FNiagaraStatelessEmitterDataBuildContext& BuildContext) const override
-	{
-		if (!IsModuleEnabled())
-		{
-			return;
-		}
-		const FNiagaraStatelessRangeVector3 GravityRange = GravityDistribution.CalculateRange(GetDefaultValue());
-
-		NiagaraStateless::FPhysicsBuildData& PhysicsBuildData = BuildContext.GetTransientBuildData<NiagaraStateless::FPhysicsBuildData>();
-		PhysicsBuildData.GravityRange.Min = GravityRange.Min;
-		PhysicsBuildData.GravityRange.Max = GravityRange.Max;
-	}
+	virtual void BuildEmitterData(const FNiagaraStatelessEmitterDataBuildContext& BuildContext) const override;
 
 	static FVector3f GetDefaultValue() { return FVector3f(0.0f, 0.0f, -980.0f); }
 
diff --git a/Engine/Plugins/FX/Niagara/Source/Niagara/Internal/Stateless/Modules/NiagaraStatelessModule_InitializeParticle.h b/Engine/Plugins/FX/Niagara/Source/Niagara/Internal/Stateless/Modules/NiagaraStatelessModule_InitializeParticle.h
--- a/Engine/Plugins/FX/Niagara/Source/Niagara/Internal/Stateless/Modules/NiagaraStatelessModule_InitializeParticle.h
+++ b/Engine/Plugins/FX/Niagara/Source/Niagara/Internal/Stateless/Modules/NiagaraStatelessModule_InitializeParticle.h
@@ -3,9 +3,6 @@
 #pragma once
 
 #include "Stateless/NiagaraStatelessModule.h"
-#include "Stateless/NiagaraStatelessEmitterDataBuildContext.h"
-#include "Stateless/NiagaraStatelessModuleShaderParameters.h"
-#include "Stateless/NiagaraStatelessParticleSimContext.h"
 #include "Stateless/Modules/NiagaraStatelessModuleCommon.h"
 
 #include "NiagaraParameterBinding.h"
@@ -19,38 +16,7 @@ class UNiagaraStatelessModule_InitializeParticle : public UNiagaraStatelessModul
 {
 	GENERATED_BODY()
 
-	static const uint32 EInitializeParticleModuleFlag_UniformSpriteSize	= 1 << 0;
-	static const uint32 EInitializeParticleModuleFlag_UniformMeshScale	= 1 << 1;
-
-	struct FModuleBuiltData
-	{
-		uint32							ModuleFlags = 0;
-		FUintVector3					InitialPosition = FUintVector3::ZeroValue;
-		FNiagaraStatelessRangeFloat		LifetimeRange;
-		FNiagaraStatelessRangeColor		ColorRange;
-		FNiagaraStatelessRangeFloat		MassRange;
-		FNiagaraStatelessRangeVector2	SpriteSizeRange;
-		FNiagaraStatelessRangeFloat		SpriteRotationRange;
-		FNiagaraStatelessRangeVector3	MeshScaleRange;
-		FNiagaraStatelessRangeFloat		RibbonWidthRange;
-
-		int32							PositionVariableOffset = INDEX_NONE;
-		int32							ColorVariableOffset = INDEX_NONE;
-		int32							RibbonWidthVariableOffset = INDEX_NONE;
-		int32							SpriteSizeVariableOffset = INDEX_NONE;
-		int32							SpriteRotationVariableOffset = INDEX_NONE;
-		int32							ScaleVariableOffset = INDEX_NONE;
-
-		int32							PreviousPositionVariableOffset = INDEX_NONE;
-		int32							PreviousRibbonWidthVariableOffset = INDEX_NONE;
-		int32							PreviousSpriteSizeVariableOffset = INDEX_NONE;
-		int32							PreviousSpriteRotationVariableOffset = INDEX_NONE;
-		int32							PreviousScaleVariableOffset = INDEX_NONE;
-	};
-
 public:
-	using FParameters = NiagaraStateless::FInitializeParticleModule_ShaderParameters;
-
 	UPROPERTY(EditAnywhere, Category = "Parameters", meta = (DisplayName = "Lifetime", Units="s"))
 	FNiagaraDistributionRangeFloat LifetimeDistribution = FNiagaraDistributionRangeFloat(FNiagaraStatelessGlobals::GetDefaultLifetimeValue());
 
@@ -78,120 +44,12 @@ public:
 	UPROPERTY(EditAnywhere, Category = "Parameters", meta = (DisableCurveDistribution))
 	FNiagaraDistributionPosition InitialPositionDistribution = FNiagaraDistributionPosition(FVector3f::ZeroVector);
 
-	virtual void BuildEmitterData(const FNiagaraStatelessEmitterDataBuildContext& BuildContext) const override
-	{
-		BuildContext.AddParticleSimulationExecSimulate(&UNiagaraStatelessModule_InitializeParticle::ParticleSimulate);
-
-		FModuleBuiltData* BuiltData			= BuildContext.AllocateBuiltData<FModuleBuiltData>();
-		BuiltData->ModuleFlags				 = SpriteSizeDistribution.IsUniform() ? EInitializeParticleModuleFlag_UniformSpriteSize : 0;
-		BuiltData->ModuleFlags				|= MeshScaleDistribution.IsUniform() ? EInitializeParticleModuleFlag_UniformMeshScale : 0;
-
-		BuiltData->InitialPosition			= BuildContext.AddDistribution(InitialPositionDistribution);
-		BuiltData->LifetimeRange			= LifetimeDistribution.CalculateRange(FNiagaraStatelessGlobals::GetDefaultLifetimeValue());
-		BuiltData->ColorRange				= BuildContext.ConvertDistributionToRange(ColorDistribution, FNiagaraStatelessGlobals::GetDefaultColorValue());
-		BuiltData->MassRange				= BuildContext.ConvertDistributionToRange(MassDistribution, FNiagaraStatelessGlobals::GetDefaultMassValue());
-		BuiltData->SpriteSizeRange			= BuildContext.ConvertDistributionToRange(SpriteSizeDistribution, FNiagaraStatelessGlobals::GetDefaultSpriteSizeValue());
-		BuiltData->SpriteRotationRange		= BuildContext.ConvertDistributionToRange(SpriteRotationDistribution, FNiagaraStatelessGlobals::GetDefaultSpriteRotationValue());
-		BuiltData->MeshScaleRange			= BuildContext.ConvertDistributionToRange(MeshScaleDistribution, FNiagaraStatelessGlobals::GetDefaultScaleValue());
-		BuiltData->RibbonWidthRange			= BuildContext.ConvertDistributionToRange(RibbonWidthDistribution, FNiagaraStatelessGlobals::GetDefaultRibbonWidthValue());
-
-		const FNiagaraStatelessGlobals& StatelessGlobals = FNiagaraStatelessGlobals::Get();
-		BuiltData->PositionVariableOffset				= BuildContext.FindParticleVariableIndex(StatelessGlobals.PositionVariable);
-		BuiltData->ColorVariableOffset					= BuildContext.FindParticleVariableIndex(StatelessGlobals.ColorVariable);
-		BuiltData->RibbonWidthVariableOffset			= BuildContext.FindParticleVariableIndex(StatelessGlobals.RibbonWidthVariable);
-		BuiltData->SpriteSizeVariableOffset				= BuildContext.FindParticleVariableIndex(StatelessGlobals.SpriteSizeVariable);
-		BuiltData->SpriteRotationVariableOffset			= BuildContext.FindParticleVariableIndex(StatelessGlobals.SpriteRotationVariable);
-		BuiltData->ScaleVariableOffset					= BuildContext.FindParticleVariableIndex(StatelessGlobals.ScaleVariable);
-		BuiltData->PreviousPositionVariableOffset		= BuildContext.FindParticleVariableIndex(StatelessGlobals.PreviousPositionVariable);
-		BuiltData->PreviousRibbonWidthVariableOffset	= BuildContext.FindParticleVariableIndex(StatelessGlobals.PreviousRibbonWidthVariable);
-		BuiltData->PreviousSpriteSizeVariableOffset		= BuildContext.FindParticleVariableIndex(StatelessGlobals.PreviousSpriteSizeVariable);
-		BuiltData->PreviousSpriteRotationVariableOffset	= BuildContext.FindParticleVariableIndex(StatelessGlobals.PreviousSpriteRotationVariable);
-		BuiltData->PreviousScaleVariableOffset			= BuildContext.FindParticleVariableIndex(StatelessGlobals.PreviousScaleVariable);
-
-		NiagaraStateless::FPhysicsBuildData& PhysicsBuildData = BuildContext.GetTransientBuildData<NiagaraStateless::FPhysicsBuildData>();
-		PhysicsBuildData.MassRange			= MassDistribution.CalculateRange(FNiagaraStatelessGlobals::GetDefaultMassValue());
-	}
-
-	virtual void BuildShaderParameters(FNiagaraStatelessShaderParametersBuilder& ShaderParametersBuilder) const override
-	{
-		ShaderParametersBuilder.AddParameterNestedStruct<FParameters>();
-	}
-
-	virtual void SetShaderParameters(const FNiagaraStatelessSetShaderParameterContext& SetShaderParameterContext) const override
-	{
-		FParameters* Parameters = SetShaderParameterContext.GetParameterNestedStruct<FParameters>();
-		const FModuleBuiltData* ModuleBuiltData = SetShaderParameterContext.ReadBuiltData<FModuleBuiltData>();
-
-		Parameters->InitializeParticle_ModuleFlags			= ModuleBuiltData->ModuleFlags;
-		Parameters->InitializeParticle_InitialPosition		= ModuleBuiltData->InitialPosition;
-		SetShaderParameterContext.ConvertRangeToScaleBias(ModuleBuiltData->ColorRange,			Parameters->InitializeParticle_ColorScale, Parameters->InitializeParticle_ColorBias);
-		SetShaderParameterContext.ConvertRangeToScaleBias(ModuleBuiltData->SpriteSizeRange,		Parameters->InitializeParticle_SpriteSizeScale, Parameters->InitializeParticle_SpriteSizeBias);
-		SetShaderParameterContext.ConvertRangeToScaleBias(ModuleBuiltData->SpriteRotationRange, Parameters->InitializeParticle_SpriteRotationScale, Parameters->InitializeParticle_SpriteRotationBias);
-		SetShaderParameterContext.ConvertRangeToScaleBias(ModuleBuiltData->MeshScaleRange,		Parameters->InitializeParticle_MeshScaleScale, Parameters->InitializeParticle_MeshScaleBias);
-		SetShaderParameterContext.ConvertRangeToScaleBias(ModuleBuiltData->RibbonWidthRange,	Parameters->InitializeParticle_RibbonWidthScale, Parameters->InitializeParticle_RibbonWidthBias);
-	}
-
-	static void ParticleSimulate(const NiagaraStateless::FParticleSimulationContext& ParticleSimulationContext)
-	{
-		using namespace NiagaraStateless;
-
-		const FModuleBuiltData* ModuleBuiltData = ParticleSimulationContext.ReadBuiltData<FModuleBuiltData>();
-		const FParameters* ShaderParameters = ParticleSimulationContext.ReadParameterNestedStruct<FParameters>();
-
-		const bool bUniformSpriteSize = (ModuleBuiltData->ModuleFlags & EInitializeParticleModuleFlag_UniformSpriteSize) != 0;
-		const bool bUniformMeshScale  = (ModuleBuiltData->ModuleFlags & EInitializeParticleModuleFlag_UniformMeshScale) != 0;
-
-		for (uint32 i = 0; i < ParticleSimulationContext.GetNumInstances(); ++i)
-		{
-			const FStatelessDistributionSampler<FVector3f> PositionSampler(ParticleSimulationContext, ModuleBuiltData->InitialPosition, i, 0);
-
-			const FVector3f		Position	= PositionSampler.GetValue(ParticleSimulationContext, 0.0f);
-			const FLinearColor	Color		= ParticleSimulationContext.RandomScaleBiasFloat(i, 1, ShaderParameters->InitializeParticle_ColorScale, ShaderParameters->InitializeParticle_ColorBias);
-			const float			RibbonWidth	= ParticleSimulationContext.RandomScaleBiasFloat(i, 2, ShaderParameters->InitializeParticle_RibbonWidthScale, ShaderParameters->InitializeParticle_RibbonWidthBias);
-			const FVector2f		SpriteSize	= ParticleSimulationContext.RandomScaleBiasFloat(i, 3, ShaderParameters->InitializeParticle_SpriteSizeScale, ShaderParameters->InitializeParticle_SpriteSizeBias, bUniformSpriteSize);
-			const float			SpriteRot	= ParticleSimulationContext.RandomScaleBiasFloat(i, 4, ShaderParameters->InitializeParticle_SpriteRotationScale, ShaderParameters->InitializeParticle_SpriteRotationBias);
-			const FVector3f		Scale		= ParticleSimulationContext.RandomScaleBiasFloat(i, 5, ShaderParameters->InitializeParticle_MeshScaleScale, ShaderParameters->InitializeParticle_MeshScaleBias, bUniformMeshScale);
-
-			ParticleSimulationContext.WriteParticleVariable(ModuleBuiltData->PositionVariableOffset,				i, Position);
-			ParticleSimulationContext.WriteParticleVariable(ModuleBuiltData->ColorVariableOffset,					i, Color);
-			ParticleSimulationContext.WriteParticleVariable(ModuleBuiltData->RibbonWidthVariableOffset,				i, RibbonWidth);
-			ParticleSimulationContext.WriteParticleVariable(ModuleBuiltData->SpriteSizeVariableOffset,				i, SpriteSize);
-			ParticleSimulationContext.WriteParticleVariable(ModuleBuiltData->SpriteRotationVariableOffset,			i, SpriteRot);
-			ParticleSimulationContext.WriteParticleVariable(ModuleBuiltData->ScaleVariableOffset,					i, Scale);
-
-			ParticleSimulationContext.WriteParticleVariable(ModuleBuiltData->PreviousPositionVariableOffset,		i, Position);
-			ParticleSimulationContext.WriteParticleVariable(ModuleBuiltData->PreviousRibbonWidthVariableOffset,		i, RibbonWidth);
-			ParticleSimulationContext.WriteParticleVariable(ModuleBuiltData->PreviousSpriteSizeVariableOffset,		i, SpriteSize);
-			ParticleSimulationContext.WriteParticleVariable(ModuleBuiltData->PreviousSpriteRotationVariableOffset,	i, SpriteRot);
-			ParticleSimulationContext.WriteParticleVariable(ModuleBuiltData->PreviousScaleVariableOffset,			i, Scale);
-		}
-	}
+	virtual void BuildEmitterData(const FNiagaraStatelessEmitterDataBuildContext& BuildContext) const override;
+	virtual void BuildShaderParameters(FNiagaraStatelessShaderParametersBuilder& ShaderParametersBuilder) const override;
+	virtual void SetShaderParameters(const FNiagaraStatelessSetShaderParameterContext& SetShaderParameterContext) const override;
 
 #if WITH_EDITORONLY_DATA
-	virtual const TCHAR* GetShaderTemplatePath() const override
-	{
-		return TEXT("/Plugin/FX/Niagara/Private/Stateless/Modules/NiagaraStatelessModule_InitializeParticle.ush");
-	}
-
-	virtual void GetOutputVariables(TArray<FNiagaraVariableBase>& OutVariables, EVariableFilter Filter) const override
-	{
-		const FNiagaraStatelessGlobals& StatelessGlobals = FNiagaraStatelessGlobals::Get();
-		OutVariables.AddUnique(StatelessGlobals.PositionVariable);
-		OutVariables.AddUnique(StatelessGlobals.ColorVariable);
-		OutVariables.AddUnique(StatelessGlobals.SpriteSizeVariable);
-		OutVariables.AddUnique(StatelessGlobals.SpriteRotationVariable);
-		OutVariables.AddUnique(StatelessGlobals.ScaleVariable);
-
-		OutVariables.AddUnique(StatelessGlobals.PreviousPositionVariable);
-		OutVariables.AddUnique(StatelessGlobals.PreviousSpriteSizeVariable);
-		OutVariables.AddUnique(StatelessGlobals.PreviousSpriteRotationVariable);
-		OutVariables.AddUnique(StatelessGlobals.PreviousScaleVariable);
-
-		if (bWriteRibbonWidth || Filter == EVariableFilter::None)
-		{
-			OutVariables.AddUnique(StatelessGlobals.RibbonWidthVariable);
-			OutVariables.AddUnique(StatelessGlobals.PreviousRibbonWidthVariable);
-		}
-	}
+	virtual const TCHAR* GetShaderTemplatePath() const override;
+	virtual void GetOutputVariables(TArray<FNiagaraVariableBase>& OutVariables, EVariableFilter Filter) const override;
 #endif
 };
diff --git a/Engine/Plugins/FX/Niagara/Source/Niagara/Internal/Stateless/Modules/NiagaraStatelessModule_MeshRotationRate.h b/Engine/Plugins/FX/Niagara/Source/Niagara/Internal/Stateless/Modules/NiagaraStatelessModule_MeshRotationRate.h
--- a/Engine/Plugins/FX/Niagara/Source/Niagara/Internal/Stateless/Modules/NiagaraStatelessModule_MeshRotationRate.h
+++ b/Engine/Plugins/FX/Niagara/Source/Niagara/Internal/Stateless/Modules/NiagaraStatelessModule_MeshRotationRate.h
@@ -3,9 +3,6 @@
 #pragma once
 
 #include "Stateless/NiagaraStatelessModule.h"
-#include "Stateless/NiagaraStatelessEmitterDataBuildContext.h"
-#include "Stateless/NiagaraStatelessModuleShaderParameters.h"
-#include "Stateless/NiagaraStatelessParticleSimContext.h"
 
 #include "NiagaraStatelessModule_MeshRotationRate.generated.h"
 
@@ -15,18 +12,7 @@ class UNiagaraStatelessModule_MeshRotationRate : public UNiagaraStatelessModule
 {
 	GENERATED_BODY()
 
-	struct FModuleBuiltData
-	{
-		bool							ModuleEnabled = false;
-		FNiagaraStatelessRangeVector3	RotationRange;
-		FUintVector3					RateScaleParameters = FUintVector3::ZeroValue;
-		int32							MeshOrientationVariableOffset = INDEX_NONE;
-		int32							PreviousMeshOrientationVariableOffset = INDEX_NONE;
-	};
-
 public:
-	using FParameters = NiagaraStateless::FMeshRotationRateModule_ShaderParameters;
-
 	UPROPERTY(EditAnywhere, Category = "Parameters", meta = (InlineEditConditionToggle))
 	bool bUseRateScale = false;
 
@@ -36,93 +22,15 @@ public:
 	UPROPERTY(EditAnywhere, Category = "Parameters", meta = (DisplayName = "Rate Scale", EditCondition = "bUseRateScale"))
 	FNiagaraDistributionCurveVector3 RateScaleDistribution = FNiagaraDistributionCurveVector3(ENiagaraDistributionCurveLUTMode::Accumulate);
 
-	virtual void BuildEmitterData(const FNiagaraStatelessEmitterDataBuildContext& BuildContext) const override
-	{
-		FModuleBuiltData* BuiltData = BuildContext.AllocateBuiltData<FModuleBuiltData>();
-		if (!IsModuleEnabled())
-		{
-			return;
-		}
-
-		const FNiagaraStatelessGlobals& StatelessGlobals	= FNiagaraStatelessGlobals::Get();
-		BuiltData->MeshOrientationVariableOffset			= BuildContext.FindParticleVariableIndex(StatelessGlobals.MeshOrientationVariable);
-		BuiltData->PreviousMeshOrientationVariableOffset	= BuildContext.FindParticleVariableIndex(StatelessGlobals.PreviousMeshOrientationVariable);
-
-		if (BuiltData->MeshOrientationVariableOffset != INDEX_NONE || BuiltData->PreviousMeshOrientationVariableOffset != INDEX_NONE)
-		{
-			BuiltData->ModuleEnabled = true;
-			BuiltData->RotationRange = BuildContext.ConvertDistributionToRange(RotationRateDistribution, FVector3f::ZeroVector);
-			BuiltData->RotationRange.Min *= 1.0f / 360.0f;
-			BuiltData->RotationRange.Max *= 1.0f / 360.0f;
-			if (bUseRateScale)
-			{
-				BuiltData->RateScaleParameters = BuildContext.AddDistributionAsCurve(RateScaleDistribution, 1.0f);
-			}
-			else
-			{
-				BuiltData->RateScaleParameters = BuildContext.AddDistributionAsCurve(FNiagaraDistributionCurveVector3(ENiagaraDistributionCurveLUTMode::Accumulate), 1.0f);
-			}
-
-			BuildContext.AddParticleSimulationExecSimulate(&UNiagaraStatelessModule_MeshRotationRate::ParticleSimulate);
-		}
-	}
-
-	virtual void BuildShaderParameters(FNiagaraStatelessShaderParametersBuilder& ShaderParametersBuilder) const override
-	{
-		ShaderParametersBuilder.AddParameterNestedStruct<FParameters>();
-	}
-
-	virtual void SetShaderParameters(const FNiagaraStatelessSetShaderParameterContext& SetShaderParameterContext) const override
-	{
-		FParameters* Parameters = SetShaderParameterContext.GetParameterNestedStruct<FParameters>();
-		const FModuleBuiltData* ModuleBuiltData = SetShaderParameterContext.ReadBuiltData<FModuleBuiltData>();
-		Parameters->MeshRotationRate_ModuleEnabled = ModuleBuiltData->ModuleEnabled;
-		SetShaderParameterContext.ConvertRangeToScaleBias(ModuleBuiltData->RotationRange, Parameters->MeshRotationRate_Scale, Parameters->MeshRotationRate_Bias);
-		Parameters->MeshRotationRate_RateScaleParameters = ModuleBuiltData->RateScaleParameters;
-	}
-
-	static void ParticleSimulate(const NiagaraStateless::FParticleSimulationContext& ParticleSimulationContext)
-	{
-		using namespace NiagaraStateless;
-
-		const FModuleBuiltData* ModuleBuiltData = ParticleSimulationContext.ReadBuiltData<FModuleBuiltData>();
-		const FParameters* ShaderParameters = ParticleSimulationContext.ReadParameterNestedStruct<FParameters>();
-
-		const float* LifetimeData = ParticleSimulationContext.GetParticleLifetime();
-		const float* AgeData = ParticleSimulationContext.GetParticleNormalizedAge();
-		const float* PreviousAgeData = ParticleSimulationContext.GetParticlePreviousNormalizedAge();
-
-		for (uint32 i = 0; i < ParticleSimulationContext.GetNumInstances(); ++i)
-		{
-			const FVector3f RotationRate = ParticleSimulationContext.RandomScaleBiasFloat(i, 0, ShaderParameters->MeshRotationRate_Scale, ShaderParameters->MeshRotationRate_Bias);
-			const FVector3f MultiplyRate = ParticleSimulationContext.SampleCurve<FVector3f>(ShaderParameters->MeshRotationRate_RateScaleParameters, AgeData[i]);
-			const FVector3f PreviousMultiplyRate = ParticleSimulationContext.SampleCurve<FVector3f>(ShaderParameters->MeshRotationRate_RateScaleParameters, PreviousAgeData[i]);
-
-			FQuat4f MeshOrientation = ParticleSimulationContext.ReadParticleVariable(ModuleBuiltData->MeshOrientationVariableOffset, i, FQuat4f::Identity);
-			FQuat4f PreviousMeshOrientation = ParticleSimulationContext.ReadParticleVariable(ModuleBuiltData->PreviousMeshOrientationVariableOffset, i, FQuat4f::Identity);
-
-			MeshOrientation *= ParticleSimulationContext.RotatorToQuat(LifetimeData[i] * RotationRate * MultiplyRate);
-			PreviousMeshOrientation *= ParticleSimulationContext.RotatorToQuat(LifetimeData[i] * RotationRate * PreviousMultiplyRate);
-
-			ParticleSimulationContext.WriteParticleVariable(ModuleBuiltData->MeshOrientationVariableOffset, i, MeshOrientation);
-			ParticleSimulationContext.WriteParticleVariable(ModuleBuiltData->PreviousMeshOrientationVariableOffset, i, PreviousMeshOrientation);
-		}
-	}
+	virtual void BuildEmitterData(const FNiagaraStatelessEmitterDataBuildContext& BuildContext) const override;
+	virtual void BuildShaderParameters(FNiagaraStatelessShaderParametersBuilder& ShaderParametersBuilder) const override;
+	virtual void SetShaderParameters(const FNiagaraStatelessSetShaderParameterContext& SetShaderParameterContext) const override;
 
 #if WITH_EDITOR
 	virtual bool CanDisableModule() const override { return true; }
 #endif
 #if WITH_EDITORONLY_DATA
-	virtual const TCHAR* GetShaderTemplatePath() const override
-	{
-		return TEXT("/Plugin/FX/Niagara/Private/Stateless/Modules/NiagaraStatelessModule_MeshRotationRate.ush");
-	}
-
-	virtual void GetOutputVariables(TArray<FNiagaraVariableBase>& OutVariables, EVariableFilter Filter) const override
-	{
-		const FNiagaraStatelessGlobals& StatelessGlobals = FNiagaraStatelessGlobals::Get();
-		OutVariables.AddUnique(StatelessGlobals.MeshOrientationVariable);
-		OutVariables.AddUnique(StatelessGlobals.PreviousMeshOrientationVariable);
-	}
+	virtual const TCHAR* GetShaderTemplatePath() const override;
+	virtual void GetOutputVariables(TArray<FNiagaraVariableBase>& OutVariables, EVariableFilter Filter) const override;
 #endif
 };
diff --git a/Engine/Plugins/FX/Niagara/Source/Niagara/Internal/Stateless/Modules/NiagaraStatelessModule_ScaleColor.h b/Engine/Plugins/FX/Niagara/Source/Niagara/Internal/Stateless/Modules/NiagaraStatelessModule_ScaleColor.h
--- a/Engine/Plugins/FX/Niagara/Source/Niagara/Internal/Stateless/Modules/NiagaraStatelessModule_ScaleColor.h
+++ b/Engine/Plugins/FX/Niagara/Source/Niagara/Internal/Stateless/Modules/NiagaraStatelessModule_ScaleColor.h
@@ -3,9 +3,6 @@
 #pragma once
 
 #include "Stateless/NiagaraStatelessModule.h"
-#include "Stateless/NiagaraStatelessEmitterDataBuildContext.h"
-#include "Stateless/NiagaraStatelessModuleShaderParameters.h"
-#include "Stateless/NiagaraStatelessParticleSimContext.h"
 
 #include "NiagaraStatelessModule_ScaleColor.generated.h"
 
@@ -15,82 +12,19 @@ class UNiagaraStatelessModule_ScaleColor : public UNiagaraStatelessModule
 {
 	GENERATED_BODY()
 
-	struct FModuleBuiltData
-	{
-		FUintVector3	DistributionParameters = FUintVector3::ZeroValue;
-		int32			ColorVariableOffset = INDEX_NONE;
-	};
-
 public:
-	using FParameters = NiagaraStateless::FScaleColorModule_ShaderParameters;
-
 	UPROPERTY(EditAnywhere, Category = "Parameters", meta = (DisplayName = "Scale"))
 	FNiagaraDistributionColor ScaleDistribution = FNiagaraDistributionColor(FLinearColor::White);
 
-	virtual void BuildEmitterData(const FNiagaraStatelessEmitterDataBuildContext& BuildContext) const override
-	{
-		FModuleBuiltData* BuiltData = BuildContext.AllocateBuiltData<FModuleBuiltData>();
-		if (!IsModuleEnabled())
-		{
-			return;
-		}
-
-		const FNiagaraStatelessGlobals& StatelessGlobals = FNiagaraStatelessGlobals::Get();
-		BuiltData->ColorVariableOffset = BuildContext.FindParticleVariableIndex(StatelessGlobals.ColorVariable);
-
-		if (BuiltData->ColorVariableOffset != INDEX_NONE)
-		{
-			BuiltData->DistributionParameters = BuildContext.AddDistribution(ScaleDistribution);
-
-			BuildContext.AddParticleSimulationExecSimulate(&UNiagaraStatelessModule_ScaleColor::ParticleSimulate);
-		}
-	}
-
-	virtual void BuildShaderParameters(FNiagaraStatelessShaderParametersBuilder& ShaderParametersBuilder) const override
-	{
-		ShaderParametersBuilder.AddParameterNestedStruct<FParameters>();
-	}
-
-	virtual void SetShaderParameters(const FNiagaraStatelessSetShaderParameterContext& SetShaderParameterContext) const override
-	{
-		const FModuleBuiltData* ModuleBuiltData = SetShaderParameterContext.ReadBuiltData<FModuleBuiltData>();
-
-		FParameters* Parameters = SetShaderParameterContext.GetParameterNestedStruct<FParameters>();
-		Parameters->ScaleColor_Distribution = ModuleBuiltData->DistributionParameters;
-	}
-
-	static void ParticleSimulate(const NiagaraStateless::FParticleSimulationContext& ParticleSimulationContext)
-	{
-		using namespace NiagaraStateless;
-
-		const FModuleBuiltData* ModuleBuiltData = ParticleSimulationContext.ReadBuiltData<FModuleBuiltData>();
-		const float* NormalizedAgeData = ParticleSimulationContext.GetParticleNormalizedAge();
-
-		for (uint32 i = 0; i < ParticleSimulationContext.GetNumInstances(); ++i)
-		{
-			FStatelessDistributionSampler<FLinearColor> ColorSampler(ParticleSimulationContext, ModuleBuiltData->DistributionParameters, i, 0);
-
-			FLinearColor Color = ParticleSimulationContext.ReadParticleVariable(ModuleBuiltData->ColorVariableOffset, i, FLinearColor::White);
-
-			Color *= ColorSampler.GetValue(ParticleSimulationContext, NormalizedAgeData[i]);
-
-			ParticleSimulationContext.WriteParticleVariable(ModuleBuiltData->ColorVariableOffset, i, Color);
-		}
-	}
+	virtual void BuildEmitterData(const FNiagaraStatelessEmitterDataBuildContext& BuildContext) const override;
+	virtual void BuildShaderParameters(FNiagaraStatelessShaderParametersBuilder& ShaderParametersBuilder) const override;
+	virtual void SetShaderParameters(const FNiagaraStatelessSetShaderParameterContext& SetShaderParameterContext) const override;
 
 #if WITH_EDITOR
 	virtual bool CanDisableModule() const override { return true; }
 #endif
 #if WITH_EDITORONLY_DATA
-	virtual const TCHAR* GetShaderTemplatePath() const override
-	{
-		return TEXT("/Plugin/FX/Niagara/Private/Stateless/Modules/NiagaraStatelessModule_ScaleColor.ush");
-	}
-
-	virtual void GetOutputVariables(TArray<FNiagaraVariableBase>& OutVariables, EVariableFilter Filter) const override
-	{
-		const FNiagaraStatelessGlobals& StatelessGlobals = FNiagaraStatelessGlobals::Get();
-		OutVariables.AddUnique(StatelessGlobals.ColorVariable);
-	}
+	virtual const TCHAR* GetShaderTemplatePath() const override;
+	virtual void GetOutputVariables(TArray<FNiagaraVariableBase>& OutVariables, EVariableFilter Filter) const override;
 #endif
 };
diff --git a/Engine/Plugins/FX/Niagara/Source/Niagara/Internal/Stateless/Modules/NiagaraStatelessModule_ScaleMeshSize.h b/Engine/Plugins/FX/Niagara/Source/Niagara/Internal/Stateless/Modules/NiagaraStatelessModule_ScaleMeshSize.h
--- a/Engine/Plugins/FX/Niagara/Source/Niagara/Internal/Stateless/Modules/NiagaraStatelessModule_ScaleMeshSize.h
+++ b/Engine/Plugins/FX/Niagara/Source/Niagara/Internal/Stateless/Modules/NiagaraStatelessModule_ScaleMeshSize.h
@@ -3,9 +3,6 @@
 #pragma once
 
 #include "Stateless/NiagaraStatelessModule.h"
-#include "Stateless/NiagaraStatelessEmitterDataBuildContext.h"
-#include "Stateless/NiagaraStatelessModuleShaderParameters.h"
-#include "Stateless/NiagaraStatelessParticleSimContext.h"
 
 #include "NiagaraParameterBinding.h"
 
@@ -18,104 +15,17 @@ class UNiagaraStatelessModule_ScaleMeshSize : public UNiagaraStatelessModule
 {
 	GENERATED_BODY()
 
-	struct FModuleBuiltData
-	{
-		FUintVector3	DistributionParameters = FUintVector3::ZeroValue;
-		FVector3f		CurveScale = FVector3f::OneVector;
-		int32			CurveScaleOffset = INDEX_NONE;
-
-		int32			ScaleVariableOffset = INDEX_NONE;
-		int32			PreviousScaleVariableOffset = INDEX_NONE;
-	};
-
 public:
-	using FParameters = NiagaraStateless::FScaleMeshSizeModule_ShaderParameters;
-
 	UPROPERTY(EditAnywhere, Category = "Parameters", meta = (DisplayName = "Scale"))
 	FNiagaraDistributionVector3 ScaleDistribution = FNiagaraDistributionVector3(FVector3f::OneVector);
 
 	UPROPERTY(EditAnywhere, Category = "Parameters", meta = (EditConditionHides, EditCondition = "UseScaleCurveRange()"))
 	FNiagaraParameterBindingWithValue ScaleCurveRange;
 
-	virtual void PostInitProperties() override
-	{
-		Super::PostInitProperties();
-
-	#if WITH_EDITORONLY_DATA
-		if (HasAnyFlags(RF_ClassDefaultObject) == false)
-		{
-			ScaleCurveRange.SetUsage(ENiagaraParameterBindingUsage::NotParticle);
-			ScaleCurveRange.SetAllowedTypeDefinitions({ FNiagaraTypeDefinition::GetVec3Def() });
-			ScaleCurveRange.SetDefaultParameter(FNiagaraTypeDefinition::GetVec3Def(), FVector3f::OneVector);
-		}
-	#endif
-	}
-
-	virtual void BuildEmitterData(const FNiagaraStatelessEmitterDataBuildContext& BuildContext) const override
-	{
-		FModuleBuiltData* BuiltData = BuildContext.AllocateBuiltData<FModuleBuiltData>();
-		if (!IsModuleEnabled())
-		{
-			return;
-		}
-
-		const FNiagaraStatelessGlobals& StatelessGlobals = FNiagaraStatelessGlobals::Get();
-		BuiltData->ScaleVariableOffset = BuildContext.FindParticleVariableIndex(StatelessGlobals.ScaleVariable);
-		BuiltData->PreviousScaleVariableOffset = BuildContext.FindParticleVariableIndex(StatelessGlobals.PreviousScaleVariable);
-
-		if (BuiltData->ScaleVariableOffset == INDEX_NONE && BuiltData->PreviousScaleVariableOffset == INDEX_NONE)
-		{
-			return;
-		}
-
-		BuiltData->DistributionParameters = BuildContext.AddDistribution(ScaleDistribution);
-		if (UseScaleCurveRange())
-		{
-			BuiltData->CurveScaleOffset = BuildContext.AddRendererBinding(ScaleCurveRange.ResolvedParameter);
-			BuiltData->CurveScale = ScaleCurveRange.GetDefaultValue<FVector3f>();
-		}
-		BuildContext.AddParticleSimulationExecSimulate(&UNiagaraStatelessModule_ScaleMeshSize::ParticleSimulate);
-	}
-
-	virtual void BuildShaderParameters(FNiagaraStatelessShaderParametersBuilder& ShaderParametersBuilder) const override
-	{
-		ShaderParametersBuilder.AddParameterNestedStruct<FParameters>();
-	}
-
-	virtual void SetShaderParameters(const FNiagaraStatelessSetShaderParameterContext& SetShaderParameterContext) const override
-	{
-		const FModuleBuiltData* ModuleBuiltData = SetShaderParameterContext.ReadBuiltData<FModuleBuiltData>();
-
-		FParameters* Parameters						= SetShaderParameterContext.GetParameterNestedStruct<FParameters>();
-		Parameters->ScaleMeshSize_Distribution		= ModuleBuiltData->DistributionParameters;
-		Parameters->ScaleMeshSize_CurveScale		= ModuleBuiltData->CurveScale;
-		Parameters->ScaleMeshSize_CurveScaleOffset	= ModuleBuiltData->CurveScaleOffset;
-	}
-
-	static void ParticleSimulate(const NiagaraStateless::FParticleSimulationContext& ParticleSimulationContext)
-	{
-		using namespace NiagaraStateless;
-
-		const FModuleBuiltData* ModuleBuiltData = ParticleSimulationContext.ReadBuiltData<FModuleBuiltData>();
-		const float* NormalizedAgeData = ParticleSimulationContext.GetParticleNormalizedAge();
-		const float* PreviousNormalizedAgeData = ParticleSimulationContext.GetParticlePreviousNormalizedAge();
-
-		const FVector3f CurveScale = ParticleSimulationContext.GetParameterBufferFloat(ModuleBuiltData->CurveScaleOffset, ModuleBuiltData->CurveScale);
-
-		for (uint32 i = 0; i < ParticleSimulationContext.GetNumInstances(); ++i)
-		{
-			FStatelessDistributionSampler<FVector3f> ScaleSampler(ParticleSimulationContext, ModuleBuiltData->DistributionParameters, i, 0);
-
-			FVector3f Scale = ParticleSimulationContext.ReadParticleVariable(ModuleBuiltData->ScaleVariableOffset, i, FVector3f::OneVector);
-			FVector3f PreviousScale = ParticleSimulationContext.ReadParticleVariable(ModuleBuiltData->PreviousScaleVariableOffset, i, FVector3f::OneVector);
-
-			Scale = Scale * ScaleSampler.GetValue(ParticleSimulationContext, NormalizedAgeData[i]) * CurveScale;
-			PreviousScale = PreviousScale * ScaleSampler.GetValue(ParticleSimulationContext, PreviousNormalizedAgeData[i]) * CurveScale;
-
-			ParticleSimulationContext.WriteParticleVariable(ModuleBuiltData->ScaleVariableOffset, i, Scale);
-			ParticleSimulationContext.WriteParticleVariable(ModuleBuiltData->PreviousScaleVariableOffset, i, PreviousScale);
-		}
-	}
+	virtual void PostInitProperties() override;
+	virtual void BuildEmitterData(const FNiagaraStatelessEmitterDataBuildContext& BuildContext) const override;
+	virtual void BuildShaderParameters(FNiagaraStatelessShaderParametersBuilder& ShaderParametersBuilder) const override;
+	virtual void SetShaderParameters(const FNiagaraStatelessSetShaderParameterContext& SetShaderParameterContext) const override;
 
 	UFUNCTION()
 	bool UseScaleCurveRange() const { return ScaleDistribution.IsCurve(); }
@@ -124,16 +34,7 @@ public:
 	virtual bool CanDisableModule() const override { return true; }
 #endif
 #if WITH_EDITORONLY_DATA
-	virtual const TCHAR* GetShaderTemplatePath() const override
-	{
-		return TEXT("/Plugin/FX/Niagara/Private/Stateless/Modules/NiagaraStatelessModule_ScaleMeshSize.ush");
-	}
-
-	virtual void GetOutputVariables(TArray<FNiagaraVariableBase>& OutVariables, EVariableFilter Filter) const override
-	{
-		const FNiagaraStatelessGlobals& StatelessGlobals = FNiagaraStatelessGlobals::Get();
-		OutVariables.AddUnique(StatelessGlobals.ScaleVariable);
-		OutVariables.AddUnique(StatelessGlobals.PreviousScaleVariable);
-	}
+	virtual const TCHAR* GetShaderTemplatePath() const override;
+	virtual void GetOutputVariables(TArray<FNiagaraVariableBase>& OutVariables, EVariableFilter Filter) const override;
 #endif
 };
diff --git a/Engine/Plugins/FX/Niagara/Source/Niagara/Internal/Stateless/Modules/NiagaraStatelessModule_ScaleRibbonWidth.h b/Engine/Plugins/FX/Niagara/Source/Niagara/Internal/Stateless/Modules/NiagaraStatelessModule_ScaleRibbonWidth.h
--- a/Engine/Plugins/FX/Niagara/Source/Niagara/Internal/Stateless/Modules/NiagaraStatelessModule_ScaleRibbonWidth.h
+++ b/Engine/Plugins/FX/Niagara/Source/Niagara/Internal/Stateless/Modules/NiagaraStatelessModule_ScaleRibbonWidth.h
@@ -4,9 +4,6 @@
 
 #include "Stateless/NiagaraStatelessCommon.h"
 #include "Stateless/NiagaraStatelessModule.h"
-#include "Stateless/NiagaraStatelessEmitterDataBuildContext.h"
-#include "Stateless/NiagaraStatelessModuleShaderParameters.h"
-#include "Stateless/NiagaraStatelessParticleSimContext.h"
 #include "NiagaraParameterBinding.h"
 
 #include "NiagaraStatelessModule_ScaleRibbonWidth.generated.h"
@@ -18,101 +15,18 @@ class UNiagaraStatelessModule_ScaleRibbonWidth : public UNiagaraStatelessModule
 {
 	GENERATED_BODY()
 
-	struct FModuleBuiltData
-	{
-		FUintVector3	DistributionParameters = FUintVector3::ZeroValue;
-		float			CurveScale = 1.0f;
-		int32			CurveScaleOffset = INDEX_NONE;
-
-		int32			RibbonWidthVariableOffset = INDEX_NONE;
-		int32			PreviousRibbonWidthVariableOffset = INDEX_NONE;
-	};
-
 public:
-	using FParameters = NiagaraStateless::FScaleRibbonWidthModule_ShaderParameters;
-
 	UPROPERTY(EditAnywhere, Category = "Parameters", meta = (DisplayName = "Scale"))
 	FNiagaraDistributionFloat ScaleDistribution = FNiagaraDistributionFloat(1.0f);
 
 	UPROPERTY(EditAnywhere, Category = "Parameters", meta = (EditConditionHides, EditCondition = "UseScaleCurveRange()"))
 	FNiagaraParameterBindingWithValue ScaleCurveRange;
 
-	virtual void PostInitProperties() override
-	{
-		Super::PostInitProperties();
-
-	#if WITH_EDITORONLY_DATA
-		if (HasAnyFlags(RF_ClassDefaultObject) == false)
-		{
-			ScaleCurveRange.SetUsage(ENiagaraParameterBindingUsage::NotParticle);
-			ScaleCurveRange.SetAllowedTypeDefinitions({ FNiagaraTypeDefinition::GetVec2Def() });
-			ScaleCurveRange.SetDefaultParameter(FNiagaraTypeDefinition::GetFloatDef(), 1.0f);
-		}
-	#endif
-	}
-
-	virtual void BuildEmitterData(const FNiagaraStatelessEmitterDataBuildContext& BuildContext) const override
-	{
-		FModuleBuiltData* BuiltData = BuildContext.AllocateBuiltData<FModuleBuiltData>();
-
-		const FNiagaraStatelessGlobals& StatelessGlobals = FNiagaraStatelessGlobals::Get();
-		BuiltData->RibbonWidthVariableOffset			= BuildContext.FindParticleVariableIndex(StatelessGlobals.RibbonWidthVariable);
-		BuiltData->PreviousRibbonWidthVariableOffset	= BuildContext.FindParticleVariableIndex(StatelessGlobals.PreviousRibbonWidthVariable);
-
-		const bool bAttributesUsed = (BuiltData->RibbonWidthVariableOffset != INDEX_NONE || BuiltData->PreviousRibbonWidthVariableOffset != INDEX_NONE);
-		if (IsModuleEnabled() && bAttributesUsed)
-		{
-			BuiltData->DistributionParameters = BuildContext.AddDistribution(ScaleDistribution);
-			if (UseScaleCurveRange())
-			{
-				BuiltData->CurveScaleOffset = BuildContext.AddRendererBinding(ScaleCurveRange.ResolvedParameter);
-				BuiltData->CurveScale = ScaleCurveRange.GetDefaultValue<float>();
-			}
-
-			BuildContext.AddParticleSimulationExecSimulate(&UNiagaraStatelessModule_ScaleRibbonWidth::ParticleSimulate);
-		}
-	}
-
-	virtual void BuildShaderParameters(FNiagaraStatelessShaderParametersBuilder& ShaderParametersBuilder) const override
-	{
-		ShaderParametersBuilder.AddParameterNestedStruct<FParameters>();
-	}
+	virtual void PostInitProperties() override;
 
-	virtual void SetShaderParameters(const FNiagaraStatelessSetShaderParameterContext& SetShaderParameterContext) const override
-	{
-		const FModuleBuiltData* ModuleBuiltData = SetShaderParameterContext.ReadBuiltData<FModuleBuiltData>();
-
-		FParameters* Parameters						= SetShaderParameterContext.GetParameterNestedStruct<FParameters>();
-		Parameters->ScaleRibbonWidth_Distribution	= ModuleBuiltData->DistributionParameters;
-		Parameters->ScaleRibbonWidth_CurveScale		= SetShaderParameterContext.GetRendererParameterValue(ModuleBuiltData->CurveScaleOffset, ModuleBuiltData->CurveScale);
-	}
-
-	static void ParticleSimulate(const NiagaraStateless::FParticleSimulationContext& ParticleSimulationContext)
-	{
-		using namespace NiagaraStateless;
-
-		const FModuleBuiltData* ModuleBuiltData = ParticleSimulationContext.ReadBuiltData<FModuleBuiltData>();
-
-		const float ScaleFactor					= ParticleSimulationContext.GetParameterBufferFloat(ModuleBuiltData->CurveScaleOffset, ModuleBuiltData->CurveScale);
-		const float* NormalizedAgeData			= ParticleSimulationContext.GetParticleNormalizedAge();
-		const float* PreviousNormalizedAgeData	= ParticleSimulationContext.GetParticlePreviousNormalizedAge();
-
-		for (uint32 i = 0; i < ParticleSimulationContext.GetNumInstances(); ++i)
-		{
-			const FStatelessDistributionSampler<float> SpriteScaleSampler(ParticleSimulationContext, ModuleBuiltData->DistributionParameters, i, 0);
-
-			const float NormalizedAge			= NormalizedAgeData[i];
-			const float PreviousNormalizedAge	= PreviousNormalizedAgeData[i];
-
-			const float RibbonWidth			= ParticleSimulationContext.ReadParticleVariable(ModuleBuiltData->RibbonWidthVariableOffset, i, 0.0f);
-			const float PreviousRibbonWidth = ParticleSimulationContext.ReadParticleVariable(ModuleBuiltData->PreviousRibbonWidthVariableOffset, i, 0.0f);
-			const float Scale				= SpriteScaleSampler.GetValue(ParticleSimulationContext, NormalizedAge) * ScaleFactor;
-			const float PreviousScale		= SpriteScaleSampler.GetValue(ParticleSimulationContext, PreviousNormalizedAge) * ScaleFactor;
-
-			ParticleSimulationContext.WriteParticleVariable(ModuleBuiltData->RibbonWidthVariableOffset, i, RibbonWidth * Scale);
-			ParticleSimulationContext.WriteParticleVariable(ModuleBuiltData->PreviousRibbonWidthVariableOffset, i, PreviousRibbonWidth * PreviousScale);
-		}
-	}
+	virtual void BuildEmitterData(const FNiagaraStatelessEmitterDataBuildContext& BuildContext) const override;
+	virtual void BuildShaderParameters(FNiagaraStatelessShaderParametersBuilder& ShaderParametersBuilder) const override;
+	virtual void SetShaderParameters(const FNiagaraStatelessSetShaderParameterContext& SetShaderParameterContext) const override;
 
 	UFUNCTION()
 	bool UseScaleCurveRange() const { return ScaleDistribution.IsCurve(); }
@@ -121,16 +35,7 @@ public:
 	virtual bool CanDisableModule() const override { return true; }
 #endif
 #if WITH_EDITORONLY_DATA
-	virtual const TCHAR* GetShaderTemplatePath() const override
-	{
-		return TEXT("/Plugin/FX/Niagara/Private/Stateless/Modules/NiagaraStatelessModule_ScaleRibbonWidth.ush");
-	}
-
-	virtual void GetOutputVariables(TArray<FNiagaraVariableBase>& OutVariables, EVariableFilter Filter) const override
-	{
-		const FNiagaraStatelessGlobals& StatelessGlobals = FNiagaraStatelessGlobals::Get();
-		OutVariables.AddUnique(StatelessGlobals.RibbonWidthVariable);
-		OutVariables.AddUnique(StatelessGlobals.PreviousRibbonWidthVariable);
-	}
+	virtual const TCHAR* GetShaderTemplatePath() const override;
+	virtual void GetOutputVariables(TArray<FNiagaraVariableBase>& OutVariables, EVariableFilter Filter) const override;
 #endif
 };
diff --git a/Engine/Plugins/FX/Niagara/Source/Niagara/Internal/Stateless/Modules/NiagaraStatelessModule_ScaleSpriteSize.h b/Engine/Plugins/FX/Niagara/Source/Niagara/Internal/Stateless/Modules/NiagaraStatelessModule_ScaleSpriteSize.h
--- a/Engine/Plugins/FX/Niagara/Source/Niagara/Internal/Stateless/Modules/NiagaraStatelessModule_ScaleSpriteSize.h
+++ b/Engine/Plugins/FX/Niagara/Source/Niagara/Internal/Stateless/Modules/NiagaraStatelessModule_ScaleSpriteSize.h
@@ -4,9 +4,6 @@
 
 #include "Stateless/NiagaraStatelessCommon.h"
 #include "Stateless/NiagaraStatelessModule.h"
-#include "Stateless/NiagaraStatelessEmitterDataBuildContext.h"
-#include "Stateless/NiagaraStatelessModuleShaderParameters.h"
-#include "Stateless/NiagaraStatelessParticleSimContext.h"
 #include "NiagaraParameterBinding.h"
 
 #include "NiagaraStatelessModule_ScaleSpriteSize.generated.h"
@@ -18,101 +15,18 @@ class UNiagaraStatelessModule_ScaleSpriteSize : public UNiagaraStatelessModule
 {
 	GENERATED_BODY()
 
-	struct FModuleBuiltData
-	{
-		FUintVector3	DistributionParameters = FUintVector3::ZeroValue;
-		FVector2f		CurveScale = FVector2f::One();
-		int32			CurveScaleOffset = INDEX_NONE;
-
-		int32			SpriteSizeVariableOffset = INDEX_NONE;
-		int32			PreviousSpriteSizeVariableOffset = INDEX_NONE;
-	};
-
 public:
-	using FParameters = NiagaraStateless::FScaleSpriteSizeModule_ShaderParameters;
-
 	UPROPERTY(EditAnywhere, Category = "Parameters", meta = (DisplayName = "Scale"))
 	FNiagaraDistributionVector2 ScaleDistribution = FNiagaraDistributionVector2(FVector2f::One());
 
 	UPROPERTY(EditAnywhere, Category = "Parameters", meta = (EditConditionHides, EditCondition = "UseScaleCurveRange()"))
 	FNiagaraParameterBindingWithValue ScaleCurveRange;
 
-	virtual void PostInitProperties() override
-	{
-		Super::PostInitProperties();
-
-	#if WITH_EDITORONLY_DATA
-		if (HasAnyFlags(RF_ClassDefaultObject) == false)
-		{
-			ScaleCurveRange.SetUsage(ENiagaraParameterBindingUsage::NotParticle);
-			ScaleCurveRange.SetAllowedTypeDefinitions({ FNiagaraTypeDefinition::GetVec2Def() });
-			ScaleCurveRange.SetDefaultParameter(FNiagaraTypeDefinition::GetVec2Def(), FVector2f::One());
-		}
-	#endif
-	}
-
-	virtual void BuildEmitterData(const FNiagaraStatelessEmitterDataBuildContext& BuildContext) const override
-	{
-		FModuleBuiltData* BuiltData = BuildContext.AllocateBuiltData<FModuleBuiltData>();
-
-		const FNiagaraStatelessGlobals& StatelessGlobals = FNiagaraStatelessGlobals::Get();
-		BuiltData->SpriteSizeVariableOffset			= BuildContext.FindParticleVariableIndex(StatelessGlobals.SpriteSizeVariable);
-		BuiltData->PreviousSpriteSizeVariableOffset	= BuildContext.FindParticleVariableIndex(StatelessGlobals.PreviousSpriteSizeVariable);
-
-		const bool bAttributesUsed = (BuiltData->SpriteSizeVariableOffset != INDEX_NONE || BuiltData->PreviousSpriteSizeVariableOffset != INDEX_NONE);
-		if (IsModuleEnabled() && bAttributesUsed)
-		{
-			BuiltData->DistributionParameters = BuildContext.AddDistribution(ScaleDistribution);
-			if (UseScaleCurveRange())
-			{
-				BuiltData->CurveScaleOffset = BuildContext.AddRendererBinding(ScaleCurveRange.ResolvedParameter);
-				BuiltData->CurveScale = ScaleCurveRange.GetDefaultValue<FVector2f>();
-			}
-
-			BuildContext.AddParticleSimulationExecSimulate(&UNiagaraStatelessModule_ScaleSpriteSize::ParticleSimulate);
-		}
-	}
-
-	virtual void BuildShaderParameters(FNiagaraStatelessShaderParametersBuilder& ShaderParametersBuilder) const override
-	{
-		ShaderParametersBuilder.AddParameterNestedStruct<FParameters>();
-	}
+	virtual void PostInitProperties() override;
 
-	virtual void SetShaderParameters(const FNiagaraStatelessSetShaderParameterContext& SetShaderParameterContext) const override
-	{
-		const FModuleBuiltData* ModuleBuiltData = SetShaderParameterContext.ReadBuiltData<FModuleBuiltData>();
-
-		FParameters* Parameters						= SetShaderParameterContext.GetParameterNestedStruct<FParameters>();
-		Parameters->ScaleSpriteSize_Distribution	= ModuleBuiltData->DistributionParameters;
-		Parameters->ScaleSpriteSize_CurveScale		= SetShaderParameterContext.GetRendererParameterValue(ModuleBuiltData->CurveScaleOffset, ModuleBuiltData->CurveScale);
-	}
-
-	static void ParticleSimulate(const NiagaraStateless::FParticleSimulationContext& ParticleSimulationContext)
-	{
-		using namespace NiagaraStateless;
-
-		const FModuleBuiltData* ModuleBuiltData = ParticleSimulationContext.ReadBuiltData<FModuleBuiltData>();
-
-		const FVector2f ScaleFactor				= ParticleSimulationContext.GetParameterBufferFloat(ModuleBuiltData->CurveScaleOffset, ModuleBuiltData->CurveScale);
-		const float* NormalizedAgeData			= ParticleSimulationContext.GetParticleNormalizedAge();
-		const float* PreviousNormalizedAgeData	= ParticleSimulationContext.GetParticlePreviousNormalizedAge();
-
-		for (uint32 i = 0; i < ParticleSimulationContext.GetNumInstances(); ++i)
-		{
-			const FStatelessDistributionSampler<FVector2f> SpriteScaleSampler(ParticleSimulationContext, ModuleBuiltData->DistributionParameters, i, 0);
-
-			const float NormalizedAge			= NormalizedAgeData[i];
-			const float PreviousNormalizedAge	= PreviousNormalizedAgeData[i];
-
-			const FVector2f SpriteSize			= ParticleSimulationContext.ReadParticleVariable(ModuleBuiltData->SpriteSizeVariableOffset, i, FVector2f::ZeroVector);
-			const FVector2f PreviousSpriteSize	= ParticleSimulationContext.ReadParticleVariable(ModuleBuiltData->PreviousSpriteSizeVariableOffset, i, FVector2f::ZeroVector);
-			const FVector2f Scale				= SpriteScaleSampler.GetValue(ParticleSimulationContext, NormalizedAge) * ScaleFactor;
-			const FVector2f PreviousScale		= SpriteScaleSampler.GetValue(ParticleSimulationContext, PreviousNormalizedAge) * ScaleFactor;
-
-			ParticleSimulationContext.WriteParticleVariable(ModuleBuiltData->SpriteSizeVariableOffset, i, SpriteSize * Scale);
-			ParticleSimulationContext.WriteParticleVariable(ModuleBuiltData->PreviousSpriteSizeVariableOffset, i, PreviousSpriteSize * PreviousScale);
-		}
-	}
+	virtual void BuildEmitterData(const FNiagaraStatelessEmitterDataBuildContext& BuildContext) const override;
+	virtual void BuildShaderParameters(FNiagaraStatelessShaderParametersBuilder& ShaderParametersBuilder) const override;
+	virtual void SetShaderParameters(const FNiagaraStatelessSetShaderParameterContext& SetShaderParameterContext) const override;
 
 	UFUNCTION()
 	bool UseScaleCurveRange() const { return ScaleDistribution.IsCurve(); }
@@ -121,16 +35,7 @@ public:
 	virtual bool CanDisableModule() const override { return true; }
 #endif
 #if WITH_EDITORONLY_DATA
-	virtual const TCHAR* GetShaderTemplatePath() const override
-	{
-		return TEXT("/Plugin/FX/Niagara/Private/Stateless/Modules/NiagaraStatelessModule_ScaleSpriteSize.ush");
-	}
-
-	virtual void GetOutputVariables(TArray<FNiagaraVariableBase>& OutVariables, EVariableFilter Filter) const override
-	{
-		const FNiagaraStatelessGlobals& StatelessGlobals = FNiagaraStatelessGlobals::Get();
-		OutVariables.AddUnique(StatelessGlobals.SpriteSizeVariable);
-		OutVariables.AddUnique(StatelessGlobals.PreviousSpriteSizeVariable);
-	}
+	virtual const TCHAR* GetShaderTemplatePath() const override;
+	virtual void GetOutputVariables(TArray<FNiagaraVariableBase>& OutVariables, EVariableFilter Filter) const override;
 #endif
 };
diff --git a/Engine/Plugins/FX/Niagara/Source/Niagara/Internal/Stateless/Modules/NiagaraStatelessModule_SpriteFacingAndAlignment.h b/Engine/Plugins/FX/Niagara/Source/Niagara/Internal/Stateless/Modules/NiagaraStatelessModule_SpriteFacingAndAlignment.h
--- a/Engine/Plugins/FX/Niagara/Source/Niagara/Internal/Stateless/Modules/NiagaraStatelessModule_SpriteFacingAndAlignment.h
+++ b/Engine/Plugins/FX/Niagara/Source/Niagara/Internal/Stateless/Modules/NiagaraStatelessModule_SpriteFacingAndAlignment.h
@@ -3,9 +3,6 @@
 #pragma once
 
 #include "Stateless/NiagaraStatelessModule.h"
-#include "Stateless/NiagaraStatelessEmitterDataBuildContext.h"
-#include "Stateless/NiagaraStatelessModuleShaderParameters.h"
-#include "Stateless/NiagaraStatelessParticleSimContext.h"
 
 #include "NiagaraStatelessModule_SpriteFacingAndAlignment.generated.h"
 
@@ -15,20 +12,7 @@ class UNiagaraStatelessModule_SpriteFacingAndAlignment : public UNiagaraStateles
 {
 	GENERATED_BODY()
 
-	struct FModuleBuiltData
-	{
-		FNiagaraStatelessRangeVector3	SpriteFacing	= FNiagaraStatelessRangeVector3(FVector3f::XAxisVector);
-		FNiagaraStatelessRangeVector3	SpriteAlignment = FNiagaraStatelessRangeVector3(FVector3f::YAxisVector);
-
-		int32		SpriteFacingVariableOffset = INDEX_NONE;
-		int32		PreviousSpriteFacingVariableOffset = INDEX_NONE;
-		int32		SpriteAlignmentVariable = INDEX_NONE;
-		int32		PreviousSpriteAlignmentVariable = INDEX_NONE;
-	};
-
 public:
-	using FParameters = NiagaraStateless::FSpriteFacingAndAlignmentModule_ShaderParameters;
-
 	UPROPERTY(EditAnywhere, Category = "Parameters", meta = (InlineEditConditionToggle))
 	bool bSpriteFacingEnabled = true;
 
@@ -41,97 +25,15 @@ public:
 	UPROPERTY(EditAnywhere, Category = "Parameters", meta = (DisableRangeDistribution, DisableUniformDistribution, EditCondition = "bSpriteAlignmentEnabled"))
 	FNiagaraDistributionRangeVector3	SpriteAlignment = FNiagaraDistributionRangeVector3(FVector3f::YAxisVector);
 
-	virtual void BuildEmitterData(const FNiagaraStatelessEmitterDataBuildContext& BuildContext) const override
-	{
-		FModuleBuiltData* BuiltData = BuildContext.AllocateBuiltData<FModuleBuiltData>();
-		if (!IsModuleEnabled())
-		{
-			return;
-		}
-
-		const FNiagaraStatelessGlobals& StatelessGlobals = FNiagaraStatelessGlobals::Get();
-		if (bSpriteFacingEnabled)
-		{
-			BuiltData->SpriteFacingVariableOffset			= BuildContext.FindParticleVariableIndex(StatelessGlobals.SpriteFacingVariable);
-			BuiltData->PreviousSpriteFacingVariableOffset	= BuildContext.FindParticleVariableIndex(StatelessGlobals.PreviousSpriteFacingVariable);
-		}
-		if (bSpriteAlignmentEnabled)
-		{
-			BuiltData->SpriteAlignmentVariable				= BuildContext.FindParticleVariableIndex(StatelessGlobals.SpriteAlignmentVariable);
-			BuiltData->PreviousSpriteAlignmentVariable		= BuildContext.FindParticleVariableIndex(StatelessGlobals.PreviousSpriteAlignmentVariable);
-		}
-
-		if ((BuiltData->SpriteFacingVariableOffset == INDEX_NONE) && (BuiltData->PreviousSpriteFacingVariableOffset == INDEX_NONE) &&
-			(BuiltData->SpriteAlignmentVariable == INDEX_NONE) && (BuiltData->PreviousSpriteAlignmentVariable == INDEX_NONE))
-		{
-			return;
-		}
-
-		if (bSpriteFacingEnabled)
-		{
-			BuiltData->SpriteFacing = BuildContext.ConvertDistributionToRange(SpriteFacing, FVector3f::ZeroVector);
-		}
-		if (bSpriteAlignmentEnabled)
-		{
-			BuiltData->SpriteAlignment	= BuildContext.ConvertDistributionToRange(SpriteAlignment, FVector3f::ZeroVector);
-		}
-
-		BuildContext.AddParticleSimulationExecSimulate(&UNiagaraStatelessModule_SpriteFacingAndAlignment::ParticleSimulate);
-	}
-
-	virtual void BuildShaderParameters(FNiagaraStatelessShaderParametersBuilder& ShaderParametersBuilder) const override
-	{
-		ShaderParametersBuilder.AddParameterNestedStruct<FParameters>();
-	}
-
-	virtual void SetShaderParameters(const FNiagaraStatelessSetShaderParameterContext& SetShaderParameterContext) const override
-	{
-		FParameters* Parameters = SetShaderParameterContext.GetParameterNestedStruct<FParameters>();
-		const FModuleBuiltData* ModuleBuiltData = SetShaderParameterContext.ReadBuiltData<FModuleBuiltData>();
-
-		Parameters->SpriteFacingAndAlignment_SpriteFacing		= SetShaderParameterContext.ConvertRangeToValue(ModuleBuiltData->SpriteFacing);
-		Parameters->SpriteFacingAndAlignment_SpriteAlignment	= SetShaderParameterContext.ConvertRangeToValue(ModuleBuiltData->SpriteAlignment);
-	}
-
-	static void ParticleSimulate(const NiagaraStateless::FParticleSimulationContext& ParticleSimulationContext)
-	{
-		using namespace NiagaraStateless;
-
-		const FModuleBuiltData* ModuleBuiltData = ParticleSimulationContext.ReadBuiltData<FModuleBuiltData>();
-		const FParameters* Parameters = ParticleSimulationContext.ReadParameterNestedStruct<FParameters>();
-
-		for (uint32 i = 0; i < ParticleSimulationContext.GetNumInstances(); ++i)
-		{
-			ParticleSimulationContext.WriteParticleVariable(ModuleBuiltData->SpriteFacingVariableOffset, i, Parameters->SpriteFacingAndAlignment_SpriteFacing);
-			ParticleSimulationContext.WriteParticleVariable(ModuleBuiltData->PreviousSpriteFacingVariableOffset, i, Parameters->SpriteFacingAndAlignment_SpriteFacing);
-
-			ParticleSimulationContext.WriteParticleVariable(ModuleBuiltData->SpriteAlignmentVariable, i, Parameters->SpriteFacingAndAlignment_SpriteAlignment);
-			ParticleSimulationContext.WriteParticleVariable(ModuleBuiltData->PreviousSpriteAlignmentVariable, i, Parameters->SpriteFacingAndAlignment_SpriteAlignment);
-		}
-	}
+	virtual void BuildEmitterData(const FNiagaraStatelessEmitterDataBuildContext& BuildContext) const override;
+	virtual void BuildShaderParameters(FNiagaraStatelessShaderParametersBuilder& ShaderParametersBuilder) const override;
+	virtual void SetShaderParameters(const FNiagaraStatelessSetShaderParameterContext& SetShaderParameterContext) const override;
 
 #if WITH_EDITOR
 	virtual bool CanDisableModule() const override { return true; }
 #endif
 #if WITH_EDITORONLY_DATA
-	virtual const TCHAR* GetShaderTemplatePath() const override
-	{
-		return TEXT("/Plugin/FX/Niagara/Private/Stateless/Modules/NiagaraStatelessModule_SpriteFacingAndAlignment.ush");
-	}
-
-	virtual void GetOutputVariables(TArray<FNiagaraVariableBase>& OutVariables, EVariableFilter Filter) const override
-	{
-		const FNiagaraStatelessGlobals& StatelessGlobals = FNiagaraStatelessGlobals::Get();
-		if (bSpriteFacingEnabled || Filter == EVariableFilter::None)
-		{
-			OutVariables.AddUnique(StatelessGlobals.SpriteFacingVariable);
-			OutVariables.AddUnique(StatelessGlobals.PreviousSpriteFacingVariable);
-		}
-		if (bSpriteAlignmentEnabled || Filter == EVariableFilter::None)
-		{
-			OutVariables.AddUnique(StatelessGlobals.SpriteAlignmentVariable);
-			OutVariables.AddUnique(StatelessGlobals.PreviousSpriteAlignmentVariable);
-		}
-	}
+	virtual const TCHAR* GetShaderTemplatePath() const override;
+	virtual void GetOutputVariables(TArray<FNiagaraVariableBase>& OutVariables, EVariableFilter Filter) const override;
 #endif
 };
diff --git a/Engine/Plugins/FX/Niagara/Source/Niagara/Internal/Stateless/Modules/NiagaraStatelessModule_SpriteRotationRate.h b/Engine/Plugins/FX/Niagara/Source/Niagara/Internal/Stateless/Modules/NiagaraStatelessModule_SpriteRotationRate.h
--- a/Engine/Plugins/FX/Niagara/Source/Niagara/Internal/Stateless/Modules/NiagaraStatelessModule_SpriteRotationRate.h
+++ b/Engine/Plugins/FX/Niagara/Source/Niagara/Internal/Stateless/Modules/NiagaraStatelessModule_SpriteRotationRate.h
@@ -3,9 +3,6 @@
 #pragma once
 
 #include "Stateless/NiagaraStatelessModule.h"
-#include "Stateless/NiagaraStatelessEmitterDataBuildContext.h"
-#include "Stateless/NiagaraStatelessModuleShaderParameters.h"
-#include "Stateless/NiagaraStatelessParticleSimContext.h"
 
 #include "NiagaraStatelessModule_SpriteRotationRate.generated.h"
 
@@ -15,18 +12,7 @@ class UNiagaraStatelessModule_SpriteRotationRate : public UNiagaraStatelessModul
 {
 	GENERATED_BODY()
 
-	struct FModuleBuiltData
-	{
-		int32							ModuleEnabled = false;
-		FNiagaraStatelessRangeFloat		RotationRange;
-		FUintVector3					RateScaleParameters = FUintVector3::ZeroValue;
-		int32							SpriteRotationVariableOffset = INDEX_NONE;
-		int32							PreviousSpriteRotationVariableOffset = INDEX_NONE;
-	};
-
 public:
-	using FParameters = NiagaraStateless::FSpriteRotationRateModule_ShaderParameters;
-
 	UPROPERTY(EditAnywhere, Category = "Parameters", meta = (InlineEditConditionToggle))
 	bool bUseRateScale = false;
 
@@ -36,94 +22,15 @@ public:
 	UPROPERTY(EditAnywhere, Category = "Parameters", meta = (DisplayName = "Rate Scale", EditCondition = "bUseRateScale"))
 	FNiagaraDistributionCurveFloat RateScaleDistribution = FNiagaraDistributionCurveFloat(ENiagaraDistributionCurveLUTMode::Accumulate);
 
-	virtual void BuildEmitterData(const FNiagaraStatelessEmitterDataBuildContext& BuildContext) const override
-	{
-		FModuleBuiltData* BuiltData = BuildContext.AllocateBuiltData<FModuleBuiltData>();
-		if (!IsModuleEnabled())
-		{
-			return;
-		}
-
-		const FNiagaraStatelessGlobals& StatelessGlobals = FNiagaraStatelessGlobals::Get();
-		BuiltData->SpriteRotationVariableOffset			= BuildContext.FindParticleVariableIndex(StatelessGlobals.SpriteRotationVariable);
-		BuiltData->PreviousSpriteRotationVariableOffset	= BuildContext.FindParticleVariableIndex(StatelessGlobals.PreviousSpriteRotationVariable);
-
-		if (BuiltData->SpriteRotationVariableOffset == INDEX_NONE && BuiltData->PreviousSpriteRotationVariableOffset == INDEX_NONE)
-		{
-			return;
-		}
-
-		BuiltData->ModuleEnabled = true;
-		BuiltData->RotationRange = BuildContext.ConvertDistributionToRange(RotationRateDistribution, 0.0f);
-		if (bUseRateScale)
-		{
-			BuiltData->RateScaleParameters = BuildContext.AddDistributionAsCurve(RateScaleDistribution, 1.0f);
-		}
-		else
-		{
-			BuiltData->RateScaleParameters = BuildContext.AddDistributionAsCurve(FNiagaraDistributionCurveFloat(ENiagaraDistributionCurveLUTMode::Accumulate), 1.0f);
-		}
-
-		BuildContext.AddParticleSimulationExecSimulate(&UNiagaraStatelessModule_SpriteRotationRate::ParticleSimulate);
-	}
-
-	virtual void BuildShaderParameters(FNiagaraStatelessShaderParametersBuilder& ShaderParametersBuilder) const override
-	{
-		ShaderParametersBuilder.AddParameterNestedStruct<FParameters>();
-	}
-
-	virtual void SetShaderParameters(const FNiagaraStatelessSetShaderParameterContext& SetShaderParameterContext) const override
-	{
-		FParameters* Parameters = SetShaderParameterContext.GetParameterNestedStruct<FParameters>();
-		const FModuleBuiltData* ModuleBuiltData = SetShaderParameterContext.ReadBuiltData<FModuleBuiltData>();
-
-		Parameters->SpriteRotationRate_Enabled = ModuleBuiltData->ModuleEnabled;
-		SetShaderParameterContext.ConvertRangeToScaleBias(ModuleBuiltData->RotationRange, Parameters->SpriteRotationRate_Scale, Parameters->SpriteRotationRate_Bias);
-		Parameters->SpriteRotationRate_RateScaleParameters = ModuleBuiltData->RateScaleParameters;
-	}
-
-	static void ParticleSimulate(const NiagaraStateless::FParticleSimulationContext& ParticleSimulationContext)
-	{
-		using namespace NiagaraStateless;
-
-		const FModuleBuiltData* ModuleBuiltData = ParticleSimulationContext.ReadBuiltData<FModuleBuiltData>();
-		const FParameters* ShaderParameters = ParticleSimulationContext.ReadParameterNestedStruct<FParameters>();
-
-		const float* LifetimeData = ParticleSimulationContext.GetParticleLifetime();
-		const float* AgeData = ParticleSimulationContext.GetParticleNormalizedAge();
-		const float* PreviousAgeData = ParticleSimulationContext.GetParticlePreviousNormalizedAge();
-
-		for (uint32 i = 0; i < ParticleSimulationContext.GetNumInstances(); ++i)
-		{
-			const float RotationRate = ParticleSimulationContext.RandomScaleBiasFloat(i, 0, ShaderParameters->SpriteRotationRate_Scale, ShaderParameters->SpriteRotationRate_Bias);
-			const float MultiplyRate = ParticleSimulationContext.SampleCurve<float>(ShaderParameters->SpriteRotationRate_RateScaleParameters, AgeData[i]);
-			const float PreviousMultiplyRate = ParticleSimulationContext.SampleCurve<float>(ShaderParameters->SpriteRotationRate_RateScaleParameters, PreviousAgeData[i]);
-
-			float SpriteRotation = ParticleSimulationContext.ReadParticleVariable(ModuleBuiltData->SpriteRotationVariableOffset, i, 0.0f);
-			float PreviousSpriteRotation = ParticleSimulationContext.ReadParticleVariable(ModuleBuiltData->PreviousSpriteRotationVariableOffset, i, 0.0f);
-
-			SpriteRotation += LifetimeData[i] * RotationRate * MultiplyRate;
-			PreviousSpriteRotation += LifetimeData[i] * RotationRate * PreviousMultiplyRate;
-
-			ParticleSimulationContext.WriteParticleVariable(ModuleBuiltData->SpriteRotationVariableOffset, i, SpriteRotation);
-			ParticleSimulationContext.WriteParticleVariable(ModuleBuiltData->PreviousSpriteRotationVariableOffset, i, PreviousSpriteRotation);
-		}
-	}
+	virtual void BuildEmitterData(const FNiagaraStatelessEmitterDataBuildContext& BuildContext) const override;
+	virtual void BuildShaderParameters(FNiagaraStatelessShaderParametersBuilder& ShaderParametersBuilder) const override;
+	virtual void SetShaderParameters(const FNiagaraStatelessSetShaderParameterContext& SetShaderParameterContext) const override;
 
 #if WITH_EDITOR
 	virtual bool CanDisableModule() const override { return true; }
 #endif
 #if WITH_EDITORONLY_DATA
-	virtual const TCHAR* GetShaderTemplatePath() const override
-	{
-		return TEXT("/Plugin/FX/Niagara/Private/Stateless/Modules/NiagaraStatelessModule_SpriteRotationRate.ush");
-	}
-
-	virtual void GetOutputVariables(TArray<FNiagaraVariableBase>& OutVariables, EVariableFilter Filter) const override
-	{
-		const FNiagaraStatelessGlobals& StatelessGlobals = FNiagaraStatelessGlobals::Get();
-		OutVariables.AddUnique(StatelessGlobals.SpriteRotationVariable);
-		OutVariables.AddUnique(StatelessGlobals.PreviousSpriteRotationVariable);
-	}
+	virtual const TCHAR* GetShaderTemplatePath() const override;
+	virtual void GetOutputVariables(TArray<FNiagaraVariableBase>& OutVariables, EVariableFilter Filter) const override;
 #endif
 };
diff --git a/Engine/Plugins/FX/Niagara/Source/Niagara/Private/Stateless/Modules/NiagaraStatelessModule_AccelerationForce.cpp b/Engine/Plugins/FX/Niagara/Source/Niagara/Private/Stateless/Modules/NiagaraStatelessModule_AccelerationForce.cpp
new file mode 100644
--- /dev/null
+++ b/Engine/Plugins/FX/Niagara/Source/Niagara/Private/Stateless/Modules/NiagaraStatelessModule_AccelerationForce.cpp
@@ -0,0 +1,19 @@
+// Copyright Epic Games, Inc. All Rights Reserved.
+
+#include "Stateless/Modules/NiagaraStatelessModule_AccelerationForce.h"
+#include "Stateless/NiagaraStatelessEmitterDataBuildContext.h"
+
+void UNiagaraStatelessModule_AccelerationForce::BuildEmitterData(const FNiagaraStatelessEmitterDataBuildContext& BuildContext) const
+{
+	if (!IsModuleEnabled())
+	{
+		return;
+	}
+
+	const FNiagaraStatelessRangeVector3 AccelerationRange = AccelerationDistribution.CalculateRange(FVector3f::ZeroVector);
+
+	NiagaraStateless::FPhysicsBuildData& PhysicsBuildData = BuildContext.GetTransientBuildData<NiagaraStateless::FPhysicsBuildData>();
+	PhysicsBuildData.AccelerationCoordinateSpace = CoordinateSpace;
+	PhysicsBuildData.AccelerationRange.Min += AccelerationRange.Min;
+	PhysicsBuildData.AccelerationRange.Max += AccelerationRange.Max;
+}
diff --git a/Engine/Plugins/FX/Niagara/Source/Niagara/Private/Stateless/Modules/NiagaraStatelessModule_CameraOffset.cpp b/Engine/Plugins/FX/Niagara/Source/Niagara/Private/Stateless/Modules/NiagaraStatelessModule_CameraOffset.cpp
new file mode 100644
--- /dev/null
+++ b/Engine/Plugins/FX/Niagara/Source/Niagara/Private/Stateless/Modules/NiagaraStatelessModule_CameraOffset.cpp
@@ -0,0 +1,92 @@
+// Copyright Epic Games, Inc. All Rights Reserved.
+
+#include "Stateless/Modules/NiagaraStatelessModule_CameraOffset.h"
+
+#include "Stateless/NiagaraStatelessEmitterDataBuildContext.h"
+#include "Stateless/NiagaraStatelessModuleShaderParameters.h"
+#include "Stateless/NiagaraStatelessParticleSimContext.h"
+
+namespace NSMCameraOffsetPrivate
+{
+	struct FModuleBuiltData
+	{
+		FUintVector3	DistributionParameters = FUintVector3::ZeroValue;
+		int32			CameraVariableOffset = INDEX_NONE;
+		int32			PreviousCameraVariableOffset = INDEX_NONE;
+	};
+
+	using FParameters = NiagaraStateless::FCameraOffsetModule_ShaderParameters;
+
+	static void ParticleSimulate(const NiagaraStateless::FParticleSimulationContext& ParticleSimulationContext)
+	{
+		using namespace NiagaraStateless;
+
+		const FModuleBuiltData* ModuleBuiltData = ParticleSimulationContext.ReadBuiltData<FModuleBuiltData>();
+		const float* NormalizedAgeData = ParticleSimulationContext.GetParticleNormalizedAge();
+		const float* PreviousNormalizedAgeData = ParticleSimulationContext.GetParticlePreviousNormalizedAge();
+
+		for (uint32 i = 0; i < ParticleSimulationContext.GetNumInstances(); ++i)
+		{
+			const FStatelessDistributionSampler<float> CameraOffsetSampler(ParticleSimulationContext, ModuleBuiltData->DistributionParameters, i, 0);
+
+			const float NormalizedAge = NormalizedAgeData[i];
+			const float PreviousNormalizedAge = PreviousNormalizedAgeData[i];
+
+			const float CameraOffset = CameraOffsetSampler.GetValue(ParticleSimulationContext, NormalizedAge);
+			const float PreviousCameraOffset = CameraOffsetSampler.GetValue(ParticleSimulationContext, PreviousNormalizedAge);
+
+			ParticleSimulationContext.WriteParticleVariable(ModuleBuiltData->CameraVariableOffset, i, CameraOffset);
+			ParticleSimulationContext.WriteParticleVariable(ModuleBuiltData->PreviousCameraVariableOffset, i, PreviousCameraOffset);
+		}
+	}
+}
+
+void UNiagaraStatelessModule_CameraOffset::BuildEmitterData(const FNiagaraStatelessEmitterDataBuildContext& BuildContext) const
+{
+	using namespace NSMCameraOffsetPrivate;
+
+	const FNiagaraStatelessGlobals& StatelessGlobals = FNiagaraStatelessGlobals::Get();
+
+	FModuleBuiltData* BuiltData				= BuildContext.AllocateBuiltData<FModuleBuiltData>();
+	BuiltData->CameraVariableOffset			= BuildContext.FindParticleVariableIndex(StatelessGlobals.CameraOffsetVariable);
+	BuiltData->PreviousCameraVariableOffset	= BuildContext.FindParticleVariableIndex(StatelessGlobals.PreviousCameraOffsetVariable);
+
+	const bool bAttributesUsed = (BuiltData->CameraVariableOffset != INDEX_NONE || BuiltData->PreviousCameraVariableOffset != INDEX_NONE);
+	if (IsModuleEnabled() && bAttributesUsed)
+	{
+		BuiltData->DistributionParameters = BuildContext.AddDistribution(CameraOffsetDistribution);
+
+		BuildContext.AddParticleSimulationExecSimulate(&ParticleSimulate);
+	}
+}
+
+void UNiagaraStatelessModule_CameraOffset::BuildShaderParameters(FNiagaraStatelessShaderParametersBuilder& ShaderParametersBuilder) const
+{
+	using namespace NSMCameraOffsetPrivate;
+
+	ShaderParametersBuilder.AddParameterNestedStruct<FParameters>();
+}
+
+void UNiagaraStatelessModule_CameraOffset::SetShaderParameters(const FNiagaraStatelessSetShaderParameterContext& SetShaderParameterContext) const
+{
+	using namespace NSMCameraOffsetPrivate;
+
+	const FModuleBuiltData* ModuleBuiltData = SetShaderParameterContext.ReadBuiltData<FModuleBuiltData>();
+
+	FParameters* Parameters = SetShaderParameterContext.GetParameterNestedStruct<FParameters>();
+	Parameters->CameraOffset_Distribution = ModuleBuiltData->DistributionParameters;
+}
+
+#if WITH_EDITORONLY_DATA
+const TCHAR* UNiagaraStatelessModule_CameraOffset::GetShaderTemplatePath() const
+{
+	return TEXT("/Plugin/FX/Niagara/Private/Stateless/Modules/NiagaraStatelessModule_CameraOffset.ush");
+}
+
+void UNiagaraStatelessModule_CameraOffset::GetOutputVariables(TArray<FNiagaraVariableBase>& OutVariables, EVariableFilter Filter) const
+{
+	const FNiagaraStatelessGlobals& StatelessGlobals = FNiagaraStatelessGlobals::Get();
+	OutVariables.AddUnique(StatelessGlobals.CameraOffsetVariable);
+	OutVariables.AddUnique(StatelessGlobals.PreviousCameraOffsetVariable);
+}
+#endif
diff --git a/Engine/Plugins/FX/Niagara/Source/Niagara/Private/Stateless/Modules/NiagaraStatelessModule_Drag.cpp b/Engine/Plugins/FX/Niagara/Source/Niagara/Private/Stateless/Modules/NiagaraStatelessModule_Drag.cpp
new file mode 100644
--- /dev/null
+++ b/Engine/Plugins/FX/Niagara/Source/Niagara/Private/Stateless/Modules/NiagaraStatelessModule_Drag.cpp
@@ -0,0 +1,15 @@
+// Copyright Epic Games, Inc. All Rights Reserved.
+
+#include "Stateless/Modules/NiagaraStatelessModule_Drag.h"
+#include "Stateless/NiagaraStatelessEmitterDataBuildContext.h"
+
+void UNiagaraStatelessModule_Drag::BuildEmitterData(const FNiagaraStatelessEmitterDataBuildContext& BuildContext) const
+{
+	if (!IsModuleEnabled())
+	{
+		return;
+	}
+
+	NiagaraStateless::FPhysicsBuildData& PhysicsBuildData = BuildContext.GetTransientBuildData<NiagaraStateless::FPhysicsBuildData>();
+	PhysicsBuildData.DragRange = DragDistribution.CalculateRange(DefaultDrag);
+}
diff --git a/Engine/Plugins/FX/Niagara/Source/Niagara/Private/Stateless/Modules/NiagaraStatelessModule_DynamicMaterialParameters.cpp b/Engine/Plugins/FX/Niagara/Source/Niagara/Private/Stateless/Modules/NiagaraStatelessModule_DynamicMaterialParameters.cpp
new file mode 100644
--- /dev/null
+++ b/Engine/Plugins/FX/Niagara/Source/Niagara/Private/Stateless/Modules/NiagaraStatelessModule_DynamicMaterialParameters.cpp
@@ -0,0 +1,234 @@
+// Copyright Epic Games, Inc. All Rights Reserved.
+
+#include "Stateless/Modules/NiagaraStatelessModule_DynamicMaterialParameters.h"
+
+#include "Stateless/NiagaraStatelessEmitterDataBuildContext.h"
+#include "Stateless/NiagaraStatelessModuleShaderParameters.h"
+#include "Stateless/NiagaraStatelessParticleSimContext.h"
+
+namespace NSMDynamicMaterialParametersPrivate
+{
+	static constexpr int32 NumParameters = 4;
+	static constexpr int32 NumChannelPerParameter = 4;
+
+	using FParameters = NiagaraStateless::FDynamicMaterialParametersModule_ShaderParameters;
+
+	struct FModuleBuiltData
+	{
+		FModuleBuiltData()
+		{
+			for (FUintVector3& v : ParameterDistributions)
+			{
+				v = FUintVector3::ZeroValue;
+			}
+
+			for (int32& ParameterVariableOffset : ParameterVariableOffsets)
+			{
+				ParameterVariableOffset = INDEX_NONE;
+			}
+		}
+
+		uint32			ChannelMask = 0;
+		FUintVector3	ParameterDistributions[NumParameters * NumChannelPerParameter];
+		int32			ParameterVariableOffsets[NumParameters];
+	};
+
+	static void ParticleSimulate(const NiagaraStateless::FParticleSimulationContext& ParticleSimulationContext)
+	{
+		using namespace NiagaraStateless;
+
+		const FModuleBuiltData* ModuleBuiltData = ParticleSimulationContext.ReadBuiltData<FModuleBuiltData>();
+		const float* NormalizedAgeData = ParticleSimulationContext.GetParticleNormalizedAge();
+
+		for (uint32 i = 0; i < ParticleSimulationContext.GetNumInstances(); ++i)
+		{
+			for (int32 iParameter = 0; iParameter < NumParameters; ++iParameter)
+			{
+				const uint32 ChannelMask = (ModuleBuiltData->ChannelMask >> (iParameter * NumChannelPerParameter)) & 0xf;
+				if (ChannelMask == 0)
+				{
+					continue;
+				}
+
+				const int32 FirstChannel = iParameter * NumChannelPerParameter;
+				const FStatelessDistributionSampler<float> XChannelSampler(ParticleSimulationContext, ModuleBuiltData->ParameterDistributions[FirstChannel + 0], i, FirstChannel + 0);
+				const FStatelessDistributionSampler<float> YChannelSampler(ParticleSimulationContext, ModuleBuiltData->ParameterDistributions[FirstChannel + 1], i, FirstChannel + 1);
+				const FStatelessDistributionSampler<float> ZChannelSampler(ParticleSimulationContext, ModuleBuiltData->ParameterDistributions[FirstChannel + 2], i, FirstChannel + 2);
+				const FStatelessDistributionSampler<float> WChannelSampler(ParticleSimulationContext, ModuleBuiltData->ParameterDistributions[FirstChannel + 3], i, FirstChannel + 3);
+
+				const float NormalizedAge = NormalizedAgeData[i];
+				const FVector4f DynamicParameter(
+					((ChannelMask & 0x1) != 0) ? XChannelSampler.GetValue(ParticleSimulationContext, NormalizedAge) : 0.0f,
+					((ChannelMask & 0x2) != 0) ? YChannelSampler.GetValue(ParticleSimulationContext, NormalizedAge) : 0.0f,
+					((ChannelMask & 0x4) != 0) ? ZChannelSampler.GetValue(ParticleSimulationContext, NormalizedAge) : 0.0f,
+					((ChannelMask & 0x8) != 0) ? WChannelSampler.GetValue(ParticleSimulationContext, NormalizedAge) : 0.0f
+				);
+
+				ParticleSimulationContext.WriteParticleVariable(ModuleBuiltData->ParameterVariableOffsets[iParameter], i, DynamicParameter);
+			}
+		}
+	}
+
+}
+
+const FNiagaraStatelessDynamicParameterSet& UNiagaraStatelessModule_DynamicMaterialParameters::GetParameterSet(int32 ParameterIndex) const
+{
+	switch (ParameterIndex)
+	{
+		case 0: return Parameter0;
+		case 1: return Parameter1;
+		case 2: return Parameter2;
+		case 3: return Parameter3;
+		default: checkNoEntry(); return Parameter0;
+	}
+}
+
+uint32 UNiagaraStatelessModule_DynamicMaterialParameters::GetParameterChannelMask(int32 ParameterIndex) const
+{
+	uint32 Mask = 0;
+	if ((ParameterIndex == 0 && bParameter0Enabled) ||
+		(ParameterIndex == 1 && bParameter1Enabled) ||
+		(ParameterIndex == 2 && bParameter2Enabled) ||
+		(ParameterIndex == 3 && bParameter3Enabled))
+	{
+		const FNiagaraStatelessDynamicParameterSet& ParameterSet = GetParameterSet(ParameterIndex);
+		Mask |= ParameterSet.bXChannelEnabled ? 1 << 0 : 0;
+		Mask |= ParameterSet.bYChannelEnabled ? 1 << 1 : 0;
+		Mask |= ParameterSet.bZChannelEnabled ? 1 << 2 : 0;
+		Mask |= ParameterSet.bWChannelEnabled ? 1 << 3 : 0;
+	}
+	return Mask;
+}
+
+const FNiagaraDistributionFloat& UNiagaraStatelessModule_DynamicMaterialParameters::GetParameterChannelDistribution(int32 ParameterIndex, int32 ChannelIndex) const
+{
+	const FNiagaraStatelessDynamicParameterSet& ParameterSet = GetParameterSet(ParameterIndex);
+	switch (ChannelIndex)
+	{
+		case 0: return ParameterSet.XChannelDistribution;
+		case 1: return ParameterSet.YChannelDistribution;
+		case 2: return ParameterSet.ZChannelDistribution;
+		case 3: return ParameterSet.WChannelDistribution;
+		default: checkNoEntry(); return ParameterSet.XChannelDistribution;
+	}
+}
+
+const FNiagaraVariableBase& UNiagaraStatelessModule_DynamicMaterialParameters::GetParameterVariable(int32 ParameterIndex) const
+{
+	const FNiagaraStatelessGlobals& StatelessGlobals = FNiagaraStatelessGlobals::Get();
+	switch (ParameterIndex)
+	{
+		default: checkNoEntry(); return StatelessGlobals.DynamicMaterialParameters0Variable;
+		case 0: return StatelessGlobals.DynamicMaterialParameters0Variable;
+		case 1: return StatelessGlobals.DynamicMaterialParameters1Variable;
+		case 2: return StatelessGlobals.DynamicMaterialParameters2Variable;
+		case 3: return StatelessGlobals.DynamicMaterialParameters3Variable;
+	}
+}
+
+
+int32 UNiagaraStatelessModule_DynamicMaterialParameters::GetRendererChannelMask() const
+{
+	using namespace NSMDynamicMaterialParametersPrivate;
+
+	uint32 ChannelMask = 0;
+	if (IsModuleEnabled())
+	{
+		for (int32 iParameter = 0; iParameter < NumParameters; ++iParameter)
+		{
+			ChannelMask |= GetParameterChannelMask(iParameter) << (iParameter * NumChannelPerParameter);
+		}
+	}
+	return ChannelMask;
+}
+
+void UNiagaraStatelessModule_DynamicMaterialParameters::BuildEmitterData(const FNiagaraStatelessEmitterDataBuildContext& BuildContext) const
+{
+	using namespace NSMDynamicMaterialParametersPrivate;
+
+	FModuleBuiltData* BuiltData = BuildContext.AllocateBuiltData<FModuleBuiltData>();
+
+	if (IsModuleEnabled())
+	{
+		for (int32 iParameter=0; iParameter < NumParameters; ++iParameter)
+		{
+			const uint32 ParameterChannelMask = GetParameterChannelMask(iParameter);
+			if (ParameterChannelMask == 0)
+			{
+				continue;
+			}
+
+			BuiltData->ParameterVariableOffsets[iParameter] = BuildContext.FindParticleVariableIndex(GetParameterVariable(iParameter));
+			if (BuiltData->ParameterVariableOffsets[iParameter] == INDEX_NONE)
+			{
+				continue;
+			}
+
+			BuiltData->ChannelMask |= ParameterChannelMask << (iParameter * NumChannelPerParameter);
+			for (int32 iChannel = 0; iChannel < NumChannelPerParameter; ++iChannel)
+			{
+				BuiltData->ParameterDistributions[(iParameter * NumChannelPerParameter) + iChannel] = BuildContext.AddDistribution(GetParameterChannelDistribution(iParameter, iChannel));
+			}
+		}
+
+		if (BuiltData->ChannelMask != 0)
+		{
+			BuildContext.AddParticleSimulationExecSimulate(&ParticleSimulate);
+		}
+	}
+}
+
+void UNiagaraStatelessModule_DynamicMaterialParameters::BuildShaderParameters(FNiagaraStatelessShaderParametersBuilder& ShaderParametersBuilder) const
+{
+	using namespace NSMDynamicMaterialParametersPrivate;
+
+	ShaderParametersBuilder.AddParameterNestedStruct<FParameters>();
+}
+
+void UNiagaraStatelessModule_DynamicMaterialParameters::SetShaderParameters(const FNiagaraStatelessSetShaderParameterContext& SetShaderParameterContext) const
+{
+	using namespace NSMDynamicMaterialParametersPrivate;
+
+	const FModuleBuiltData* ModuleBuiltData = SetShaderParameterContext.ReadBuiltData<FModuleBuiltData>();
+
+	FParameters* Parameters = SetShaderParameterContext.GetParameterNestedStruct<FParameters>();
+	Parameters->DynamicMaterialParameters_ChannelMask = ModuleBuiltData->ChannelMask;
+	Parameters->DynamicMaterialParameters_Parameter0X = ModuleBuiltData->ParameterDistributions[0];
+	Parameters->DynamicMaterialParameters_Parameter0Y = ModuleBuiltData->ParameterDistributions[1];
+	Parameters->DynamicMaterialParameters_Parameter0Z = ModuleBuiltData->ParameterDistributions[2];
+	Parameters->DynamicMaterialParameters_Parameter0W = ModuleBuiltData->ParameterDistributions[3];
+	Parameters->DynamicMaterialParameters_Parameter1X = ModuleBuiltData->ParameterDistributions[4];
+	Parameters->DynamicMaterialParameters_Parameter1Y = ModuleBuiltData->ParameterDistributions[5];
+	Parameters->DynamicMaterialParameters_Parameter1Z = ModuleBuiltData->ParameterDistributions[6];
+	Parameters->DynamicMaterialParameters_Parameter1W = ModuleBuiltData->ParameterDistributions[7];
+	Parameters->DynamicMaterialParameters_Parameter2X = ModuleBuiltData->ParameterDistributions[8];
+	Parameters->DynamicMaterialParameters_Parameter2Y = ModuleBuiltData->ParameterDistributions[9];
+	Parameters->DynamicMaterialParameters_Parameter2Z = ModuleBuiltData->ParameterDistributions[10];
+	Parameters->DynamicMaterialParameters_Parameter2W = ModuleBuiltData->ParameterDistributions[11];
+	Parameters->DynamicMaterialParameters_Parameter3X = ModuleBuiltData->ParameterDistributions[12];
+	Parameters->DynamicMaterialParameters_Parameter3Y = ModuleBuiltData->ParameterDistributions[13];
+	Parameters->DynamicMaterialParameters_Parameter3Z = ModuleBuiltData->ParameterDistributions[14];
+	Parameters->DynamicMaterialParameters_Parameter3W = ModuleBuiltData->ParameterDistributions[15];
+}
+
+#if WITH_EDITORONLY_DATA
+const TCHAR* UNiagaraStatelessModule_DynamicMaterialParameters::GetShaderTemplatePath() const
+{
+	return TEXT("/Plugin/FX/Niagara/Private/Stateless/Modules/NiagaraStatelessModule_DynamicMaterialParameters.ush");
+}
+
+void UNiagaraStatelessModule_DynamicMaterialParameters::GetOutputVariables(TArray<FNiagaraVariableBase>& OutVariables, EVariableFilter Filter) const
+{
+	using namespace NSMDynamicMaterialParametersPrivate;
+
+	const FNiagaraStatelessGlobals& StatelessGlobals = FNiagaraStatelessGlobals::Get();
+
+	for (int32 i = 0; i < NumParameters; ++i)
+	{
+		if (GetParameterChannelMask(i) != 0 || Filter == EVariableFilter::None)
+		{
+			OutVariables.AddUnique(GetParameterVariable(i));
+		}
+	}
+}
+#endif
diff --git a/Engine/Plugins/FX/Niagara/Source/Niagara/Private/Stateless/Modules/NiagaraStatelessModule_GravityForce.cpp b/Engine/Plugins/FX/Niagara/Source/Niagara/Private/Stateless/Modules/NiagaraStatelessModule_GravityForce.cpp
new file mode 100644
--- /dev/null
+++ b/Engine/Plugins/FX/Niagara/Source/Niagara/Private/Stateless/Modules/NiagaraStatelessModule_GravityForce.cpp
@@ -0,0 +1,19 @@
+// Copyright Epic Games, Inc. All Rights Reserved.
+
+#include "Stateless/Modules/NiagaraStatelessModule_GravityForce.h"
+
+#include "Stateless/NiagaraStatelessEmitterDataBuildContext.h"
+#include "Stateless/Modules/NiagaraStatelessModuleCommon.h"
+
+void UNiagaraStatelessModule_GravityForce::BuildEmitterData(const FNiagaraStatelessEmitterDataBuildContext& BuildContext) const
+{
+	if (!IsModuleEnabled())
+	{
+		return;
+	}
+	const FNiagaraStatelessRangeVector3 GravityRange = GravityDistribution.CalculateRange(GetDefaultValue());
+
+	NiagaraStateless::FPhysicsBuildData& PhysicsBuildData = BuildContext.GetTransientBuildData<NiagaraStateless::FPhysicsBuildData>();
+	PhysicsBuildData.GravityRange.Min = GravityRange.Min;
+	PhysicsBuildData.GravityRange.Max = GravityRange.Max;
+}
diff --git a/Engine/Plugins/FX/Niagara/Source/Niagara/Private/Stateless/Modules/NiagaraStatelessModule_InitializeParticle.cpp b/Engine/Plugins/FX/Niagara/Source/Niagara/Private/Stateless/Modules/NiagaraStatelessModule_InitializeParticle.cpp
new file mode 100644
--- /dev/null
+++ b/Engine/Plugins/FX/Niagara/Source/Niagara/Private/Stateless/Modules/NiagaraStatelessModule_InitializeParticle.cpp
@@ -0,0 +1,171 @@
+// Copyright Epic Games, Inc. All Rights Reserved.
+
+#include "Stateless/Modules/NiagaraStatelessModule_InitializeParticle.h"
+
+#include "Stateless/NiagaraStatelessEmitterDataBuildContext.h"
+#include "Stateless/NiagaraStatelessModuleShaderParameters.h"
+#include "Stateless/NiagaraStatelessParticleSimContext.h"
+#include "Stateless/Modules/NiagaraStatelessModuleCommon.h"
+
+#include "NiagaraParameterBinding.h"
+#include "NiagaraParameterStore.h"
+
+namespace NSMInitializeParticlePrivate
+{
+	using FParameters = NiagaraStateless::FInitializeParticleModule_ShaderParameters;
+
+	static const uint32 EInitializeParticleModuleFlag_UniformSpriteSize = 1 << 0;
+	static const uint32 EInitializeParticleModuleFlag_UniformMeshScale = 1 << 1;
+
+	struct FModuleBuiltData
+	{
+		uint32							ModuleFlags = 0;
+		FUintVector3					InitialPosition = FUintVector3::ZeroValue;
+		FNiagaraStatelessRangeFloat		LifetimeRange;
+		FNiagaraStatelessRangeColor		ColorRange;
+		FNiagaraStatelessRangeFloat		MassRange;
+		FNiagaraStatelessRangeVector2	SpriteSizeRange;
+		FNiagaraStatelessRangeFloat		SpriteRotationRange;
+		FNiagaraStatelessRangeVector3	MeshScaleRange;
+		FNiagaraStatelessRangeFloat		RibbonWidthRange;
+
+		int32							PositionVariableOffset = INDEX_NONE;
+		int32							ColorVariableOffset = INDEX_NONE;
+		int32							RibbonWidthVariableOffset = INDEX_NONE;
+		int32							SpriteSizeVariableOffset = INDEX_NONE;
+		int32							SpriteRotationVariableOffset = INDEX_NONE;
+		int32							ScaleVariableOffset = INDEX_NONE;
+
+		int32							PreviousPositionVariableOffset = INDEX_NONE;
+		int32							PreviousRibbonWidthVariableOffset = INDEX_NONE;
+		int32							PreviousSpriteSizeVariableOffset = INDEX_NONE;
+		int32							PreviousSpriteRotationVariableOffset = INDEX_NONE;
+		int32							PreviousScaleVariableOffset = INDEX_NONE;
+	};
+
+	static void ParticleSimulate(const NiagaraStateless::FParticleSimulationContext& ParticleSimulationContext)
+	{
+		using namespace NiagaraStateless;
+
+		const FModuleBuiltData* ModuleBuiltData = ParticleSimulationContext.ReadBuiltData<FModuleBuiltData>();
+		const FParameters* ShaderParameters = ParticleSimulationContext.ReadParameterNestedStruct<FParameters>();
+
+		const bool bUniformSpriteSize = (ModuleBuiltData->ModuleFlags & EInitializeParticleModuleFlag_UniformSpriteSize) != 0;
+		const bool bUniformMeshScale = (ModuleBuiltData->ModuleFlags & EInitializeParticleModuleFlag_UniformMeshScale) != 0;
+
+		for (uint32 i = 0; i < ParticleSimulationContext.GetNumInstances(); ++i)
+		{
+			const FStatelessDistributionSampler<FVector3f> PositionSampler(ParticleSimulationContext, ModuleBuiltData->InitialPosition, i, 0);
+
+			const FVector3f		Position = PositionSampler.GetValue(ParticleSimulationContext, 0.0f);
+			const FLinearColor	Color = ParticleSimulationContext.RandomScaleBiasFloat(i, 1, ShaderParameters->InitializeParticle_ColorScale, ShaderParameters->InitializeParticle_ColorBias);
+			const float			RibbonWidth = ParticleSimulationContext.RandomScaleBiasFloat(i, 2, ShaderParameters->InitializeParticle_RibbonWidthScale, ShaderParameters->InitializeParticle_RibbonWidthBias);
+			const FVector2f		SpriteSize = ParticleSimulationContext.RandomScaleBiasFloat(i, 3, ShaderParameters->InitializeParticle_SpriteSizeScale, ShaderParameters->InitializeParticle_SpriteSizeBias, bUniformSpriteSize);
+			const float			SpriteRot = ParticleSimulationContext.RandomScaleBiasFloat(i, 4, ShaderParameters->InitializeParticle_SpriteRotationScale, ShaderParameters->InitializeParticle_SpriteRotationBias);
+			const FVector3f		Scale = ParticleSimulationContext.RandomScaleBiasFloat(i, 5, ShaderParameters->InitializeParticle_MeshScaleScale, ShaderParameters->InitializeParticle_MeshScaleBias, bUniformMeshScale);
+
+			ParticleSimulationContext.WriteParticleVariable(ModuleBuiltData->PositionVariableOffset, i, Position);
+			ParticleSimulationContext.WriteParticleVariable(ModuleBuiltData->ColorVariableOffset, i, Color);
+			ParticleSimulationContext.WriteParticleVariable(ModuleBuiltData->RibbonWidthVariableOffset, i, RibbonWidth);
+			ParticleSimulationContext.WriteParticleVariable(ModuleBuiltData->SpriteSizeVariableOffset, i, SpriteSize);
+			ParticleSimulationContext.WriteParticleVariable(ModuleBuiltData->SpriteRotationVariableOffset, i, SpriteRot);
+			ParticleSimulationContext.WriteParticleVariable(ModuleBuiltData->ScaleVariableOffset, i, Scale);
+
+			ParticleSimulationContext.WriteParticleVariable(ModuleBuiltData->PreviousPositionVariableOffset, i, Position);
+			ParticleSimulationContext.WriteParticleVariable(ModuleBuiltData->PreviousRibbonWidthVariableOffset, i, RibbonWidth);
+			ParticleSimulationContext.WriteParticleVariable(ModuleBuiltData->PreviousSpriteSizeVariableOffset, i, SpriteSize);
+			ParticleSimulationContext.WriteParticleVariable(ModuleBuiltData->PreviousSpriteRotationVariableOffset, i, SpriteRot);
+			ParticleSimulationContext.WriteParticleVariable(ModuleBuiltData->PreviousScaleVariableOffset, i, Scale);
+		}
+	}
+
+}
+
+void UNiagaraStatelessModule_InitializeParticle::BuildEmitterData(const FNiagaraStatelessEmitterDataBuildContext& BuildContext) const
+{
+	using namespace NSMInitializeParticlePrivate;
+
+	BuildContext.AddParticleSimulationExecSimulate(&ParticleSimulate);
+
+	FModuleBuiltData* BuiltData			= BuildContext.AllocateBuiltData<FModuleBuiltData>();
+	BuiltData->ModuleFlags				 = SpriteSizeDistribution.IsUniform() ? EInitializeParticleModuleFlag_UniformSpriteSize : 0;
+	BuiltData->ModuleFlags				|= MeshScaleDistribution.IsUniform() ? EInitializeParticleModuleFlag_UniformMeshScale : 0;
+
+	BuiltData->InitialPosition			= BuildContext.AddDistribution(InitialPositionDistribution);
+	BuiltData->LifetimeRange			= LifetimeDistribution.CalculateRange(FNiagaraStatelessGlobals::GetDefaultLifetimeValue());
+	BuiltData->ColorRange				= BuildContext.ConvertDistributionToRange(ColorDistribution, FNiagaraStatelessGlobals::GetDefaultColorValue());
+	BuiltData->MassRange				= BuildContext.ConvertDistributionToRange(MassDistribution, FNiagaraStatelessGlobals::GetDefaultMassValue());
+	BuiltData->SpriteSizeRange			= BuildContext.ConvertDistributionToRange(SpriteSizeDistribution, FNiagaraStatelessGlobals::GetDefaultSpriteSizeValue());
+	BuiltData->SpriteRotationRange		= BuildContext.ConvertDistributionToRange(SpriteRotationDistribution, FNiagaraStatelessGlobals::GetDefaultSpriteRotationValue());
+	BuiltData->MeshScaleRange			= BuildContext.ConvertDistributionToRange(MeshScaleDistribution, FNiagaraStatelessGlobals::GetDefaultScaleValue());
+	BuiltData->RibbonWidthRange			= BuildContext.ConvertDistributionToRange(RibbonWidthDistribution, FNiagaraStatelessGlobals::GetDefaultRibbonWidthValue());
+
+	const FNiagaraStatelessGlobals& StatelessGlobals = FNiagaraStatelessGlobals::Get();
+	BuiltData->PositionVariableOffset				= BuildContext.FindParticleVariableIndex(StatelessGlobals.PositionVariable);
+	BuiltData->ColorVariableOffset					= BuildContext.FindParticleVariableIndex(StatelessGlobals.ColorVariable);
+	BuiltData->RibbonWidthVariableOffset			= BuildContext.FindParticleVariableIndex(StatelessGlobals.RibbonWidthVariable);
+	BuiltData->SpriteSizeVariableOffset				= BuildContext.FindParticleVariableIndex(StatelessGlobals.SpriteSizeVariable);
+	BuiltData->SpriteRotationVariableOffset			= BuildContext.FindParticleVariableIndex(StatelessGlobals.SpriteRotationVariable);
+	BuiltData->ScaleVariableOffset					= BuildContext.FindParticleVariableIndex(StatelessGlobals.ScaleVariable);
+	BuiltData->PreviousPositionVariableOffset		= BuildContext.FindParticleVariableIndex(StatelessGlobals.PreviousPositionVariable);
+	BuiltData->PreviousRibbonWidthVariableOffset	= BuildContext.FindParticleVariableIndex(StatelessGlobals.PreviousRibbonWidthVariable);
+	BuiltData->PreviousSpriteSizeVariableOffset		= BuildContext.FindParticleVariableIndex(StatelessGlobals.PreviousSpriteSizeVariable);
+	BuiltData->PreviousSpriteRotationVariableOffset	= BuildContext.FindParticleVariableIndex(StatelessGlobals.PreviousSpriteRotationVariable);
+	BuiltData->PreviousScaleVariableOffset			= BuildContext.FindParticleVariableIndex(StatelessGlobals.PreviousScaleVariable);
+
+	NiagaraStateless::FPhysicsBuildData& PhysicsBuildData = BuildContext.GetTransientBuildData<NiagaraStateless::FPhysicsBuildData>();
+	PhysicsBuildData.MassRange			= MassDistribution.CalculateRange(FNiagaraStatelessGlobals::GetDefaultMassValue());
+}
+
+void UNiagaraStatelessModule_InitializeParticle::BuildShaderParameters(FNiagaraStatelessShaderParametersBuilder& ShaderParametersBuilder) const
+{
+	using namespace NSMInitializeParticlePrivate;
+
+	ShaderParametersBuilder.AddParameterNestedStruct<FParameters>();
+}
+
+void UNiagaraStatelessModule_InitializeParticle::SetShaderParameters(const FNiagaraStatelessSetShaderParameterContext& SetShaderParameterContext) const
+{
+	using namespace NSMInitializeParticlePrivate;
+
+	FParameters* Parameters = SetShaderParameterContext.GetParameterNestedStruct<FParameters>();
+	const FModuleBuiltData* ModuleBuiltData = SetShaderParameterContext.ReadBuiltData<FModuleBuiltData>();
+
+	Parameters->InitializeParticle_ModuleFlags			= ModuleBuiltData->ModuleFlags;
+	Parameters->InitializeParticle_InitialPosition		= ModuleBuiltData->InitialPosition;
+	SetShaderParameterContext.ConvertRangeToScaleBias(ModuleBuiltData->ColorRange,			Parameters->InitializeParticle_ColorScale, Parameters->InitializeParticle_ColorBias);
+	SetShaderParameterContext.ConvertRangeToScaleBias(ModuleBuiltData->SpriteSizeRange,		Parameters->InitializeParticle_SpriteSizeScale, Parameters->InitializeParticle_SpriteSizeBias);
+	SetShaderParameterContext.ConvertRangeToScaleBias(ModuleBuiltData->SpriteRotationRange, Parameters->InitializeParticle_SpriteRotationScale, Parameters->InitializeParticle_SpriteRotationBias);
+	SetShaderParameterContext.ConvertRangeToScaleBias(ModuleBuiltData->MeshScaleRange,		Parameters->InitializeParticle_MeshScaleScale, Parameters->InitializeParticle_MeshScaleBias);
+	SetShaderParameterContext.ConvertRangeToScaleBias(ModuleBuiltData->RibbonWidthRange,	Parameters->InitializeParticle_RibbonWidthScale, Parameters->InitializeParticle_RibbonWidthBias);
+}
+
+#if WITH_EDITORONLY_DATA
+const TCHAR* UNiagaraStatelessModule_InitializeParticle::GetShaderTemplatePath() const
+{
+	return TEXT("/Plugin/FX/Niagara/Private/Stateless/Modules/NiagaraStatelessModule_InitializeParticle.ush");
+}
+
+void UNiagaraStatelessModule_InitializeParticle::GetOutputVariables(TArray<FNiagaraVariableBase>& OutVariables, EVariableFilter Filter) const
+{
+	using namespace NSMInitializeParticlePrivate;
+
+	const FNiagaraStatelessGlobals& StatelessGlobals = FNiagaraStatelessGlobals::Get();
+	OutVariables.AddUnique(StatelessGlobals.PositionVariable);
+	OutVariables.AddUnique(StatelessGlobals.ColorVariable);
+	OutVariables.AddUnique(StatelessGlobals.SpriteSizeVariable);
+	OutVariables.AddUnique(StatelessGlobals.SpriteRotationVariable);
+	OutVariables.AddUnique(StatelessGlobals.ScaleVariable);
+
+	OutVariables.AddUnique(StatelessGlobals.PreviousPositionVariable);
+	OutVariables.AddUnique(StatelessGlobals.PreviousSpriteSizeVariable);
+	OutVariables.AddUnique(StatelessGlobals.PreviousSpriteRotationVariable);
+	OutVariables.AddUnique(StatelessGlobals.PreviousScaleVariable);
+
+	if (bWriteRibbonWidth || Filter == EVariableFilter::None)
+	{
+		OutVariables.AddUnique(StatelessGlobals.RibbonWidthVariable);
+		OutVariables.AddUnique(StatelessGlobals.PreviousRibbonWidthVariable);
+	}
+}
+#endif
diff --git a/Engine/Plugins/FX/Niagara/Source/Niagara/Private/Stateless/Modules/NiagaraStatelessModule_MeshRotationRate.cpp b/Engine/Plugins/FX/Niagara/Source/Niagara/Private/Stateless/Modules/NiagaraStatelessModule_MeshRotationRate.cpp
new file mode 100644
--- /dev/null
+++ b/Engine/Plugins/FX/Niagara/Source/Niagara/Private/Stateless/Modules/NiagaraStatelessModule_MeshRotationRate.cpp
@@ -0,0 +1,116 @@
+// Copyright Epic Games, Inc. All Rights Reserved.
+
+#include "Stateless/Modules/NiagaraStatelessModule_MeshRotationRate.h"
+
+#include "Stateless/NiagaraStatelessEmitterDataBuildContext.h"
+#include "Stateless/NiagaraStatelessModuleShaderParameters.h"
+#include "Stateless/NiagaraStatelessParticleSimContext.h"
+
+namespace NSMMeshRotationRatePrivate
+{
+	using FParameters = NiagaraStateless::FMeshRotationRateModule_ShaderParameters;
+
+	struct FModuleBuiltData
+	{
+		bool							ModuleEnabled = false;
+		FNiagaraStatelessRangeVector3	RotationRange;
+		FUintVector3					RateScaleParameters = FUintVector3::ZeroValue;
+		int32							MeshOrientationVariableOffset = INDEX_NONE;
+		int32							PreviousMeshOrientationVariableOffset = INDEX_NONE;
+	};
+
+	static void ParticleSimulate(const NiagaraStateless::FParticleSimulationContext& ParticleSimulationContext)
+	{
+		using namespace NiagaraStateless;
+
+		const FModuleBuiltData* ModuleBuiltData = ParticleSimulationContext.ReadBuiltData<FModuleBuiltData>();
+		const FParameters* ShaderParameters = ParticleSimulationContext.ReadParameterNestedStruct<FParameters>();
+
+		const float* LifetimeData = ParticleSimulationContext.GetParticleLifetime();
+		const float* AgeData = ParticleSimulationContext.GetParticleNormalizedAge();
+		const float* PreviousAgeData = ParticleSimulationContext.GetParticlePreviousNormalizedAge();
+
+		for (uint32 i = 0; i < ParticleSimulationContext.GetNumInstances(); ++i)
+		{
+			const FVector3f RotationRate = ParticleSimulationContext.RandomScaleBiasFloat(i, 0, ShaderParameters->MeshRotationRate_Scale, ShaderParameters->MeshRotationRate_Bias);
+			const FVector3f MultiplyRate = ParticleSimulationContext.SampleCurve<FVector3f>(ShaderParameters->MeshRotationRate_RateScaleParameters, AgeData[i]);
+			const FVector3f PreviousMultiplyRate = ParticleSimulationContext.SampleCurve<FVector3f>(ShaderParameters->MeshRotationRate_RateScaleParameters, PreviousAgeData[i]);
+
+			FQuat4f MeshOrientation = ParticleSimulationContext.ReadParticleVariable(ModuleBuiltData->MeshOrientationVariableOffset, i, FQuat4f::Identity);
+			FQuat4f PreviousMeshOrientation = ParticleSimulationContext.ReadParticleVariable(ModuleBuiltData->PreviousMeshOrientationVariableOffset, i, FQuat4f::Identity);
+
+			MeshOrientation *= ParticleSimulationContext.RotatorToQuat(LifetimeData[i] * RotationRate * MultiplyRate);
+			PreviousMeshOrientation *= ParticleSimulationContext.RotatorToQuat(LifetimeData[i] * RotationRate * PreviousMultiplyRate);
+
+			ParticleSimulationContext.WriteParticleVariable(ModuleBuiltData->MeshOrientationVariableOffset, i, MeshOrientation);
+			ParticleSimulationContext.WriteParticleVariable(ModuleBuiltData->PreviousMeshOrientationVariableOffset, i, PreviousMeshOrientation);
+		}
+	}
+}
+
+void UNiagaraStatelessModule_MeshRotationRate::BuildEmitterData(const FNiagaraStatelessEmitterDataBuildContext& BuildContext) const
+{
+	using namespace NSMMeshRotationRatePrivate;
+
+	FModuleBuiltData* BuiltData = BuildContext.AllocateBuiltData<FModuleBuiltData>();
+	if (!IsModuleEnabled())
+	{
+		return;
+	}
+
+	const FNiagaraStatelessGlobals& StatelessGlobals	= FNiagaraStatelessGlobals::Get();
+	BuiltData->MeshOrientationVariableOffset			= BuildContext.FindParticleVariableIndex(StatelessGlobals.MeshOrientationVariable);
+	BuiltData->PreviousMeshOrientationVariableOffset	= BuildContext.FindParticleVariableIndex(StatelessGlobals.PreviousMeshOrientationVariable);
+
+	if (BuiltData->MeshOrientationVariableOffset != INDEX_NONE || BuiltData->PreviousMeshOrientationVariableOffset != INDEX_NONE)
+	{
+		BuiltData->ModuleEnabled = true;
+		BuiltData->RotationRange = BuildContext.ConvertDistributionToRange(RotationRateDistribution, FVector3f::ZeroVector);
+		BuiltData->RotationRange.Min *= 1.0f / 360.0f;
+		BuiltData->RotationRange.Max *= 1.0f / 360.0f;
+		if (bUseRateScale)
+		{
+			BuiltData->RateScaleParameters = BuildContext.AddDistributionAsCurve(RateScaleDistribution, 1.0f);
+		}
+		else
+		{
+			BuiltData->RateScaleParameters = BuildContext.AddDistributionAsCurve(FNiagaraDistributionCurveVector3(ENiagaraDistributionCurveLUTMode::Accumulate), 1.0f);
+		}
+
+		BuildContext.AddParticleSimulationExecSimulate(&ParticleSimulate);
+	}
+}
+
+void UNiagaraStatelessModule_MeshRotationRate::BuildShaderParameters(FNiagaraStatelessShaderParametersBuilder& ShaderParametersBuilder) const
+{
+	using namespace NSMMeshRotationRatePrivate;
+
+	ShaderParametersBuilder.AddParameterNestedStruct<FParameters>();
+}
+
+void UNiagaraStatelessModule_MeshRotationRate::SetShaderParameters(const FNiagaraStatelessSetShaderParameterContext& SetShaderParameterContext) const
+{
+	using namespace NSMMeshRotationRatePrivate;
+
+	FParameters* Parameters = SetShaderParameterContext.GetParameterNestedStruct<FParameters>();
+	const FModuleBuiltData* ModuleBuiltData = SetShaderParameterContext.ReadBuiltData<FModuleBuiltData>();
+	Parameters->MeshRotationRate_ModuleEnabled = ModuleBuiltData->ModuleEnabled;
+	SetShaderParameterContext.ConvertRangeToScaleBias(ModuleBuiltData->RotationRange, Parameters->MeshRotationRate_Scale, Parameters->MeshRotationRate_Bias);
+	Parameters->MeshRotationRate_RateScaleParameters = ModuleBuiltData->RateScaleParameters;
+}
+
+#if WITH_EDITORONLY_DATA
+const TCHAR* UNiagaraStatelessModule_MeshRotationRate::GetShaderTemplatePath() const
+{
+	return TEXT("/Plugin/FX/Niagara/Private/Stateless/Modules/NiagaraStatelessModule_MeshRotationRate.ush");
+}
+
+void UNiagaraStatelessModule_MeshRotationRate::GetOutputVariables(TArray<FNiagaraVariableBase>& OutVariables, EVariableFilter Filter) const
+{
+	using namespace NSMMeshRotationRatePrivate;
+
+	const FNiagaraStatelessGlobals& StatelessGlobals = FNiagaraStatelessGlobals::Get();
+	OutVariables.AddUnique(StatelessGlobals.MeshOrientationVariable);
+	OutVariables.AddUnique(StatelessGlobals.PreviousMeshOrientationVariable);
+}
+#endif
diff --git a/Engine/Plugins/FX/Niagara/Source/Niagara/Private/Stateless/Modules/NiagaraStatelessModule_ScaleColor.cpp b/Engine/Plugins/FX/Niagara/Source/Niagara/Private/Stateless/Modules/NiagaraStatelessModule_ScaleColor.cpp
new file mode 100644
--- /dev/null
+++ b/Engine/Plugins/FX/Niagara/Source/Niagara/Private/Stateless/Modules/NiagaraStatelessModule_ScaleColor.cpp
@@ -0,0 +1,87 @@
+// Copyright Epic Games, Inc. All Rights Reserved.
+
+#include "Stateless/Modules/NiagaraStatelessModule_ScaleColor.h"
+#include "Stateless/NiagaraStatelessEmitterDataBuildContext.h"
+#include "Stateless/NiagaraStatelessModuleShaderParameters.h"
+#include "Stateless/NiagaraStatelessParticleSimContext.h"
+
+namespace NSMScaleColorPrivate
+{
+	struct FModuleBuiltData
+	{
+		FUintVector3	DistributionParameters = FUintVector3::ZeroValue;
+		int32			ColorVariableOffset = INDEX_NONE;
+	};
+
+	using FParameters = NiagaraStateless::FScaleColorModule_ShaderParameters;
+
+	static void ParticleSimulate(const NiagaraStateless::FParticleSimulationContext& ParticleSimulationContext)
+	{
+		using namespace NiagaraStateless;
+
+		const FModuleBuiltData* ModuleBuiltData = ParticleSimulationContext.ReadBuiltData<FModuleBuiltData>();
+		const float* NormalizedAgeData = ParticleSimulationContext.GetParticleNormalizedAge();
+
+		for (uint32 i = 0; i < ParticleSimulationContext.GetNumInstances(); ++i)
+		{
+			FStatelessDistributionSampler<FLinearColor> ColorSampler(ParticleSimulationContext, ModuleBuiltData->DistributionParameters, i, 0);
+
+			FLinearColor Color = ParticleSimulationContext.ReadParticleVariable(ModuleBuiltData->ColorVariableOffset, i, FLinearColor::White);
+
+			Color *= ColorSampler.GetValue(ParticleSimulationContext, NormalizedAgeData[i]);
+
+			ParticleSimulationContext.WriteParticleVariable(ModuleBuiltData->ColorVariableOffset, i, Color);
+		}
+	}
+}
+
+void UNiagaraStatelessModule_ScaleColor::BuildEmitterData(const FNiagaraStatelessEmitterDataBuildContext& BuildContext) const
+{
+	using namespace NSMScaleColorPrivate;
+
+	FModuleBuiltData* BuiltData = BuildContext.AllocateBuiltData<FModuleBuiltData>();
+	if (!IsModuleEnabled())
+	{
+		return;
+	}
+
+	const FNiagaraStatelessGlobals& StatelessGlobals = FNiagaraStatelessGlobals::Get();
+	BuiltData->ColorVariableOffset = BuildContext.FindParticleVariableIndex(StatelessGlobals.ColorVariable);
+
+	if (BuiltData->ColorVariableOffset != INDEX_NONE)
+	{
+		BuiltData->DistributionParameters = BuildContext.AddDistribution(ScaleDistribution);
+
+		BuildContext.AddParticleSimulationExecSimulate(&ParticleSimulate);
+	}
+}
+
+void UNiagaraStatelessModule_ScaleColor::BuildShaderParameters(FNiagaraStatelessShaderParametersBuilder& ShaderParametersBuilder) const
+{
+	using namespace NSMScaleColorPrivate;
+
+	ShaderParametersBuilder.AddParameterNestedStruct<FParameters>();
+}
+
+void UNiagaraStatelessModule_ScaleColor::SetShaderParameters(const FNiagaraStatelessSetShaderParameterContext& SetShaderParameterContext) const
+{
+	using namespace NSMScaleColorPrivate;
+
+	const FModuleBuiltData* ModuleBuiltData = SetShaderParameterContext.ReadBuiltData<FModuleBuiltData>();
+
+	FParameters* Parameters = SetShaderParameterContext.GetParameterNestedStruct<FParameters>();
+	Parameters->ScaleColor_Distribution = ModuleBuiltData->DistributionParameters;
+}
+
+#if WITH_EDITORONLY_DATA
+const TCHAR* UNiagaraStatelessModule_ScaleColor::GetShaderTemplatePath() const
+{
+	return TEXT("/Plugin/FX/Niagara/Private/Stateless/Modules/NiagaraStatelessModule_ScaleColor.ush");
+}
+
+void UNiagaraStatelessModule_ScaleColor::GetOutputVariables(TArray<FNiagaraVariableBase>& OutVariables, EVariableFilter Filter) const
+{
+	const FNiagaraStatelessGlobals& StatelessGlobals = FNiagaraStatelessGlobals::Get();
+	OutVariables.AddUnique(StatelessGlobals.ColorVariable);
+}
+#endif
diff --git a/Engine/Plugins/FX/Niagara/Source/Niagara/Private/Stateless/Modules/NiagaraStatelessModule_ScaleMeshSize.cpp b/Engine/Plugins/FX/Niagara/Source/Niagara/Private/Stateless/Modules/NiagaraStatelessModule_ScaleMeshSize.cpp
new file mode 100644
--- /dev/null
+++ b/Engine/Plugins/FX/Niagara/Source/Niagara/Private/Stateless/Modules/NiagaraStatelessModule_ScaleMeshSize.cpp
@@ -0,0 +1,121 @@
+// Copyright Epic Games, Inc. All Rights Reserved.
+
+#include "Stateless/Modules/NiagaraStatelessModule_ScaleMeshSize.h"
+#include "Stateless/NiagaraStatelessEmitterDataBuildContext.h"
+#include "Stateless/NiagaraStatelessModuleShaderParameters.h"
+#include "Stateless/NiagaraStatelessParticleSimContext.h"
+
+namespace NSMScaleMeshSizePrivate
+{
+	struct FModuleBuiltData
+	{
+		FUintVector3	DistributionParameters = FUintVector3::ZeroValue;
+		FVector3f		CurveScale = FVector3f::OneVector;
+		int32			CurveScaleOffset = INDEX_NONE;
+
+		int32			ScaleVariableOffset = INDEX_NONE;
+		int32			PreviousScaleVariableOffset = INDEX_NONE;
+	};
+
+	using FParameters = NiagaraStateless::FScaleMeshSizeModule_ShaderParameters;
+
+	static void ParticleSimulate(const NiagaraStateless::FParticleSimulationContext& ParticleSimulationContext)
+	{
+		using namespace NiagaraStateless;
+
+		const FModuleBuiltData* ModuleBuiltData = ParticleSimulationContext.ReadBuiltData<FModuleBuiltData>();
+		const float* NormalizedAgeData = ParticleSimulationContext.GetParticleNormalizedAge();
+		const float* PreviousNormalizedAgeData = ParticleSimulationContext.GetParticlePreviousNormalizedAge();
+
+		const FVector3f CurveScale = ParticleSimulationContext.GetParameterBufferFloat(ModuleBuiltData->CurveScaleOffset, ModuleBuiltData->CurveScale);
+
+		for (uint32 i = 0; i < ParticleSimulationContext.GetNumInstances(); ++i)
+		{
+			FStatelessDistributionSampler<FVector3f> ScaleSampler(ParticleSimulationContext, ModuleBuiltData->DistributionParameters, i, 0);
+
+			FVector3f Scale = ParticleSimulationContext.ReadParticleVariable(ModuleBuiltData->ScaleVariableOffset, i, FVector3f::OneVector);
+			FVector3f PreviousScale = ParticleSimulationContext.ReadParticleVariable(ModuleBuiltData->PreviousScaleVariableOffset, i, FVector3f::OneVector);
+
+			Scale = Scale * ScaleSampler.GetValue(ParticleSimulationContext, NormalizedAgeData[i]) * CurveScale;
+			PreviousScale = PreviousScale * ScaleSampler.GetValue(ParticleSimulationContext, PreviousNormalizedAgeData[i]) * CurveScale;
+
+			ParticleSimulationContext.WriteParticleVariable(ModuleBuiltData->ScaleVariableOffset, i, Scale);
+			ParticleSimulationContext.WriteParticleVariable(ModuleBuiltData->PreviousScaleVariableOffset, i, PreviousScale);
+		}
+	}
+}
+
+void UNiagaraStatelessModule_ScaleMeshSize::PostInitProperties()
+{
+	Super::PostInitProperties();
+
+#if WITH_EDITORONLY_DATA
+	if (HasAnyFlags(RF_ClassDefaultObject) == false)
+	{
+		ScaleCurveRange.SetUsage(ENiagaraParameterBindingUsage::NotParticle);
+		ScaleCurveRange.SetAllowedTypeDefinitions({ FNiagaraTypeDefinition::GetVec3Def() });
+		ScaleCurveRange.SetDefaultParameter(FNiagaraTypeDefinition::GetVec3Def(), FVector3f::OneVector);
+	}
+#endif
+}
+
+void UNiagaraStatelessModule_ScaleMeshSize::BuildEmitterData(const FNiagaraStatelessEmitterDataBuildContext& BuildContext) const
+{
+	using namespace NSMScaleMeshSizePrivate;
+
+	FModuleBuiltData* BuiltData = BuildContext.AllocateBuiltData<FModuleBuiltData>();
+	if (!IsModuleEnabled())
+	{
+		return;
+	}
+
+	const FNiagaraStatelessGlobals& StatelessGlobals = FNiagaraStatelessGlobals::Get();
+	BuiltData->ScaleVariableOffset = BuildContext.FindParticleVariableIndex(StatelessGlobals.ScaleVariable);
+	BuiltData->PreviousScaleVariableOffset = BuildContext.FindParticleVariableIndex(StatelessGlobals.PreviousScaleVariable);
+
+	if (BuiltData->ScaleVariableOffset == INDEX_NONE && BuiltData->PreviousScaleVariableOffset == INDEX_NONE)
+	{
+		return;
+	}
+
+	BuiltData->DistributionParameters = BuildContext.AddDistribution(ScaleDistribution);
+	if (UseScaleCurveRange())
+	{
+		BuiltData->CurveScaleOffset = BuildContext.AddRendererBinding(ScaleCurveRange.ResolvedParameter);
+		BuiltData->CurveScale = ScaleCurveRange.GetDefaultValue<FVector3f>();
+	}
+	BuildContext.AddParticleSimulationExecSimulate(&ParticleSimulate);
+}
+
+void UNiagaraStatelessModule_ScaleMeshSize::BuildShaderParameters(FNiagaraStatelessShaderParametersBuilder& ShaderParametersBuilder) const
+{
+	using namespace NSMScaleMeshSizePrivate;
+
+	ShaderParametersBuilder.AddParameterNestedStruct<FParameters>();
+}
+
+void UNiagaraStatelessModule_ScaleMeshSize::SetShaderParameters(const FNiagaraStatelessSetShaderParameterContext& SetShaderParameterContext) const
+{
+	using namespace NSMScaleMeshSizePrivate;
+
+	const FModuleBuiltData* ModuleBuiltData = SetShaderParameterContext.ReadBuiltData<FModuleBuiltData>();
+
+	FParameters* Parameters						= SetShaderParameterContext.GetParameterNestedStruct<FParameters>();
+	Parameters->ScaleMeshSize_Distribution		= ModuleBuiltData->DistributionParameters;
+	Parameters->ScaleMeshSize_CurveScale		= ModuleBuiltData->CurveScale;
+	Parameters->ScaleMeshSize_CurveScaleOffset	= ModuleBuiltData->CurveScaleOffset;
+}
+
+#if WITH_EDITORONLY_DATA
+const TCHAR* UNiagaraStatelessModule_ScaleMeshSize::GetShaderTemplatePath() const
+{
+	return TEXT("/Plugin/FX/Niagara/Private/Stateless/Modules/NiagaraStatelessModule_ScaleMeshSize.ush");
+}
+
+void UNiagaraStatelessModule_ScaleMeshSize::GetOutputVariables(TArray<FNiagaraVariableBase>& OutVariables, EVariableFilter Filter) const
+{
+	const FNiagaraStatelessGlobals& StatelessGlobals = FNiagaraStatelessGlobals::Get();
+	OutVariables.AddUnique(StatelessGlobals.ScaleVariable);
+	OutVariables.AddUnique(StatelessGlobals.PreviousScaleVariable);
+}
+#endif
diff --git a/Engine/Plugins/FX/Niagara/Source/Niagara/Private/Stateless/Modules/NiagaraStatelessModule_ScaleRibbonWidth.cpp b/Engine/Plugins/FX/Niagara/Source/Niagara/Private/Stateless/Modules/NiagaraStatelessModule_ScaleRibbonWidth.cpp
new file mode 100644
--- /dev/null
+++ b/Engine/Plugins/FX/Niagara/Source/Niagara/Private/Stateless/Modules/NiagaraStatelessModule_ScaleRibbonWidth.cpp
@@ -0,0 +1,122 @@
+// Copyright Epic Games, Inc. All Rights Reserved.
+
+#include "Stateless/Modules/NiagaraStatelessModule_ScaleRibbonWidth.h"
+#include "Stateless/NiagaraStatelessEmitterDataBuildContext.h"
+#include "Stateless/NiagaraStatelessModuleShaderParameters.h"
+#include "Stateless/NiagaraStatelessParticleSimContext.h"
+
+namespace NSMScaleRibbonWidthPrivate
+{
+	struct FModuleBuiltData
+	{
+		FUintVector3	DistributionParameters = FUintVector3::ZeroValue;
+		float			CurveScale = 1.0f;
+		int32			CurveScaleOffset = INDEX_NONE;
+
+		int32			RibbonWidthVariableOffset = INDEX_NONE;
+		int32			PreviousRibbonWidthVariableOffset = INDEX_NONE;
+	};
+
+	using FParameters = NiagaraStateless::FScaleRibbonWidthModule_ShaderParameters;
+
+	static void ParticleSimulate(const NiagaraStateless::FParticleSimulationContext& ParticleSimulationContext)
+	{
+		using namespace NiagaraStateless;
+
+		const FModuleBuiltData* ModuleBuiltData = ParticleSimulationContext.ReadBuiltData<FModuleBuiltData>();
+
+		const float ScaleFactor = ParticleSimulationContext.GetParameterBufferFloat(ModuleBuiltData->CurveScaleOffset, ModuleBuiltData->CurveScale);
+		const float* NormalizedAgeData = ParticleSimulationContext.GetParticleNormalizedAge();
+		const float* PreviousNormalizedAgeData = ParticleSimulationContext.GetParticlePreviousNormalizedAge();
+
+		for (uint32 i = 0; i < ParticleSimulationContext.GetNumInstances(); ++i)
+		{
+			const FStatelessDistributionSampler<float> SpriteScaleSampler(ParticleSimulationContext, ModuleBuiltData->DistributionParameters, i, 0);
+
+			const float NormalizedAge = NormalizedAgeData[i];
+			const float PreviousNormalizedAge = PreviousNormalizedAgeData[i];
+
+			const float RibbonWidth = ParticleSimulationContext.ReadParticleVariable(ModuleBuiltData->RibbonWidthVariableOffset, i, 0.0f);
+			const float PreviousRibbonWidth = ParticleSimulationContext.ReadParticleVariable(ModuleBuiltData->PreviousRibbonWidthVariableOffset, i, 0.0f);
+			const float Scale = SpriteScaleSampler.GetValue(ParticleSimulationContext, NormalizedAge) * ScaleFactor;
+			const float PreviousScale = SpriteScaleSampler.GetValue(ParticleSimulationContext, PreviousNormalizedAge) * ScaleFactor;
+
+			ParticleSimulationContext.WriteParticleVariable(ModuleBuiltData->RibbonWidthVariableOffset, i, RibbonWidth * Scale);
+			ParticleSimulationContext.WriteParticleVariable(ModuleBuiltData->PreviousRibbonWidthVariableOffset, i, PreviousRibbonWidth * PreviousScale);
+		}
+	}
+}
+
+void UNiagaraStatelessModule_ScaleRibbonWidth::PostInitProperties()
+{
+	using namespace NSMScaleRibbonWidthPrivate;
+
+	Super::PostInitProperties();
+
+#if WITH_EDITORONLY_DATA
+	if (HasAnyFlags(RF_ClassDefaultObject) == false)
+	{
+		ScaleCurveRange.SetUsage(ENiagaraParameterBindingUsage::NotParticle);
+		ScaleCurveRange.SetAllowedTypeDefinitions({ FNiagaraTypeDefinition::GetVec2Def() });
+		ScaleCurveRange.SetDefaultParameter(FNiagaraTypeDefinition::GetFloatDef(), 1.0f);
+	}
+#endif
+}
+
+void UNiagaraStatelessModule_ScaleRibbonWidth::BuildEmitterData(const FNiagaraStatelessEmitterDataBuildContext& BuildContext) const
+{
+	using namespace NSMScaleRibbonWidthPrivate;
+
+	FModuleBuiltData* BuiltData = BuildContext.AllocateBuiltData<FModuleBuiltData>();
+
+	const FNiagaraStatelessGlobals& StatelessGlobals = FNiagaraStatelessGlobals::Get();
+	BuiltData->RibbonWidthVariableOffset			= BuildContext.FindParticleVariableIndex(StatelessGlobals.RibbonWidthVariable);
+	BuiltData->PreviousRibbonWidthVariableOffset	= BuildContext.FindParticleVariableIndex(StatelessGlobals.PreviousRibbonWidthVariable);
+
+	const bool bAttributesUsed = (BuiltData->RibbonWidthVariableOffset != INDEX_NONE || BuiltData->PreviousRibbonWidthVariableOffset != INDEX_NONE);
+	if (IsModuleEnabled() && bAttributesUsed)
+	{
+		BuiltData->DistributionParameters = BuildContext.AddDistribution(ScaleDistribution);
+		if (UseScaleCurveRange())
+		{
+			BuiltData->CurveScaleOffset = BuildContext.AddRendererBinding(ScaleCurveRange.ResolvedParameter);
+			BuiltData->CurveScale = ScaleCurveRange.GetDefaultValue<float>();
+		}
+
+		BuildContext.AddParticleSimulationExecSimulate(&ParticleSimulate);
+	}
+}
+
+void UNiagaraStatelessModule_ScaleRibbonWidth::BuildShaderParameters(FNiagaraStatelessShaderParametersBuilder& ShaderParametersBuilder) const
+{
+	using namespace NSMScaleRibbonWidthPrivate;
+
+	ShaderParametersBuilder.AddParameterNestedStruct<FParameters>();
+}
+
+void UNiagaraStatelessModule_ScaleRibbonWidth::SetShaderParameters(const FNiagaraStatelessSetShaderParameterContext& SetShaderParameterContext) const
+{
+	using namespace NSMScaleRibbonWidthPrivate;
+
+	const FModuleBuiltData* ModuleBuiltData = SetShaderParameterContext.ReadBuiltData<FModuleBuiltData>();
+
+	FParameters* Parameters						= SetShaderParameterContext.GetParameterNestedStruct<FParameters>();
+	Parameters->ScaleRibbonWidth_Distribution	= ModuleBuiltData->DistributionParameters;
+	Parameters->ScaleRibbonWidth_CurveScale		= SetShaderParameterContext.GetRendererParameterValue(ModuleBuiltData->CurveScaleOffset, ModuleBuiltData->CurveScale);
+}
+
+#if WITH_EDITORONLY_DATA
+const TCHAR* UNiagaraStatelessModule_ScaleRibbonWidth::GetShaderTemplatePath() const
+{
+	return TEXT("/Plugin/FX/Niagara/Private/Stateless/Modules/NiagaraStatelessModule_ScaleRibbonWidth.ush");
+}
+
+void UNiagaraStatelessModule_ScaleRibbonWidth::GetOutputVariables(TArray<FNiagaraVariableBase>& OutVariables, EVariableFilter Filter) const
+{
+	using namespace NSMScaleRibbonWidthPrivate;
+
+	const FNiagaraStatelessGlobals& StatelessGlobals = FNiagaraStatelessGlobals::Get();
+	OutVariables.AddUnique(StatelessGlobals.RibbonWidthVariable);
+	OutVariables.AddUnique(StatelessGlobals.PreviousRibbonWidthVariable);
+}
+#endif
diff --git a/Engine/Plugins/FX/Niagara/Source/Niagara/Private/Stateless/Modules/NiagaraStatelessModule_ScaleSpriteSize.cpp b/Engine/Plugins/FX/Niagara/Source/Niagara/Private/Stateless/Modules/NiagaraStatelessModule_ScaleSpriteSize.cpp
new file mode 100644
--- /dev/null
+++ b/Engine/Plugins/FX/Niagara/Source/Niagara/Private/Stateless/Modules/NiagaraStatelessModule_ScaleSpriteSize.cpp
@@ -0,0 +1,123 @@
+// Copyright Epic Games, Inc. All Rights Reserved.
+
+#include "Stateless/Modules/NiagaraStatelessModule_ScaleSpriteSize.h"
+#include "Stateless/NiagaraStatelessEmitterDataBuildContext.h"
+#include "Stateless/NiagaraStatelessModuleShaderParameters.h"
+#include "Stateless/NiagaraStatelessParticleSimContext.h"
+#include "NiagaraParameterBinding.h"
+
+namespace NSMScaleSpriteSizePrivate
+{
+	struct FModuleBuiltData
+	{
+		FUintVector3	DistributionParameters = FUintVector3::ZeroValue;
+		FVector2f		CurveScale = FVector2f::One();
+		int32			CurveScaleOffset = INDEX_NONE;
+
+		int32			SpriteSizeVariableOffset = INDEX_NONE;
+		int32			PreviousSpriteSizeVariableOffset = INDEX_NONE;
+	};
+
+	using FParameters = NiagaraStateless::FScaleSpriteSizeModule_ShaderParameters;
+
+	static void ParticleSimulate(const NiagaraStateless::FParticleSimulationContext& ParticleSimulationContext)
+	{
+		using namespace NiagaraStateless;
+
+		const FModuleBuiltData* ModuleBuiltData = ParticleSimulationContext.ReadBuiltData<FModuleBuiltData>();
+
+		const FVector2f ScaleFactor = ParticleSimulationContext.GetParameterBufferFloat(ModuleBuiltData->CurveScaleOffset, ModuleBuiltData->CurveScale);
+		const float* NormalizedAgeData = ParticleSimulationContext.GetParticleNormalizedAge();
+		const float* PreviousNormalizedAgeData = ParticleSimulationContext.GetParticlePreviousNormalizedAge();
+
+		for (uint32 i = 0; i < ParticleSimulationContext.GetNumInstances(); ++i)
+		{
+			const FStatelessDistributionSampler<FVector2f> SpriteScaleSampler(ParticleSimulationContext, ModuleBuiltData->DistributionParameters, i, 0);
+
+			const float NormalizedAge = NormalizedAgeData[i];
+			const float PreviousNormalizedAge = PreviousNormalizedAgeData[i];
+
+			const FVector2f SpriteSize = ParticleSimulationContext.ReadParticleVariable(ModuleBuiltData->SpriteSizeVariableOffset, i, FVector2f::ZeroVector);
+			const FVector2f PreviousSpriteSize = ParticleSimulationContext.ReadParticleVariable(ModuleBuiltData->PreviousSpriteSizeVariableOffset, i, FVector2f::ZeroVector);
+			const FVector2f Scale = SpriteScaleSampler.GetValue(ParticleSimulationContext, NormalizedAge) * ScaleFactor;
+			const FVector2f PreviousScale = SpriteScaleSampler.GetValue(ParticleSimulationContext, PreviousNormalizedAge) * ScaleFactor;
+
+			ParticleSimulationContext.WriteParticleVariable(ModuleBuiltData->SpriteSizeVariableOffset, i, SpriteSize * Scale);
+			ParticleSimulationContext.WriteParticleVariable(ModuleBuiltData->PreviousSpriteSizeVariableOffset, i, PreviousSpriteSize * PreviousScale);
+		}
+	}
+}
+
+void UNiagaraStatelessModule_ScaleSpriteSize::PostInitProperties()
+{
+	using namespace NSMScaleSpriteSizePrivate;
+
+	Super::PostInitProperties();
+
+#if WITH_EDITORONLY_DATA
+	if (HasAnyFlags(RF_ClassDefaultObject) == false)
+	{
+		ScaleCurveRange.SetUsage(ENiagaraParameterBindingUsage::NotParticle);
+		ScaleCurveRange.SetAllowedTypeDefinitions({ FNiagaraTypeDefinition::GetVec2Def() });
+		ScaleCurveRange.SetDefaultParameter(FNiagaraTypeDefinition::GetVec2Def(), FVector2f::One());
+	}
+#endif
+}
+
+void UNiagaraStatelessModule_ScaleSpriteSize::BuildEmitterData(const FNiagaraStatelessEmitterDataBuildContext& BuildContext) const
+{
+	using namespace NSMScaleSpriteSizePrivate;
+
+	FModuleBuiltData* BuiltData = BuildContext.AllocateBuiltData<FModuleBuiltData>();
+
+	const FNiagaraStatelessGlobals& StatelessGlobals = FNiagaraStatelessGlobals::Get();
+	BuiltData->SpriteSizeVariableOffset			= BuildContext.FindParticleVariableIndex(StatelessGlobals.SpriteSizeVariable);
+	BuiltData->PreviousSpriteSizeVariableOffset	= BuildContext.FindParticleVariableIndex(StatelessGlobals.PreviousSpriteSizeVariable);
+
+	const bool bAttributesUsed = (BuiltData->SpriteSizeVariableOffset != INDEX_NONE || BuiltData->PreviousSpriteSizeVariableOffset != INDEX_NONE);
+	if (IsModuleEnabled() && bAttributesUsed)
+	{
+		BuiltData->DistributionParameters = BuildContext.AddDistribution(ScaleDistribution);
+		if (UseScaleCurveRange())
+		{
+			BuiltData->CurveScaleOffset = BuildContext.AddRendererBinding(ScaleCurveRange.ResolvedParameter);
+			BuiltData->CurveScale = ScaleCurveRange.GetDefaultValue<FVector2f>();
+		}
+
+		BuildContext.AddParticleSimulationExecSimulate(&ParticleSimulate);
+	}
+}
+
+void UNiagaraStatelessModule_ScaleSpriteSize::BuildShaderParameters(FNiagaraStatelessShaderParametersBuilder& ShaderParametersBuilder) const
+{
+	using namespace NSMScaleSpriteSizePrivate;
+
+	ShaderParametersBuilder.AddParameterNestedStruct<FParameters>();
+}
+
+void UNiagaraStatelessModule_ScaleSpriteSize::SetShaderParameters(const FNiagaraStatelessSetShaderParameterContext& SetShaderParameterContext) const
+{
+	using namespace NSMScaleSpriteSizePrivate;
+
+	const FModuleBuiltData* ModuleBuiltData = SetShaderParameterContext.ReadBuiltData<FModuleBuiltData>();
+
+	FParameters* Parameters						= SetShaderParameterContext.GetParameterNestedStruct<FParameters>();
+	Parameters->ScaleSpriteSize_Distribution	= ModuleBuiltData->DistributionParameters;
+	Parameters->ScaleSpriteSize_CurveScale		= SetShaderParameterContext.GetRendererParameterValue(ModuleBuiltData->CurveScaleOffset, ModuleBuiltData->CurveScale);
+}
+
+#if WITH_EDITORONLY_DATA
+const TCHAR* UNiagaraStatelessModule_ScaleSpriteSize::GetShaderTemplatePath() const
+{
+	return TEXT("/Plugin/FX/Niagara/Private/Stateless/Modules/NiagaraStatelessModule_ScaleSpriteSize.ush");
+}
+
+void UNiagaraStatelessModule_ScaleSpriteSize::GetOutputVariables(TArray<FNiagaraVariableBase>& OutVariables, EVariableFilter Filter) const
+{
+	using namespace NSMScaleSpriteSizePrivate;
+
+	const FNiagaraStatelessGlobals& StatelessGlobals = FNiagaraStatelessGlobals::Get();
+	OutVariables.AddUnique(StatelessGlobals.SpriteSizeVariable);
+	OutVariables.AddUnique(StatelessGlobals.PreviousSpriteSizeVariable);
+}
+#endif
diff --git a/Engine/Plugins/FX/Niagara/Source/Niagara/Private/Stateless/Modules/NiagaraStatelessModule_SpriteFacingAndAlignment.cpp b/Engine/Plugins/FX/Niagara/Source/Niagara/Private/Stateless/Modules/NiagaraStatelessModule_SpriteFacingAndAlignment.cpp
new file mode 100644
--- /dev/null
+++ b/Engine/Plugins/FX/Niagara/Source/Niagara/Private/Stateless/Modules/NiagaraStatelessModule_SpriteFacingAndAlignment.cpp
@@ -0,0 +1,122 @@
+// Copyright Epic Games, Inc. All Rights Reserved.
+
+#include "Stateless/Modules/NiagaraStatelessModule_SpriteFacingAndAlignment.h"
+#include "Stateless/NiagaraStatelessModule.h"
+#include "Stateless/NiagaraStatelessEmitterDataBuildContext.h"
+#include "Stateless/NiagaraStatelessModuleShaderParameters.h"
+#include "Stateless/NiagaraStatelessParticleSimContext.h"
+
+namespace NSMSpriteFacingAndAlignmentPrivate
+{
+	struct FModuleBuiltData
+	{
+		FNiagaraStatelessRangeVector3	SpriteFacing	= FNiagaraStatelessRangeVector3(FVector3f::XAxisVector);
+		FNiagaraStatelessRangeVector3	SpriteAlignment = FNiagaraStatelessRangeVector3(FVector3f::YAxisVector);
+
+		int32		SpriteFacingVariableOffset = INDEX_NONE;
+		int32		PreviousSpriteFacingVariableOffset = INDEX_NONE;
+		int32		SpriteAlignmentVariable = INDEX_NONE;
+		int32		PreviousSpriteAlignmentVariable = INDEX_NONE;
+	};
+
+	using FParameters = NiagaraStateless::FSpriteFacingAndAlignmentModule_ShaderParameters;
+
+	static void ParticleSimulate(const NiagaraStateless::FParticleSimulationContext& ParticleSimulationContext)
+	{
+		using namespace NiagaraStateless;
+
+		const FModuleBuiltData* ModuleBuiltData = ParticleSimulationContext.ReadBuiltData<FModuleBuiltData>();
+		const FParameters* Parameters = ParticleSimulationContext.ReadParameterNestedStruct<FParameters>();
+
+		for (uint32 i = 0; i < ParticleSimulationContext.GetNumInstances(); ++i)
+		{
+			ParticleSimulationContext.WriteParticleVariable(ModuleBuiltData->SpriteFacingVariableOffset, i, Parameters->SpriteFacingAndAlignment_SpriteFacing);
+			ParticleSimulationContext.WriteParticleVariable(ModuleBuiltData->PreviousSpriteFacingVariableOffset, i, Parameters->SpriteFacingAndAlignment_SpriteFacing);
+
+			ParticleSimulationContext.WriteParticleVariable(ModuleBuiltData->SpriteAlignmentVariable, i, Parameters->SpriteFacingAndAlignment_SpriteAlignment);
+			ParticleSimulationContext.WriteParticleVariable(ModuleBuiltData->PreviousSpriteAlignmentVariable, i, Parameters->SpriteFacingAndAlignment_SpriteAlignment);
+		}
+	}
+}
+
+void UNiagaraStatelessModule_SpriteFacingAndAlignment::BuildEmitterData(const FNiagaraStatelessEmitterDataBuildContext& BuildContext) const
+{
+	using namespace NSMSpriteFacingAndAlignmentPrivate;
+
+	FModuleBuiltData* BuiltData = BuildContext.AllocateBuiltData<FModuleBuiltData>();
+	if (!IsModuleEnabled())
+	{
+		return;
+	}
+
+	const FNiagaraStatelessGlobals& StatelessGlobals = FNiagaraStatelessGlobals::Get();
+	if (bSpriteFacingEnabled)
+	{
+		BuiltData->SpriteFacingVariableOffset			= BuildContext.FindParticleVariableIndex(StatelessGlobals.SpriteFacingVariable);
+		BuiltData->PreviousSpriteFacingVariableOffset	= BuildContext.FindParticleVariableIndex(StatelessGlobals.PreviousSpriteFacingVariable);
+	}
+	if (bSpriteAlignmentEnabled)
+	{
+		BuiltData->SpriteAlignmentVariable				= BuildContext.FindParticleVariableIndex(StatelessGlobals.SpriteAlignmentVariable);
+		BuiltData->PreviousSpriteAlignmentVariable		= BuildContext.FindParticleVariableIndex(StatelessGlobals.PreviousSpriteAlignmentVariable);
+	}
+
+	if ((BuiltData->SpriteFacingVariableOffset == INDEX_NONE) && (BuiltData->PreviousSpriteFacingVariableOffset == INDEX_NONE) &&
+		(BuiltData->SpriteAlignmentVariable == INDEX_NONE) && (BuiltData->PreviousSpriteAlignmentVariable == INDEX_NONE))
+	{
+		return;
+	}
+
+	if (bSpriteFacingEnabled)
+	{
+		BuiltData->SpriteFacing = BuildContext.ConvertDistributionToRange(SpriteFacing, FVector3f::ZeroVector);
+	}
+	if (bSpriteAlignmentEnabled)
+	{
+		BuiltData->SpriteAlignment	= BuildContext.ConvertDistributionToRange(SpriteAlignment, FVector3f::ZeroVector);
+	}
+
+	BuildContext.AddParticleSimulationExecSimulate(&ParticleSimulate);
+}
+
+void UNiagaraStatelessModule_SpriteFacingAndAlignment::BuildShaderParameters(FNiagaraStatelessShaderParametersBuilder& ShaderParametersBuilder) const
+{
+	using namespace NSMSpriteFacingAndAlignmentPrivate;
+
+	ShaderParametersBuilder.AddParameterNestedStruct<FParameters>();
+}
+
+void UNiagaraStatelessModule_SpriteFacingAndAlignment::SetShaderParameters(const FNiagaraStatelessSetShaderParameterContext& SetShaderParameterContext) const
+{
+	using namespace NSMSpriteFacingAndAlignmentPrivate;
+
+	FParameters* Parameters = SetShaderParameterContext.GetParameterNestedStruct<FParameters>();
+	const FModuleBuiltData* ModuleBuiltData = SetShaderParameterContext.ReadBuiltData<FModuleBuiltData>();
+
+	Parameters->SpriteFacingAndAlignment_SpriteFacing		= SetShaderParameterContext.ConvertRangeToValue(ModuleBuiltData->SpriteFacing);
+	Parameters->SpriteFacingAndAlignment_SpriteAlignment	= SetShaderParameterContext.ConvertRangeToValue(ModuleBuiltData->SpriteAlignment);
+}
+
+#if WITH_EDITORONLY_DATA
+const TCHAR* UNiagaraStatelessModule_SpriteFacingAndAlignment::GetShaderTemplatePath() const
+{
+	return TEXT("/Plugin/FX/Niagara/Private/Stateless/Modules/NiagaraStatelessModule_SpriteFacingAndAlignment.ush");
+}
+
+void UNiagaraStatelessModule_SpriteFacingAndAlignment::GetOutputVariables(TArray<FNiagaraVariableBase>& OutVariables, EVariableFilter Filter) const
+{
+	using namespace NSMSpriteFacingAndAlignmentPrivate;
+
+	const FNiagaraStatelessGlobals& StatelessGlobals = FNiagaraStatelessGlobals::Get();
+	if (bSpriteFacingEnabled || Filter == EVariableFilter::None)
+	{
+		OutVariables.AddUnique(StatelessGlobals.SpriteFacingVariable);
+		OutVariables.AddUnique(StatelessGlobals.PreviousSpriteFacingVariable);
+	}
+	if (bSpriteAlignmentEnabled || Filter == EVariableFilter::None)
+	{
+		OutVariables.AddUnique(StatelessGlobals.SpriteAlignmentVariable);
+		OutVariables.AddUnique(StatelessGlobals.PreviousSpriteAlignmentVariable);
+	}
+}
+#endif
diff --git a/Engine/Plugins/FX/Niagara/Source/Niagara/Private/Stateless/Modules/NiagaraStatelessModule_SpriteRotationRate.cpp b/Engine/Plugins/FX/Niagara/Source/Niagara/Private/Stateless/Modules/NiagaraStatelessModule_SpriteRotationRate.cpp
new file mode 100644
--- /dev/null
+++ b/Engine/Plugins/FX/Niagara/Source/Niagara/Private/Stateless/Modules/NiagaraStatelessModule_SpriteRotationRate.cpp
@@ -0,0 +1,116 @@
+// Copyright Epic Games, Inc. All Rights Reserved.
+
+#include "Stateless/Modules/NiagaraStatelessModule_SpriteRotationRate.h"
+#include "Stateless/NiagaraStatelessEmitterDataBuildContext.h"
+#include "Stateless/NiagaraStatelessModuleShaderParameters.h"
+#include "Stateless/NiagaraStatelessParticleSimContext.h"
+
+namespace NSMSpriteRotationRatePrivate
+{
+	struct FModuleBuiltData
+	{
+		int32							ModuleEnabled = false;
+		FNiagaraStatelessRangeFloat		RotationRange;
+		FUintVector3					RateScaleParameters = FUintVector3::ZeroValue;
+		int32							SpriteRotationVariableOffset = INDEX_NONE;
+		int32							PreviousSpriteRotationVariableOffset = INDEX_NONE;
+	};
+
+	using FParameters = NiagaraStateless::FSpriteRotationRateModule_ShaderParameters;
+
+	static void ParticleSimulate(const NiagaraStateless::FParticleSimulationContext& ParticleSimulationContext)
+	{
+		using namespace NiagaraStateless;
+
+		const FModuleBuiltData* ModuleBuiltData = ParticleSimulationContext.ReadBuiltData<FModuleBuiltData>();
+		const FParameters* ShaderParameters = ParticleSimulationContext.ReadParameterNestedStruct<FParameters>();
+
+		const float* LifetimeData = ParticleSimulationContext.GetParticleLifetime();
+		const float* AgeData = ParticleSimulationContext.GetParticleNormalizedAge();
+		const float* PreviousAgeData = ParticleSimulationContext.GetParticlePreviousNormalizedAge();
+
+		for (uint32 i = 0; i < ParticleSimulationContext.GetNumInstances(); ++i)
+		{
+			const float RotationRate = ParticleSimulationContext.RandomScaleBiasFloat(i, 0, ShaderParameters->SpriteRotationRate_Scale, ShaderParameters->SpriteRotationRate_Bias);
+			const float MultiplyRate = ParticleSimulationContext.SampleCurve<float>(ShaderParameters->SpriteRotationRate_RateScaleParameters, AgeData[i]);
+			const float PreviousMultiplyRate = ParticleSimulationContext.SampleCurve<float>(ShaderParameters->SpriteRotationRate_RateScaleParameters, PreviousAgeData[i]);
+
+			float SpriteRotation = ParticleSimulationContext.ReadParticleVariable(ModuleBuiltData->SpriteRotationVariableOffset, i, 0.0f);
+			float PreviousSpriteRotation = ParticleSimulationContext.ReadParticleVariable(ModuleBuiltData->PreviousSpriteRotationVariableOffset, i, 0.0f);
+
+			SpriteRotation += LifetimeData[i] * RotationRate * MultiplyRate;
+			PreviousSpriteRotation += LifetimeData[i] * RotationRate * PreviousMultiplyRate;
+
+			ParticleSimulationContext.WriteParticleVariable(ModuleBuiltData->SpriteRotationVariableOffset, i, SpriteRotation);
+			ParticleSimulationContext.WriteParticleVariable(ModuleBuiltData->PreviousSpriteRotationVariableOffset, i, PreviousSpriteRotation);
+		}
+	}
+}
+
+void UNiagaraStatelessModule_SpriteRotationRate::BuildEmitterData(const FNiagaraStatelessEmitterDataBuildContext& BuildContext) const
+{
+	using namespace NSMSpriteRotationRatePrivate;
+
+	FModuleBuiltData* BuiltData = BuildContext.AllocateBuiltData<FModuleBuiltData>();
+	if (!IsModuleEnabled())
+	{
+		return;
+	}
+
+	const FNiagaraStatelessGlobals& StatelessGlobals = FNiagaraStatelessGlobals::Get();
+	BuiltData->SpriteRotationVariableOffset			= BuildContext.FindParticleVariableIndex(StatelessGlobals.SpriteRotationVariable);
+	BuiltData->PreviousSpriteRotationVariableOffset	= BuildContext.FindParticleVariableIndex(StatelessGlobals.PreviousSpriteRotationVariable);
+
+	if (BuiltData->SpriteRotationVariableOffset == INDEX_NONE && BuiltData->PreviousSpriteRotationVariableOffset == INDEX_NONE)
+	{
+		return;
+	}
+
+	BuiltData->ModuleEnabled = true;
+	BuiltData->RotationRange = BuildContext.ConvertDistributionToRange(RotationRateDistribution, 0.0f);
+	if (bUseRateScale)
+	{
+		BuiltData->RateScaleParameters = BuildContext.AddDistributionAsCurve(RateScaleDistribution, 1.0f);
+	}
+	else
+	{
+		BuiltData->RateScaleParameters = BuildContext.AddDistributionAsCurve(FNiagaraDistributionCurveFloat(ENiagaraDistributionCurveLUTMode::Accumulate), 1.0f);
+	}
+
+	BuildContext.AddParticleSimulationExecSimulate(&ParticleSimulate);
+}
+
+void UNiagaraStatelessModule_SpriteRotationRate::BuildShaderParameters(FNiagaraStatelessShaderParametersBuilder& ShaderParametersBuilder) const
+{
+	using namespace NSMSpriteRotationRatePrivate;
+
+	ShaderParametersBuilder.AddParameterNestedStruct<FParameters>();
+}
+
+void UNiagaraStatelessModule_SpriteRotationRate::SetShaderParameters(const FNiagaraStatelessSetShaderParameterContext& SetShaderParameterContext) const
+{
+	using namespace NSMSpriteRotationRatePrivate;
+
+	FParameters* Parameters = SetShaderParameterContext.GetParameterNestedStruct<FParameters>();
+	const FModuleBuiltData* ModuleBuiltData = SetShaderParameterContext.ReadBuiltData<FModuleBuiltData>();
+
+	Parameters->SpriteRotationRate_Enabled = ModuleBuiltData->ModuleEnabled;
+	SetShaderParameterContext.ConvertRangeToScaleBias(ModuleBuiltData->RotationRange, Parameters->SpriteRotationRate_Scale, Parameters->SpriteRotationRate_Bias);
+	Parameters->SpriteRotationRate_RateScaleParameters = ModuleBuiltData->RateScaleParameters;
+}
+
+#if WITH_EDITORONLY_DATA
+const TCHAR* UNiagaraStatelessModule_SpriteRotationRate::GetShaderTemplatePath() const
+{
+	return TEXT("/Plugin/FX/Niagara/Private/Stateless/Modules/NiagaraStatelessModule_SpriteRotationRate.ush");
+}
+
+void UNiagaraStatelessModule_SpriteRotationRate::GetOutputVariables(TArray<FNiagaraVariableBase>& OutVariables, EVariableFilter Filter) const
+{
+	using namespace NSMSpriteRotationRatePrivate;
+
+	const FNiagaraStatelessGlobals& StatelessGlobals = FNiagaraStatelessGlobals::Get();
+	OutVariables.AddUnique(StatelessGlobals.SpriteRotationVariable);
+	OutVariables.AddUnique(StatelessGlobals.PreviousSpriteRotationVariable);
+}
+#endif
diff --git a/Engine/Plugins/FX/Niagara/Source/NiagaraEditor/Private/Stateless/NiagaraStatelessEmitterTranslator.cpp b/Engine/Plugins/FX/Niagara/Source/NiagaraEditor/Private/Stateless/NiagaraStatelessEmitterTranslator.cpp
--- a/Engine/Plugins/FX/Niagara/Source/NiagaraEditor/Private/Stateless/NiagaraStatelessEmitterTranslator.cpp
+++ b/Engine/Plugins/FX/Niagara/Source/NiagaraEditor/Private/Stateless/NiagaraStatelessEmitterTranslator.cpp
@@ -54,7 +54,7 @@ namespace FNiagaraStatelessEmitterTranslator
 		{
 			return TEXT("float3");
 		}
-		else if (TypeDef == FNiagaraTypeDefinition::GetVec4Def() || TypeDef == FNiagaraTypeDefinition::GetColorDef() || TypeDef == FNiagaraTypeDefinition::GetColorDef() || TypeDef == FNiagaraTypeDefinition::GetQuatDef())
+		else if (TypeDef == FNiagaraTypeDefinition::GetVec4Def() || TypeDef == FNiagaraTypeDefinition::GetColorDef() || TypeDef == FNiagaraTypeDefinition::GetQuatDef())
 		{
 			return TEXT("float4");
 		}
@@ -115,7 +115,7 @@ namespace FNiagaraStatelessEmitterTranslator
 		HlslOutput.Append(TEXT("\tfloat Lifetime;					// Overall lifetime for the particle\n"));
 		HlslOutput.Append(TEXT("\tfloat Age;						// Current age for the particle\n"));
 		HlslOutput.Append(TEXT("\tfloat NormalizedAge;			// Current normalized age for the particle\n"));
-		HlslOutput.Append(TEXT("\tloat PreviousAge;				// Previous age for the particle\n"));
+		HlslOutput.Append(TEXT("\tfloat PreviousAge;				// Previous age for the particle\n"));
 		HlslOutput.Append(TEXT("\tfloat PreviousNormalizedAge;	// Previous normalized age for the particle\n"));
 		HlslOutput.Append(TEXT("\tfloat DeltaTime;				// Simulation DT\n"));
 		HlslOutput.Append(TEXT("\tfloat InvDeltaTime;				// Simulation Inverse DT\n"));
