Fixed user parameter tooltip for verse identifiers

diff --git a/Engine/Plugins/FX/Niagara/Source/NiagaraEditor/Private/Customizations/UserParameters/NiagaraSystemUserParametersBuilder.cpp b/Engine/Plugins/FX/Niagara/Source/NiagaraEditor/Private/Customizations/UserParameters/NiagaraSystemUserParametersBuilder.cpp
--- a/Engine/Plugins/FX/Niagara/Source/NiagaraEditor/Private/Customizations/UserParameters/NiagaraSystemUserParametersBuilder.cpp
+++ b/Engine/Plugins/FX/Niagara/Source/NiagaraEditor/Private/Customizations/UserParameters/NiagaraSystemUserParametersBuilder.cpp
@@ -746,6 +746,7 @@ FText FNiagaraSystemUserParameterBuilder::GetParameterAccessLevelTooltip(FNiagar
 {
 	if (VerseReflectionUtilities)
 	{
+		UserParameter.RemoveRootNamespace(FNiagaraConstants::UserNamespaceString);
 		TOptional<FText> ErrorText = VerseReflectionUtilities->IsValidVariable(UserParameter);
 		if (ErrorText.IsSet())
 		{
@@ -754,7 +755,7 @@ FText FNiagaraSystemUserParameterBuilder::GetParameterAccessLevelTooltip(FNiagar
 	}
 
 	const UEnum* Enum = StaticEnum<ENiagaraParameterAccessLevel>();
-	return Enum ? Enum->GetToolTipTextByIndex((int32)GetParameterAccessLevel(UserParameter)) : FText();
+	return Enum ? Enum->GetToolTipTextByIndex(static_cast<int32>(GetParameterAccessLevel(UserParameter))) : FText();
 }
 
 bool FNiagaraSystemUserParameterBuilder::CanExposeParameterToVerse(FNiagaraVariable UserParameter) const
