Niagara component renderer: only register one object replacement delegate with the CDO to prevent races during asset loading

#tests Created a scene component BP, used it in a component renderer and recompiled the BP after adding a new property. The component renderer picked it up after the change.

diff --git a/Engine/Plugins/FX/Niagara/Source/Niagara/Private/NiagaraComponentRendererProperties.cpp b/Engine/Plugins/FX/Niagara/Source/Niagara/Private/NiagaraComponentRendererProperties.cpp
--- a/Engine/Plugins/FX/Niagara/Source/Niagara/Private/NiagaraComponentRendererProperties.cpp
+++ b/Engine/Plugins/FX/Niagara/Source/Niagara/Private/NiagaraComponentRendererProperties.cpp
@@ -15,6 +15,7 @@
 #include "Styling/SlateIconFinder.h"
 #include "Widgets/SWidget.h"
 #include "Widgets/Text/STextBlock.h"
+#include "UObject/UObjectIterator.h"
 #endif
 #include "NiagaraCustomVersion.h"
 #include "NiagaraSettings.h"
@@ -176,29 +177,37 @@ FNiagaraTypeDefinition UNiagaraComponentRendererProperties::GetDoubleDef()
 TArray<TWeakObjectPtr<UNiagaraComponentRendererProperties>> UNiagaraComponentRendererProperties::ComponentRendererPropertiesToDeferredInit;
 
 UNiagaraComponentRendererProperties::UNiagaraComponentRendererProperties()
-	: ComponentCountLimit(15), bAssignComponentsOnParticleID(true), bOnlyActivateNewlyAquiredComponents(true)
+	: ComponentCountLimit(15)
+	, bAssignComponentsOnParticleID(true)
+	, bOnlyActivateNewlyAquiredComponents(true)
 #if WITH_EDITORONLY_DATA
 	, bVisualizeComponents(true)
 	, bOnlyCreateComponentsOnParticleSpawn_DEPRECATED(true)
 #endif
+	, RendererVisibility(0)
 	, TemplateComponent(nullptr)
 {
-#if WITH_EDITOR
-	FCoreUObjectDelegates::OnObjectsReplaced.AddUObject(this, &UNiagaraComponentRendererProperties::OnObjectsReplacedCallback);
-#endif
-
 	AttributeBindings =
 	{
 		&EnabledBinding,
 		&RendererVisibilityTagBinding,
 	};
 	IsSetterMappingDirty.store(true);
+#if WITH_EDITOR
+	if (HasAnyFlags(RF_ClassDefaultObject))
+	{
+		FCoreUObjectDelegates::OnObjectsReplaced.AddUObject(this, &UNiagaraComponentRendererProperties::OnObjectsReplacedCallback);
+	}
+#endif
 }
 
 UNiagaraComponentRendererProperties::~UNiagaraComponentRendererProperties()
 {
 #if WITH_EDITOR
-	FCoreUObjectDelegates::OnObjectsReplaced.RemoveAll(this);
+	if (HasAnyFlags(RF_ClassDefaultObject))
+	{
+		FCoreUObjectDelegates::OnObjectsReplaced.RemoveAll(this);
+	}
 #endif
 }
 
@@ -516,6 +525,19 @@ void UNiagaraComponentRendererProperties::UpdateTemplateCoordinateSpace()
 #if WITH_EDITOR
 
 void UNiagaraComponentRendererProperties::OnObjectsReplacedCallback(const TMap<UObject*, UObject*>& ReplacementsMap)
+{
+	// this callback is only registered on the CDO, so we inform all the regular instances like this
+	for (TObjectIterator<UNiagaraComponentRendererProperties> Itt; Itt; ++Itt)
+	{
+		UNiagaraComponentRendererProperties* PropertiesObj = *Itt;
+		if (PropertiesObj && PropertiesObj != this)
+		{
+			PropertiesObj->OnObjectsReplacedInternal(ReplacementsMap);
+		}
+	}
+}
+
+void UNiagaraComponentRendererProperties::OnObjectsReplacedInternal(const TMap<UObject*, UObject*>& ReplacementsMap)
 {
 	// When a custom component class is recompiled in the editor, we need to switch to the new template component object
 	if (TemplateComponent)
diff --git a/Engine/Plugins/FX/Niagara/Source/Niagara/Public/NiagaraComponentRendererProperties.h b/Engine/Plugins/FX/Niagara/Source/Niagara/Public/NiagaraComponentRendererProperties.h
--- a/Engine/Plugins/FX/Niagara/Source/Niagara/Public/NiagaraComponentRendererProperties.h
+++ b/Engine/Plugins/FX/Niagara/Source/Niagara/Public/NiagaraComponentRendererProperties.h
@@ -175,6 +175,7 @@ private:
 #if WITH_EDITOR
 	/** Callback for whenever any blueprint components are reinstanced */
 	void OnObjectsReplacedCallback(const TMap<UObject*, UObject*>& ReplacementsMap);
+	void OnObjectsReplacedInternal(const TMap<UObject*, UObject*>& ReplacementsMap);
 
 #endif
 
