Fixed a use-after-free issue in NiagaraScalabilityManager that occurs when WITH_NIAGARA_DEBUGGER is enabled. The issue was caused by a reference to FNiagaraScalabilityState being retained and accessed after the corresponding entry was removed from the FNiagaraScalabilityState::State array.

diff --git a/Engine/Plugins/FX/Niagara/Source/Niagara/Private/NiagaraScalabilityManager.cpp b/Engine/Plugins/FX/Niagara/Source/Niagara/Private/NiagaraScalabilityManager.cpp
--- a/Engine/Plugins/FX/Niagara/Source/Niagara/Private/NiagaraScalabilityManager.cpp
+++ b/Engine/Plugins/FX/Niagara/Source/Niagara/Private/NiagaraScalabilityManager.cpp
@@ -417,7 +417,7 @@ bool FNiagaraScalabilityManager::ApplyScalabilityState(int32 ComponentIndex, ENi
 			{
 				bContinueIteration = false;
 #if WITH_NIAGARA_DEBUGGER
-				WorldMan->GetSystemScalabilityKillCounts(System).Add(CompState);
+				WorldMan->GetSystemScalabilityKillCounts(System).Add(Component->DebugCachedScalabilityState);
 #endif
 			}
 		}
