More loc text identifier copy/pasta
#rnx

diff --git a/Engine/Plugins/FX/Niagara/Source/NiagaraEditor/Private/NiagaraValidationRules.cpp b/Engine/Plugins/FX/Niagara/Source/NiagaraEditor/Private/NiagaraValidationRules.cpp
--- a/Engine/Plugins/FX/Niagara/Source/NiagaraEditor/Private/NiagaraValidationRules.cpp
+++ b/Engine/Plugins/FX/Niagara/Source/NiagaraEditor/Private/NiagaraValidationRules.cpp
@@ -1715,7 +1715,7 @@ void UNiagaraValidationRule_MaterialUsage::CheckValidity(const FNiagaraValidatio
 
 				OutResults.Emplace(
 					FailedUsageSeverity,
-					LOCTEXT("MissingMaterialUsage", "Material usage flags are not set correctly."),
+					LOCTEXT("MissingMaterialUsage_Description", "Material usage flags are not set correctly."),
 					FText::Format(LOCTEXT("MissingMaterialUsage", "Materials on renderer {0} are missing the usage flag required to render with."), FText::FromName(RendererProperties->GetFName())),
 					NiagaraValidation::GetRendererStackItem(EmitterHandleModel.Get().GetEmitterStackViewModel(), RendererProperties)
 				);
