- Remove stateless emitters from compilation
- Previously they would still generate some HLSL data resulting in compilation when added / removed
- Cleaned up several areas that assumed Emitter Data would always be valid

diff --git a/Engine/Plugins/FX/Niagara/Source/Niagara/Classes/NiagaraEmitterHandle.h b/Engine/Plugins/FX/Niagara/Source/Niagara/Classes/NiagaraEmitterHandle.h
--- a/Engine/Plugins/FX/Niagara/Source/Niagara/Classes/NiagaraEmitterHandle.h
+++ b/Engine/Plugins/FX/Niagara/Source/Niagara/Classes/NiagaraEmitterHandle.h
@@ -78,8 +78,8 @@ public:
 #endif
 
 	/** Gets the copied instance of the emitter this handle references. */
-	NIAGARA_API FVersionedNiagaraEmitter GetInstance() const;
-	NIAGARA_API FVersionedNiagaraEmitter& GetInstance();
+	[[nodiscard]] NIAGARA_API const FVersionedNiagaraEmitter GetInstance() const;
+	NIAGARA_API void SetInstanceVersion(const FGuid& InVersion);
 
 	NIAGARA_API FVersionedNiagaraEmitterData* GetEmitterData() const;
 
diff --git a/Engine/Plugins/FX/Niagara/Source/Niagara/Private/NiagaraEmitterHandle.cpp b/Engine/Plugins/FX/Niagara/Source/Niagara/Private/NiagaraEmitterHandle.cpp
--- a/Engine/Plugins/FX/Niagara/Source/Niagara/Private/NiagaraEmitterHandle.cpp
+++ b/Engine/Plugins/FX/Niagara/Source/Niagara/Private/NiagaraEmitterHandle.cpp
@@ -226,19 +226,19 @@ void FNiagaraEmitterHandle::SetEmitterMode(UNiagaraSystem& InOwningSystem, ENiag
 }
 #endif
 
-FVersionedNiagaraEmitter FNiagaraEmitterHandle::GetInstance() const
+const FVersionedNiagaraEmitter FNiagaraEmitterHandle::GetInstance() const
 {
 	return VersionedInstance;
 }
 
-FVersionedNiagaraEmitter& FNiagaraEmitterHandle::GetInstance()
+void FNiagaraEmitterHandle::SetInstanceVersion(const FGuid& InVersion)
 {
-	return VersionedInstance;
+	VersionedInstance.Version = InVersion;
 }
 
 FVersionedNiagaraEmitterData* FNiagaraEmitterHandle::GetEmitterData() const
 {
-	return VersionedInstance.GetEmitterData();
+	return EmitterMode == ENiagaraEmitterMode::Standard	? VersionedInstance.GetEmitterData() : nullptr;
 }
 
 FString FNiagaraEmitterHandle::GetUniqueInstanceName() const
diff --git a/Engine/Plugins/FX/Niagara/Source/Niagara/Private/NiagaraSystem.cpp b/Engine/Plugins/FX/Niagara/Source/Niagara/Private/NiagaraSystem.cpp
--- a/Engine/Plugins/FX/Niagara/Source/Niagara/Private/NiagaraSystem.cpp
+++ b/Engine/Plugins/FX/Niagara/Source/Niagara/Private/NiagaraSystem.cpp
@@ -536,7 +536,7 @@ bool UNiagaraSystem::ChangeEmitterVersion(const FVersionedNiagaraEmitter& Versio
 		{
 			KillAllActiveCompilations();
 			Modify();
-			EmitterHandles[i].GetInstance().Version = NewVersion;
+			EmitterHandles[i].SetInstanceVersion(NewVersion);
 			return true;
 		}
 	}
@@ -3402,9 +3402,9 @@ bool UNiagaraSystem::QueryCompileComplete(bool bWait)
 				{
 					Handle.GetInstance().Emitter->OnPostCompile(Handle.GetInstance().Version);
 				}
-				else
+				else if (FVersionedNiagaraEmitterData* EmitterData = Handle.GetEmitterData())
 				{
-					Handle.GetEmitterData()->InvalidateCompileResults();
+					EmitterData->InvalidateCompileResults();
 				}
 			}
 		}
@@ -3732,11 +3732,14 @@ void UNiagaraSystem::InitEmitterCompiledData()
 		for (int32 EmitterIdx = 0; EmitterIdx < EmitterHandles.Num(); ++EmitterIdx)
 		{
 			const FNiagaraEmitterHandle& EmitterHandle = EmitterHandles[EmitterIdx];
+			if (!EmitterHandle.GetIsEnabled() || EmitterHandle.GetEmitterMode() == ENiagaraEmitterMode::Stateless)
+			{
+				continue;
+			}
+
 			const UNiagaraEmitter* Emitter = EmitterHandle.GetInstance().Emitter;
 			FNiagaraDataSetCompiledData& EmitterDataSetCompiledData = NewEmitterCompiledData[EmitterIdx]->DataSetCompiledData;
-			//-TODO:Stateless: Do we need stateless support here?
-			//if (ensureMsgf(Emitter != nullptr, TEXT("Failed to get Emitter Instance from Emitter Handle in post compile, please investigate.")))
-			if (Emitter != nullptr)
+			if (ensureMsgf(Emitter != nullptr, TEXT("Failed to get Emitter Instance from Emitter Handle in post compile, please investigate.")))
 			{
 				InitEmitterVariableAliasNames(NewEmitterCompiledData[EmitterIdx].Get(), Emitter);
 				InitEmitterDataSetCompiledData(EmitterDataSetCompiledData, EmitterHandle);			
diff --git a/Engine/Plugins/FX/Niagara/Source/Niagara/Public/NiagaraTypes.h b/Engine/Plugins/FX/Niagara/Source/Niagara/Public/NiagaraTypes.h
--- a/Engine/Plugins/FX/Niagara/Source/Niagara/Public/NiagaraTypes.h
+++ b/Engine/Plugins/FX/Niagara/Source/Niagara/Public/NiagaraTypes.h
@@ -1779,8 +1779,8 @@ struct FVersionedNiagaraEmitter
 		: Emitter(InEmitter), Version(InVersion)
 	{};
 
-	NIAGARA_API FVersionedNiagaraEmitterData* GetEmitterData() const;
-	NIAGARA_API FVersionedNiagaraEmitterWeakPtr ToWeakPtr() const;
+	[[nodiscard]] NIAGARA_API FVersionedNiagaraEmitterData* GetEmitterData() const;
+	[[nodiscard]] NIAGARA_API FVersionedNiagaraEmitterWeakPtr ToWeakPtr() const;
 
 	UPROPERTY()
 	TObjectPtr<UNiagaraEmitter> Emitter = nullptr;
@@ -1808,15 +1808,16 @@ struct FVersionedNiagaraEmitterWeakPtr
 {
 	FVersionedNiagaraEmitterWeakPtr() : Emitter(nullptr) {}
 	NIAGARA_API FVersionedNiagaraEmitterWeakPtr(UNiagaraEmitter* InEmitter, const FGuid& InVersion);
-	NIAGARA_API FVersionedNiagaraEmitter ResolveWeakPtr() const;
-	NIAGARA_API FVersionedNiagaraEmitterData* GetEmitterData() const;
-	NIAGARA_API bool IsValid() const;
+
+	[[nodiscard]] NIAGARA_API bool IsValid() const;
+	[[nodiscard]] NIAGARA_API FVersionedNiagaraEmitter ResolveWeakPtr() const;
+	[[nodiscard]] NIAGARA_API FVersionedNiagaraEmitterData* GetEmitterData() const;
+
 	bool operator==(const FVersionedNiagaraEmitterWeakPtr& Other) const
 	{
 		return Emitter == Other.Emitter && Version == Other.Version;
 	}
 
-
 	bool operator!=(const FVersionedNiagaraEmitterWeakPtr& Other) const
 	{
 		return !(*this == Other);
@@ -1826,9 +1827,9 @@ struct FVersionedNiagaraEmitterWeakPtr
 	FGuid Version;
 
 	friend inline uint32 GetTypeHash(const FVersionedNiagaraEmitterWeakPtr& Item)
-{
-	return GetTypeHash(Item.Emitter);
-}
+	{
+		return GetTypeHash(Item.Emitter);
+	}
 };
 
 //////////////////////////////////////////////////////////////////////////
diff --git a/Engine/Plugins/FX/Niagara/Source/NiagaraEditor/Private/NiagaraCompilationTasks.h b/Engine/Plugins/FX/Niagara/Source/NiagaraEditor/Private/NiagaraCompilationTasks.h
--- a/Engine/Plugins/FX/Niagara/Source/NiagaraEditor/Private/NiagaraCompilationTasks.h
+++ b/Engine/Plugins/FX/Niagara/Source/NiagaraEditor/Private/NiagaraCompilationTasks.h
@@ -89,11 +89,11 @@ struct FNiagaraSystemCompilationTask : public TSharedFromThis<FNiagaraSystemComp
 		TArray<TObjectKey<UNiagaraScript>> OwnedScriptKeys;
 
 		// Index into the EmitterHandles[] of the owning System
-		int32 SourceEmitterIndex;
+		int32 SourceEmitterIndex = INDEX_NONE;
 
 		// Index of the this in FSystemInfo::EmitterInfo
-		int32 DigestedEmitterIndex;
-		bool Enabled;
+		int32 DigestedEmitterIndex = INDEX_NONE;
+		bool Enabled = false;
 	};
 
 	struct FSystemInfo
@@ -109,8 +109,8 @@ struct FNiagaraSystemCompilationTask : public TSharedFromThis<FNiagaraSystemComp
 		FNiagaraFixedConstantResolver ConstantResolver;
 		TArray<TObjectKey<UNiagaraScript>> OwnedScriptKeys;
 
-		bool bUseRapidIterationParams;
-		bool bDisableDebugSwitches;
+		bool bUseRapidIterationParams = false;
+		bool bDisableDebugSwitches = false;
 
 		const FEmitterInfo* EmitterInfoBySourceEmitter(int32 InSourceEmitterIndex) const;
 		FEmitterInfo* EmitterInfoBySourceEmitter(int32 InSourceEmitterIndex);
diff --git a/Engine/Plugins/FX/Niagara/Source/NiagaraEditor/Private/NiagaraScriptSource.cpp b/Engine/Plugins/FX/Niagara/Source/NiagaraEditor/Private/NiagaraScriptSource.cpp
--- a/Engine/Plugins/FX/Niagara/Source/NiagaraEditor/Private/NiagaraScriptSource.cpp
+++ b/Engine/Plugins/FX/Niagara/Source/NiagaraEditor/Private/NiagaraScriptSource.cpp
@@ -346,9 +346,10 @@ TMap<FName, UNiagaraDataInterface*> UNiagaraScriptSource::ComputeObjectNameMap(U
 
 			for (const FNiagaraEmitterHandle& EmitterHandle : System.GetEmitterHandles())
 			{
-				if (EmitterHandle.GetIsEnabled() && EmitterHandle.GetInstance().Emitter != nullptr)
+				FVersionedNiagaraEmitterData* EmitterData = EmitterHandle.GetIsEnabled() ? EmitterHandle.GetEmitterData() : nullptr;
+				if (EmitterData)
 				{
-					if (UNiagaraScriptSource* EmitterSource = Cast<UNiagaraScriptSource>(EmitterHandle.GetEmitterData()->GraphSource); EmitterSource != nullptr)
+					if (UNiagaraScriptSource* EmitterSource = Cast<UNiagaraScriptSource>(EmitterData->GraphSource); EmitterSource != nullptr)
 					{
 						FindObjectNamesRecursive(EmitterSource->NodeGraph, EmitterUsage, FGuid(), EmitterHandle.GetUniqueInstanceName(), VisitedGraphs, Result);
 					}
diff --git a/Engine/Plugins/FX/Niagara/Source/NiagaraEditor/Private/ViewModels/NiagaraSystemGraphSelectionViewModel.cpp b/Engine/Plugins/FX/Niagara/Source/NiagaraEditor/Private/ViewModels/NiagaraSystemGraphSelectionViewModel.cpp
--- a/Engine/Plugins/FX/Niagara/Source/NiagaraEditor/Private/ViewModels/NiagaraSystemGraphSelectionViewModel.cpp
+++ b/Engine/Plugins/FX/Niagara/Source/NiagaraEditor/Private/ViewModels/NiagaraSystemGraphSelectionViewModel.cpp
@@ -38,7 +38,10 @@ void FNiagaraSystemGraphSelectionViewModel::RefreshSelectedEmitterScriptGraphs()
 		{
 			if (SelectedEmitterHandleIds.Contains(EmitterHandleViewModel->GetId()))
 			{
-				SelectedEmitterScriptGraphs.Add(static_cast<UNiagaraScriptSource*>(EmitterHandleViewModel->GetEmitterHandle()->GetEmitterData()->GraphSource)->NodeGraph);
+				if (FVersionedNiagaraEmitterData* EmitterData = EmitterHandleViewModel->GetEmitterHandle()->GetEmitterData())
+				{
+					SelectedEmitterScriptGraphs.Add(static_cast<UNiagaraScriptSource*>(EmitterData->GraphSource)->NodeGraph);
+				}
 			}
 		}
 	}
diff --git a/Engine/Plugins/FX/Niagara/Source/NiagaraEditor/Private/Widgets/SNiagaraGeneratedCodeView.cpp b/Engine/Plugins/FX/Niagara/Source/NiagaraEditor/Private/Widgets/SNiagaraGeneratedCodeView.cpp
--- a/Engine/Plugins/FX/Niagara/Source/NiagaraEditor/Private/Widgets/SNiagaraGeneratedCodeView.cpp
+++ b/Engine/Plugins/FX/Niagara/Source/NiagaraEditor/Private/Widgets/SNiagaraGeneratedCodeView.cpp
@@ -375,11 +375,13 @@ void SNiagaraGeneratedCodeView::UpdateUI()
 	if(SelectedEmitterHandleIds.Num() == 1)
 	{
 		TSharedPtr<FNiagaraEmitterHandleViewModel> SelectedEmitterHandleViewModel = SystemViewModel->GetEmitterHandleViewModelById(SelectedEmitterHandleIds[0]);
+
 		FNiagaraEmitterHandle* Handle = SelectedEmitterHandleViewModel.IsValid() ? SelectedEmitterHandleViewModel->GetEmitterHandle() : nullptr;
-		if (Handle)
+		const FVersionedNiagaraEmitterData* EmitterData = Handle ? Handle->GetEmitterData() : nullptr;
+		if (EmitterData)
 		{
 			TArray<UNiagaraScript*> EmitterScripts;
-			Handle->GetEmitterData()->GetScripts(EmitterScripts);
+			EmitterData->GetScripts(EmitterScripts);
 			Scripts.Append(EmitterScripts);
 		}
 	}
