- Fixes to rules that either ran on disabled emitters, or did not check that the emitter data was valid

#rnx

diff --git a/Engine/Plugins/FX/Niagara/Source/NiagaraEditor/Private/NiagaraValidationRules.cpp b/Engine/Plugins/FX/Niagara/Source/NiagaraEditor/Private/NiagaraValidationRules.cpp
--- a/Engine/Plugins/FX/Niagara/Source/NiagaraEditor/Private/NiagaraValidationRules.cpp
+++ b/Engine/Plugins/FX/Niagara/Source/NiagaraEditor/Private/NiagaraValidationRules.cpp
@@ -388,7 +388,7 @@ void UNiagaraValidationRule_NoEvents::CheckValidity(const FNiagaraValidationCont
 		}
 
 		FVersionedNiagaraEmitterData* EmitterData = EmitterHandleModel.Get().GetEmitterHandle()->GetEmitterData();
-		if (EmitterData->GetEventHandlers().Num() == 0)
+		if (!EmitterData || EmitterData->GetEventHandlers().Num() == 0)
 		{
 			continue;
 		}
@@ -421,8 +421,13 @@ void UNiagaraValidationRule_FixedGPUBoundsSet::CheckValidity(const FNiagaraValid
 	TArray<TSharedRef<FNiagaraEmitterHandleViewModel>> EmitterHandleViewModels = Context.ViewModel->GetEmitterHandleViewModels();
 	for (TSharedRef<FNiagaraEmitterHandleViewModel> EmitterHandleModel : EmitterHandleViewModels)
 	{
+		if (EmitterHandleModel->GetIsEnabled() == false)
+		{
+			continue;
+		}
+
 		FVersionedNiagaraEmitterData* EmitterData = EmitterHandleModel.Get().GetEmitterHandle()->GetEmitterData();
-		if (EmitterData->SimTarget == ENiagaraSimTarget::GPUComputeSim && EmitterData->CalculateBoundsMode == ENiagaraEmitterCalculateBoundMode::Dynamic)
+		if (EmitterData && EmitterData->SimTarget == ENiagaraSimTarget::GPUComputeSim && EmitterData->CalculateBoundsMode == ENiagaraEmitterCalculateBoundMode::Dynamic)
 		{
 			UNiagaraStackEmitterPropertiesItem* EmitterProperties = NiagaraValidation::GetStackEntry<UNiagaraStackEmitterPropertiesItem>(EmitterHandleModel.Get().GetEmitterStackViewModel());
 			FNiagaraValidationResult Result(ENiagaraValidationSeverity::Error, LOCTEXT("GpuDynamicBoundsErrorSummary", "GPU emitters do not support dynamic bounds"), LOCTEXT("GpuDynamicBoundsErrorDescription", "Gpu emitter should either not be in dynamic mode or the system must have fixed bounds."), EmitterProperties);
@@ -542,6 +547,10 @@ void UNiagaraValidationRule_RendererCount::CheckValidity(const FNiagaraValidatio
 			continue;
 		}
 		FVersionedNiagaraEmitterData* EmitterData = EmitterHandleModel.Get().GetEmitterHandle()->GetEmitterData();
+		if (EmitterData == nullptr)
+		{
+			continue;
+		}
 
 		for (int32 i=0; i < NumRendererCountLimits; ++i)
 		{
@@ -601,7 +610,17 @@ void UNiagaraValidationRule_BannedRenderers::CheckValidity(const FNiagaraValidat
 	TArray<TSharedRef<FNiagaraEmitterHandleViewModel>> EmitterHandleViewModels = Context.ViewModel->GetEmitterHandleViewModels();
 	for (TSharedRef<FNiagaraEmitterHandleViewModel> EmitterHandleModel : EmitterHandleViewModels)
 	{
+		if ( EmitterHandleModel->GetIsEnabled() == false)
+		{
+			continue;
+		}
+
 		FVersionedNiagaraEmitterData* EmitterData = EmitterHandleModel.Get().GetEmitterHandle()->GetEmitterData();
+		if (EmitterData == nullptr)
+		{
+			continue;
+		}
+
 		EmitterData->ForEachRenderer([&Results, EmitterHandleModel, this, &System](UNiagaraRendererProperties* RendererProperties)
 		{
 			if (RendererProperties->GetIsEnabled() && BannedRenderers.Contains(RendererProperties->GetClass()))
@@ -886,6 +905,10 @@ void UNiagaraValidationRule_RendererSortingEnabled::CheckValidity(const FNiagara
 		}
 
 		FVersionedNiagaraEmitterData* EmitterData = EmitterHandleModel.Get().GetEmitterHandle()->GetEmitterData();
+		if (!EmitterData)
+		{
+			continue;
+		}
 
 		const FString PlatformConflictsString = NiagaraValidation::GetPlatformConflictsString(Platforms, EmitterData->Platforms);
 		if (PlatformConflictsString.IsEmpty())
@@ -916,8 +939,13 @@ void UNiagaraValidationRule_GpuUsage::CheckValidity(const FNiagaraValidationCont
 	for (TSharedRef<FNiagaraEmitterHandleViewModel> EmitterHandleModel : Context.ViewModel->GetEmitterHandleViewModels())
 	{
 		FNiagaraEmitterHandle* EmitterHandle = EmitterHandleModel.Get().GetEmitterHandle();
+		if (!EmitterHandle || EmitterHandle->GetIsEnabled() == false)
+		{
+			continue;
+		}
+
 		FVersionedNiagaraEmitterData* EmitterData = EmitterHandle->GetEmitterData();
-		if (EmitterData->SimTarget != ENiagaraSimTarget::GPUComputeSim)
+		if (!EmitterData || EmitterData->SimTarget != ENiagaraSimTarget::GPUComputeSim)
 		{
 			continue;
 		}
@@ -983,6 +1011,10 @@ void UNiagaraValidationRule_RibbonRenderer::CheckValidity(const FNiagaraValidati
 
 		FNiagaraEmitterHandle* EmitterHandle = EmitterHandleModel.Get().GetEmitterHandle();
 		FVersionedNiagaraEmitterData* EmitterData = EmitterHandle->GetEmitterData();
+		if (EmitterData == nullptr)
+		{
+			continue;
+		}
 
 		if (NiagaraValidation::GetPlatformConflictsString(Platforms, EmitterData->Platforms).IsEmpty())
 		{
@@ -1129,7 +1161,8 @@ void UNiagaraValidationRule_LWC::CheckValidity(const FNiagaraValidationContext&
 	TArray<TSharedRef<FNiagaraEmitterHandleViewModel>> EmitterHandleViewModels = Context.ViewModel->GetEmitterHandleViewModels();
 	for (TSharedRef<FNiagaraEmitterHandleViewModel> EmitterHandleModel : EmitterHandleViewModels)
 	{
-		if (EmitterHandleModel->GetEmitterHandle()->GetEmitterData()->bLocalSpace == false)
+		const FVersionedNiagaraEmitterData* EmitterData = EmitterHandleModel->GetEmitterHandle()->GetEmitterData();
+		if (EmitterHandleModel->GetIsEnabled() && EmitterData && EmitterData->bLocalSpace == false)
 		{
 			AllModules.Append(NiagaraValidation::GetStackEntries<UNiagaraStackModuleItem>(EmitterHandleModel.Get().GetEmitterStackViewModel()));
 		}
@@ -1201,7 +1234,7 @@ void UNiagaraValidationRule_NoOpaqueRenderMaterial::CheckValidity(const FNiagara
 		for (const TSharedRef<FNiagaraEmitterHandleViewModel>& EmitterHandleModel : Context.ViewModel->GetEmitterHandleViewModels())
 		{
 			FVersionedNiagaraEmitterData* EmitterData = EmitterHandleModel->GetEmitterHandle()->GetEmitterData();
-			if (EmitterHandleModel->GetIsEnabled() && EmitterData->SimTarget == ENiagaraSimTarget::GPUComputeSim && EmitterHandleModel->GetEmitterViewModel() == SourceModule->GetEmitterViewModel())
+			if (EmitterHandleModel->GetIsEnabled() && EmitterData && EmitterData->SimTarget == ENiagaraSimTarget::GPUComputeSim && EmitterHandleModel->GetEmitterViewModel() == SourceModule->GetEmitterViewModel())
 			{
 				// Note: for these BP driven enums we can't compare the values
 				TOptional<int32> GPUCollisionType = NiagaraValidation::GetModuleStaticInt32Value(SourceModule, NAME_GPUCollisionType);
@@ -1308,6 +1341,11 @@ void UNiagaraValidationRule_SimulationStageBudget::CheckValidity(const FNiagaraV
 
 		// Simulation stages are GPU only currently
 		FVersionedNiagaraEmitterData* EmitterData = EmitterHandleModel.Get().GetEmitterHandle()->GetEmitterData();
+		if (EmitterData == nullptr)
+		{
+			continue;
+		}
+
 		if ( EmitterData->SimTarget != ENiagaraSimTarget::GPUComputeSim )
 		{
 			continue;
@@ -1577,11 +1615,18 @@ void UNiagaraValidationRule_NoMapForOnCpu::CheckValidity(const FNiagaraValidatio
 	{
 		ENiagaraScriptUsage ScriptUsage = Module->GetOutputNode()->GetUsage();
 		FNiagaraEmitterViewModel* EmitterViewModel = Module->GetEmitterViewModel().Get();
-		if (EmitterViewModel && EmitterViewModel->GetEmitter().GetEmitterData()->SimTarget == ENiagaraSimTarget::GPUComputeSim && FNiagaraUtilities::ConvertScriptUsageToStaticSwitchContext(ScriptUsage) == ENiagaraScriptContextStaticSwitch::Particle)
+		FVersionedNiagaraEmitterData* EmitterData = EmitterViewModel ? EmitterViewModel->GetEmitter().GetEmitterData() : nullptr;
+		if (!EmitterData)
+		{
+			continue;
+		}
+
+		if (EmitterData->SimTarget == ENiagaraSimTarget::GPUComputeSim && FNiagaraUtilities::ConvertScriptUsageToStaticSwitchContext(ScriptUsage) == ENiagaraScriptContextStaticSwitch::Particle)
 		{
 			// modules used in gpu scripts are ignored 
 			continue;
 		}
+
 		if (UNiagaraGraph* Graph = Module->GetModuleNode().GetCalledGraph())
 		{
 			FObjectKey GraphKey(Graph);
@@ -1618,82 +1663,91 @@ void UNiagaraValidationRule_NoMapForOnCpu::CheckValidity(const FNiagaraValidatio
 void UNiagaraValidationRule_ModuleSimTargetRestriction::CheckValidity(const FNiagaraValidationContext& Context,	TArray<FNiagaraValidationResult>& OutResults) const
 {
 	// check to see if we're called from a module
-	if (UNiagaraStackModuleItem* SourceModule = Cast<UNiagaraStackModuleItem>(Context.Source))
+	UNiagaraStackModuleItem* SourceModule = Cast<UNiagaraStackModuleItem>(Context.Source);
+	if ( SourceModule == nullptr || !SourceModule->GetIsEnabled())
 	{
-		if (SourceModule->GetIsEnabled())
-		{
-			ENiagaraScriptUsage ScriptUsage = SourceModule->GetOutputNode()->GetUsage();
-			// system and emitter scripts are always cpu scripts
-			if (FNiagaraUtilities::ConvertScriptUsageToStaticSwitchContext(ScriptUsage) != ENiagaraScriptContextStaticSwitch::Particle && SupportedSimTarget == ENiagaraSimTarget::GPUComputeSim)
-			{
-				OutResults.Emplace(
-					Severity,
-					LOCTEXT("SimTargetModuleCpuError", "This module only supports gpu sim targets"),
-					LOCTEXT("SimTargetModuleCpuDetailedError", "This module is used in a system or emitter script (which are always executed on the cpu), but the module needs to run on the gpu. Place it in a particle script or simulation stage of a gpu emitter."),
-					SourceModule
-				);
-				return;
-			}
+		return;
+	}
+
+	const ENiagaraScriptUsage ScriptUsage = SourceModule->GetOutputNode()->GetUsage();
+
+	// system and emitter scripts are always cpu scripts
+	if (FNiagaraUtilities::ConvertScriptUsageToStaticSwitchContext(ScriptUsage) != ENiagaraScriptContextStaticSwitch::Particle && SupportedSimTarget == ENiagaraSimTarget::GPUComputeSim)
+	{
+		OutResults.Emplace(
+			Severity,
+			LOCTEXT("SimTargetModuleCpuError", "This module only supports gpu sim targets"),
+			LOCTEXT("SimTargetModuleCpuDetailedError", "This module is used in a system or emitter script (which are always executed on the cpu), but the module needs to run on the gpu. Place it in a particle script or simulation stage of a gpu emitter."),
+			SourceModule
+		);
+		return;
+	}
 			
-			TSharedPtr<FNiagaraEmitterViewModel> EmitterViewModel = SourceModule->GetEmitterViewModel();
-			if (EmitterViewModel.IsValid())
-			{
-				ENiagaraSimTarget SimTarget = EmitterViewModel->GetEmitter().GetEmitterData()->SimTarget;
-				if (SimTarget != SupportedSimTarget)
-				{
-					FNiagaraValidationResult& Result = OutResults.Emplace_GetRef(
-						Severity,
-						LOCTEXT("SimTargetModuleError", "Module is not compatible with the current emitter sim target"),
-						LOCTEXT("SimTargetModuleErrorDetailed", "This module has a restriction on the emitter sim target and can't run on both cpu/gpu emitters."),
-						SourceModule
-					);
-					Result.Fixes.Emplace(
-						FNiagaraValidationFix(
-							FText::Format(LOCTEXT("SimTargetModuleErrorFix", "Change emitter sim target to {0}"), FText::FromString(SupportedSimTarget == ENiagaraSimTarget::CPUSim ? "CPUSim" : "GPUComputeSim")),
-							FNiagaraValidationFixDelegate::CreateLambda(
-								[WeakSourceModule=MakeWeakObjectPtr(SourceModule), NewSimTarget=SupportedSimTarget]()
-								{
-									if (UNiagaraStackModuleItem* StackModule = WeakSourceModule.Get())
-									{
-										TSharedPtr<FNiagaraEmitterViewModel> EmitterViewModel = StackModule->GetEmitterViewModel();
-										if (EmitterViewModel.IsValid())
-										{
-											FVersionedNiagaraEmitter VersionedEmitter = EmitterViewModel->GetEmitter();
-											if (FVersionedNiagaraEmitterData* VersionedEmitterData = VersionedEmitter.GetEmitterData())
-											{
-												const FScopedTransaction Transaction(LOCTEXT("ChangeSimTarget", "Change emitter sim target"));
-												VersionedEmitter.Emitter->Modify();
-												VersionedEmitterData->SimTarget = NewSimTarget;
+	TSharedPtr<FNiagaraEmitterViewModel> EmitterViewModel = SourceModule->GetEmitterViewModel();
+	if (!EmitterViewModel.IsValid())
+	{
+		return;
+	}
 
-												FProperty* SimTargetProperty = FindFProperty<FProperty>(FVersionedNiagaraEmitterData::StaticStruct(), GET_MEMBER_NAME_CHECKED(FVersionedNiagaraEmitterData, SimTarget));
-												FPropertyChangedEvent PropertyChangedEvent(SimTargetProperty);
-												VersionedEmitter.Emitter->PostEditChangeVersionedProperty(PropertyChangedEvent, VersionedEmitter.Version);
+	FVersionedNiagaraEmitterData* EmitterData = EmitterViewModel->GetEmitter().GetEmitterData();
+	if (!EmitterData)
+	{
+		return;
+	}
 
-												UNiagaraSystem::RequestCompileForEmitter(VersionedEmitter);
-											}
-										}
-									}
-								}
-							)
-						)
-					);
-					Result.Fixes.Emplace(
-						FNiagaraValidationFix(
-							LOCTEXT("SimTargetModuleErrorFixDisable", "Disable module"),
-							FNiagaraValidationFixDelegate::CreateLambda(
-								[WeakSourceModule=MakeWeakObjectPtr(SourceModule)]()
+	ENiagaraSimTarget SimTarget = EmitterViewModel->GetEmitter().GetEmitterData()->SimTarget;
+	if (SimTarget != SupportedSimTarget)
+	{
+		FNiagaraValidationResult& Result = OutResults.Emplace_GetRef(
+			Severity,
+			LOCTEXT("SimTargetModuleError", "Module is not compatible with the current emitter sim target"),
+			LOCTEXT("SimTargetModuleErrorDetailed", "This module has a restriction on the emitter sim target and can't run on both cpu/gpu emitters."),
+			SourceModule
+		);
+		Result.Fixes.Emplace(
+			FNiagaraValidationFix(
+				FText::Format(LOCTEXT("SimTargetModuleErrorFix", "Change emitter sim target to {0}"), FText::FromString(SupportedSimTarget == ENiagaraSimTarget::CPUSim ? "CPUSim" : "GPUComputeSim")),
+				FNiagaraValidationFixDelegate::CreateLambda(
+					[WeakSourceModule=MakeWeakObjectPtr(SourceModule), NewSimTarget=SupportedSimTarget]()
+					{
+						if (UNiagaraStackModuleItem* StackModule = WeakSourceModule.Get())
+						{
+							TSharedPtr<FNiagaraEmitterViewModel> EmitterViewModel = StackModule->GetEmitterViewModel();
+							if (EmitterViewModel.IsValid())
+							{
+								FVersionedNiagaraEmitter VersionedEmitter = EmitterViewModel->GetEmitter();
+								if (FVersionedNiagaraEmitterData* VersionedEmitterData = VersionedEmitter.GetEmitterData())
 								{
-									if (UNiagaraStackModuleItem* StackModule = WeakSourceModule.Get())
-									{
-										StackModule->SetEnabled(false);
-									}
+									const FScopedTransaction Transaction(LOCTEXT("ChangeSimTarget", "Change emitter sim target"));
+									VersionedEmitter.Emitter->Modify();
+									VersionedEmitterData->SimTarget = NewSimTarget;
+
+									FProperty* SimTargetProperty = FindFProperty<FProperty>(FVersionedNiagaraEmitterData::StaticStruct(), GET_MEMBER_NAME_CHECKED(FVersionedNiagaraEmitterData, SimTarget));
+									FPropertyChangedEvent PropertyChangedEvent(SimTargetProperty);
+									VersionedEmitter.Emitter->PostEditChangeVersionedProperty(PropertyChangedEvent, VersionedEmitter.Version);
+
+									UNiagaraSystem::RequestCompileForEmitter(VersionedEmitter);
 								}
-							)
-						)
-					);
-				}
-			}
-		}
+							}
+						}
+					}
+				)
+			)
+		);
+		Result.Fixes.Emplace(
+			FNiagaraValidationFix(
+				LOCTEXT("SimTargetModuleErrorFixDisable", "Disable module"),
+				FNiagaraValidationFixDelegate::CreateLambda(
+					[WeakSourceModule=MakeWeakObjectPtr(SourceModule)]()
+					{
+						if (UNiagaraStackModuleItem* StackModule = WeakSourceModule.Get())
+						{
+							StackModule->SetEnabled(false);
+						}
+					}
+				)
+			)
+		);
 	}
 }
 
@@ -1702,7 +1756,17 @@ void UNiagaraValidationRule_MaterialUsage::CheckValidity(const FNiagaraValidatio
 	TArray<TSharedRef<FNiagaraEmitterHandleViewModel>> EmitterHandleViewModels = Context.ViewModel->GetEmitterHandleViewModels();
 	for (TSharedRef<FNiagaraEmitterHandleViewModel> EmitterHandleModel : EmitterHandleViewModels)
 	{
+		if (EmitterHandleModel->GetIsEnabled() == false)
+		{
+			continue;
+		}
+
 		FVersionedNiagaraEmitterData* EmitterData = EmitterHandleModel.Get().GetEmitterHandle()->GetEmitterData();
+		if (EmitterData == nullptr)
+		{
+			continue;
+		}
+
 		EmitterData->ForEachRenderer(
 			[&OutResults, EmitterHandleModel, this](UNiagaraRendererProperties* RendererProperties)
 			{
