UE: Add a null check for missing emitter data

diff --git a/Engine/Plugins/FX/Niagara/Source/Niagara/Private/NiagaraDataInterfaceParticleRead.cpp b/Engine/Plugins/FX/Niagara/Source/Niagara/Private/NiagaraDataInterfaceParticleRead.cpp
--- a/Engine/Plugins/FX/Niagara/Source/Niagara/Private/NiagaraDataInterfaceParticleRead.cpp
+++ b/Engine/Plugins/FX/Niagara/Source/Niagara/Private/NiagaraDataInterfaceParticleRead.cpp
@@ -2368,9 +2368,12 @@ void UNiagaraDataInterfaceParticleRead::GetFeedback(UNiagaraSystem* Asset, UNiag
 	Scripts.Add(Asset->GetSystemUpdateScript());
 	for (auto&& EmitterHandle : Asset->GetEmitterHandles())
 	{
-		TArray<UNiagaraScript*> OutScripts;
-		EmitterHandle.GetEmitterData()->GetScripts(OutScripts, false);
-		Scripts.Append(OutScripts);
+		if (FVersionedNiagaraEmitterData* EmitterData = EmitterHandle.GetEmitterData())
+		{
+			TArray<UNiagaraScript*> OutScripts;
+			EmitterData->GetScripts(OutScripts, false);
+			Scripts.Append(OutScripts);
+		}
 	}
 
 	// Now check if any script uses functions that require persisitent ID access
