- Fixes for null emitter data access

#rnx

diff --git a/Engine/Plugins/FX/Niagara/Source/Niagara/Private/NiagaraDataInterfaceSkeletalMesh.cpp b/Engine/Plugins/FX/Niagara/Source/Niagara/Private/NiagaraDataInterfaceSkeletalMesh.cpp
--- a/Engine/Plugins/FX/Niagara/Source/Niagara/Private/NiagaraDataInterfaceSkeletalMesh.cpp
+++ b/Engine/Plugins/FX/Niagara/Source/Niagara/Private/NiagaraDataInterfaceSkeletalMesh.cpp
@@ -2661,9 +2661,12 @@ void UNiagaraDataInterfaceSkeletalMesh::GetFeedback(UNiagaraSystem* Asset, UNiag
 			TArray<UNiagaraScript*> Scripts;
 			Scripts.Add(Asset->GetSystemSpawnScript());
 			Scripts.Add(Asset->GetSystemUpdateScript());
-			for (auto&& EmitterHandle : Asset->GetEmitterHandles())
+			for (const FNiagaraEmitterHandle& EmitterHandle : Asset->GetEmitterHandles())
 			{
-				EmitterHandle.GetEmitterData()->GetScripts(Scripts, false);
+				if (FVersionedNiagaraEmitterData* EmitterData = EmitterHandle.GetEmitterData())
+				{
+					EmitterData->GetScripts(Scripts, false);
+				}
 			}
 
 			// Now check if any script uses functions that require CPU access
diff --git a/Engine/Plugins/FX/Niagara/Source/Niagara/Private/NiagaraDataInterfaceVectorField.cpp b/Engine/Plugins/FX/Niagara/Source/Niagara/Private/NiagaraDataInterfaceVectorField.cpp
--- a/Engine/Plugins/FX/Niagara/Source/Niagara/Private/NiagaraDataInterfaceVectorField.cpp
+++ b/Engine/Plugins/FX/Niagara/Source/Niagara/Private/NiagaraDataInterfaceVectorField.cpp
@@ -272,9 +272,12 @@ void UNiagaraDataInterfaceVectorField::GetFeedback(UNiagaraSystem* InAsset, UNia
 		Scripts.Add(InAsset->GetSystemUpdateScript());
 		for (auto&& EmitterHandle : InAsset->GetEmitterHandles())
 		{
-			TArray<UNiagaraScript*> OutScripts;
-			EmitterHandle.GetEmitterData()->GetScripts(OutScripts, false);
-			Scripts.Append(OutScripts);
+			if (FVersionedNiagaraEmitterData* EmitterData = EmitterHandle.GetEmitterData())
+			{
+				TArray<UNiagaraScript*> OutScripts;
+				EmitterData->GetScripts(OutScripts, false);
+				Scripts.Append(OutScripts);
+			}
 		}
 
 		for (const auto Script : Scripts)
@@ -304,7 +307,7 @@ void UNiagaraDataInterfaceVectorField::GetFeedback(UNiagaraSystem* InAsset, UNia
 							(CachedDefaultDIs[Idx].DataInterface->Equals(this) || DIAliases.Contains(CachedDefaultDIs[Idx].Name)))
 						{
 							FVersionedNiagaraEmitter OuterEmitter = Script->GetOuterEmitter();
-							if (OuterEmitter.Emitter && (OuterEmitter.GetEmitterData()->SimTarget == ENiagaraSimTarget::CPUSim || Script->IsSystemScript(Script->Usage)))
+							if (OuterEmitter.Emitter && OuterEmitter.GetEmitterData() && (OuterEmitter.GetEmitterData()->SimTarget == ENiagaraSimTarget::CPUSim || Script->IsSystemScript(Script->Usage)))
 							{
 								bHasCPUFunctions = true;
 							}
diff --git a/Engine/Plugins/FX/Niagara/Source/Niagara/Private/NiagaraEmitterHandle.cpp b/Engine/Plugins/FX/Niagara/Source/Niagara/Private/NiagaraEmitterHandle.cpp
--- a/Engine/Plugins/FX/Niagara/Source/Niagara/Private/NiagaraEmitterHandle.cpp
+++ b/Engine/Plugins/FX/Niagara/Source/Niagara/Private/NiagaraEmitterHandle.cpp
@@ -262,14 +262,17 @@ bool FNiagaraEmitterHandle::NeedsRecompile() const
 {
 	if (GetIsEnabled())
 	{
-		TArray<UNiagaraScript*> Scripts;
-		VersionedInstance.GetEmitterData()->GetScripts(Scripts);
-
-		for (UNiagaraScript* Script : Scripts)
+		if (FVersionedNiagaraEmitterData* EmitterData = VersionedInstance.GetEmitterData())
 		{
-			if (Script->IsCompilable() && !Script->AreScriptAndSourceSynchronized())
+			TArray<UNiagaraScript*> Scripts;
+			EmitterData->GetScripts(Scripts);
+
+			for (UNiagaraScript* Script : Scripts)
 			{
-				return true;
+				if (Script->IsCompilable() && !Script->AreScriptAndSourceSynchronized())
+				{
+					return true;
+				}
 			}
 		}
 	}
diff --git a/Engine/Plugins/FX/Niagara/Source/Niagara/Private/NiagaraSimCache.cpp b/Engine/Plugins/FX/Niagara/Source/Niagara/Private/NiagaraSimCache.cpp
--- a/Engine/Plugins/FX/Niagara/Source/Niagara/Private/NiagaraSimCache.cpp
+++ b/Engine/Plugins/FX/Niagara/Source/Niagara/Private/NiagaraSimCache.cpp
@@ -386,7 +386,14 @@ bool UNiagaraSimCache::BeginWrite(FNiagaraSimCacheCreateParameters InCreateParam
 	#if WITH_EDITORONLY_DATA
 		for ( const FNiagaraEmitterHandle& EmitterHandle : Helper.NiagaraSystem->GetEmitterHandles() )
 		{
-			EmitterHandle.GetInstance().GetEmitterData()->ForEachEnabledRenderer(
+			//-TODO:Stateless
+			FVersionedNiagaraEmitterData* EmitterData = EmitterHandle.GetInstance().GetEmitterData();
+			if (EmitterData == nullptr)
+			{
+				continue;
+			}
+
+			EmitterData->ForEachEnabledRenderer(
 				[&](UNiagaraRendererProperties* RenderProperties)
 				{
 					for (FNiagaraVariableBase BoundAttribute : RenderProperties->GetBoundAttributes())
diff --git a/Engine/Plugins/FX/Niagara/Source/NiagaraEditor/Private/NiagaraEditorModule.cpp b/Engine/Plugins/FX/Niagara/Source/NiagaraEditor/Private/NiagaraEditorModule.cpp
--- a/Engine/Plugins/FX/Niagara/Source/NiagaraEditor/Private/NiagaraEditorModule.cpp
+++ b/Engine/Plugins/FX/Niagara/Source/NiagaraEditor/Private/NiagaraEditorModule.cpp
@@ -527,11 +527,14 @@ void DumpRapidIterationParametersForScript(UNiagaraScript* Script, const FString
 
 void DumpRapidIterationParametersForEmitter(FVersionedNiagaraEmitter Emitter, const FString& EmitterName)
 {
-	TArray<UNiagaraScript*> Scripts;
-	Emitter.GetEmitterData()->GetScripts(Scripts, false);
-	for (UNiagaraScript* Script : Scripts)
+	if (FVersionedNiagaraEmitterData* EmitterData = Emitter.GetEmitterData())
 	{
-		DumpRapidIterationParametersForScript(Script, EmitterName);
+		TArray<UNiagaraScript*> Scripts;
+		EmitterData->GetScripts(Scripts, false);
+		for (UNiagaraScript* Script : Scripts)
+		{
+			DumpRapidIterationParametersForScript(Script, EmitterName);
+		}
 	}
 }
 
@@ -776,11 +779,14 @@ void DumpCompileIdDataForEmitter(const FVersionedNiagaraEmitter& VersionedEmitte
 	MakeIndent(IndentLevel, Indent);
 	Dump.Append(FString::Printf(TEXT("%sEmitter: %s\n"), *Indent, *VersionedEmitter.Emitter->GetUniqueEmitterName()));
 
-	TArray<UNiagaraScript*> Scripts;
-	VersionedEmitter.GetEmitterData()->GetScripts(Scripts, false);
-	for (UNiagaraScript* Script : Scripts)
+	if (FVersionedNiagaraEmitterData* EmitterData = VersionedEmitter.GetEmitterData())
 	{
-		DumpCompileIdDataForScript(Script, IndentLevel + 1, Dump);
+		TArray<UNiagaraScript*> Scripts;
+		EmitterData->GetScripts(Scripts, false);
+		for (UNiagaraScript* Script : Scripts)
+		{
+			DumpCompileIdDataForScript(Script, IndentLevel + 1, Dump);
+		}
 	}
 }
 
diff --git a/Engine/Plugins/FX/Niagara/Source/NiagaraEditor/Private/NiagaraEditorUtilities.cpp b/Engine/Plugins/FX/Niagara/Source/NiagaraEditor/Private/NiagaraEditorUtilities.cpp
--- a/Engine/Plugins/FX/Niagara/Source/NiagaraEditor/Private/NiagaraEditorUtilities.cpp
+++ b/Engine/Plugins/FX/Niagara/Source/NiagaraEditor/Private/NiagaraEditorUtilities.cpp
@@ -2093,7 +2093,8 @@ TArray<UNiagaraComponent*> FNiagaraEditorUtilities::GetComponentsThatReferenceSy
 		{
 			for (const auto& EmitterHandle : ReferencedSystemViewModel.GetSystem().GetEmitterHandles())
 			{
-				if (Component->GetAsset()->UsesEmitter(EmitterHandle.GetEmitterData()->GetParent()))
+				FVersionedNiagaraEmitterData* EmitterData = EmitterHandle.GetEmitterData();
+				if (EmitterData && Component->GetAsset()->UsesEmitter(EmitterData->GetParent()))
 				{
 					ReferencingComponents.Add(Component);
 				}
@@ -2149,7 +2150,12 @@ const FGuid FNiagaraEditorUtilities::AddEmitterToSystem(UNiagaraSystem& InSystem
 	// We only want to use the default view state when actually adding an emitter to a system asset, not when opening up an emitter which adds said emitter to a transient system
 	if (SystemEditorData->GetOwningSystemIsPlaceholder() == false)
 	{
-		Cast<UNiagaraEmitterEditorData>(EmitterHandle.GetEmitterData()->GetEditorData())->SetShowSummaryView(InEmitterToAdd.GetEmitterData(EmitterVersion)->AddEmitterDefaultViewState == ENiagaraEmitterDefaultSummaryState::Summary ? true : false);
+		FVersionedNiagaraEmitterData* EmitterData = EmitterHandle.GetEmitterData();
+		UNiagaraEmitterEditorData* EmitterEditorData = EmitterData ? Cast<UNiagaraEmitterEditorData>(EmitterData->GetEditorData()) : nullptr;
+		if (EmitterEditorData)
+		{
+			EmitterEditorData->SetShowSummaryView(InEmitterToAdd.GetEmitterData(EmitterVersion)->AddEmitterDefaultViewState == ENiagaraEmitterDefaultSummaryState::Summary ? true : false);
+		}
 	}
 	
 	FNiagaraStackGraphUtilities::RebuildEmitterNodes(InSystem);
@@ -2318,28 +2324,32 @@ TArray<FNiagaraVariable> FNiagaraEditorUtilities::GetReferencedUserParametersFro
 	}
 
 	// then we check for user parameter bindings on renderers
-	const TArray<UNiagaraRendererProperties*>& Renderers = EmitterViewModel->GetEmitter().GetEmitterData()->GetRenderers();
-	for(UNiagaraRendererProperties* Renderer : Renderers)
+	//-TODO:Stateless
+	if (FVersionedNiagaraEmitterData* EmitterData = EmitterViewModel->GetEmitter().GetEmitterData())
 	{
-		for(TFieldIterator<FStructProperty> It(Renderer->GetClass()); It; ++It)
+		const TArray<UNiagaraRendererProperties*>& Renderers = EmitterData->GetRenderers();
+		for(UNiagaraRendererProperties* Renderer : Renderers)
 		{
-			if((*It)->Struct == StaticStruct<FNiagaraUserParameterBinding>())
+			for(TFieldIterator<FStructProperty> It(Renderer->GetClass()); It; ++It)
 			{
-				FNiagaraUserParameterBinding ParameterBinding;
-				(*It)->GetValue_InContainer(Renderer, &ParameterBinding);
-
-				if(ParameterBinding.Parameter.IsValid())
+				if((*It)->Struct == StaticStruct<FNiagaraUserParameterBinding>())
 				{
-					FNiagaraParameterHandle ReferencedParameter(ParameterBinding.Parameter.GetName());
-					if(ReferencedParameter.IsUserHandle())
+					FNiagaraUserParameterBinding ParameterBinding;
+					(*It)->GetValue_InContainer(Renderer, &ParameterBinding);
+
+					if(ParameterBinding.Parameter.IsValid())
 					{
-						UserParameters.AddUnique(ParameterBinding.Parameter);
+						FNiagaraParameterHandle ReferencedParameter(ParameterBinding.Parameter.GetName());
+						if(ReferencedParameter.IsUserHandle())
+						{
+							UserParameters.AddUnique(ParameterBinding.Parameter);
+						}
 					}
 				}
 			}
 		}
 	}
-	
+
 	return UserParameters;
 }
 
@@ -2380,21 +2390,25 @@ TArray<FNiagaraUserParameterBinding*> FNiagaraEditorUtilities::GetUserParameterB
 	TArray<FNiagaraUserParameterBinding*> Bindings;
 
 	// First we gather the bindings on renderers of an emitter
-	const TArray<UNiagaraRendererProperties*>& Renderers = EmitterViewModel->GetEmitter().GetEmitterData()->GetRenderers();
-	for(UNiagaraRendererProperties* Renderer : Renderers)
+	//-TODO:Stateless
+	if (FVersionedNiagaraEmitterData* EmitterData = EmitterViewModel->GetEmitter().GetEmitterData())
 	{
-		for(TFieldIterator<FStructProperty> It(Renderer->GetClass()); It; ++It)
+		const TArray<UNiagaraRendererProperties*>& Renderers = EmitterData->GetRenderers();
+		for(UNiagaraRendererProperties* Renderer : Renderers)
 		{
-			if((*It)->Struct == StaticStruct<FNiagaraUserParameterBinding>())
-			{				
-				FNiagaraUserParameterBinding* ParameterBinding = (FNiagaraUserParameterBinding*) (*It)->ContainerPtrToValuePtr<void>(Renderer, 0);
+			for(TFieldIterator<FStructProperty> It(Renderer->GetClass()); It; ++It)
+			{
+				if((*It)->Struct == StaticStruct<FNiagaraUserParameterBinding>())
+				{				
+					FNiagaraUserParameterBinding* ParameterBinding = (FNiagaraUserParameterBinding*) (*It)->ContainerPtrToValuePtr<void>(Renderer, 0);
 				
-				if(ParameterBinding->Parameter.IsValid())
-				{
-					FNiagaraParameterHandle ReferencedParameter(ParameterBinding->Parameter.GetName());
-					if(ReferencedParameter.IsUserHandle() && ParameterBinding->Parameter == UserParameter)
+					if(ParameterBinding->Parameter.IsValid())
 					{
-						Bindings.Add(ParameterBinding);
+						FNiagaraParameterHandle ReferencedParameter(ParameterBinding->Parameter.GetName());
+						if(ReferencedParameter.IsUserHandle() && ParameterBinding->Parameter == UserParameter)
+						{
+							Bindings.Add(ParameterBinding);
+						}
 					}
 				}
 			}
@@ -2437,18 +2451,22 @@ TArray<TPair<FNiagaraVariableAttributeBinding*, ENiagaraRendererSourceDataMode>>
 {
 	TArray<TPair<FNiagaraVariableAttributeBinding*, ENiagaraRendererSourceDataMode>> AttributeBindings;
 	
-	const TArray<UNiagaraRendererProperties*>& Renderers = EmitterViewModel->GetEmitter().GetEmitterData()->GetRenderers();
-	for(UNiagaraRendererProperties* Renderer : Renderers)
+	//-TODO:Stateless
+	if (FVersionedNiagaraEmitterData* EmitterData = EmitterViewModel->GetEmitter().GetEmitterData())
 	{
-		for(TFieldIterator<FStructProperty> It(Renderer->GetClass()); It; ++It)
+		const TArray<UNiagaraRendererProperties*>& Renderers = EmitterData->GetRenderers();
+		for(UNiagaraRendererProperties* Renderer : Renderers)
 		{
-			if((*It)->Struct == StaticStruct<FNiagaraVariableAttributeBinding>())
-			{				
-				FNiagaraVariableAttributeBinding* ParameterBinding = (FNiagaraVariableAttributeBinding*) (*It)->ContainerPtrToValuePtr<void>(Renderer, 0);
+			for(TFieldIterator<FStructProperty> It(Renderer->GetClass()); It; ++It)
+			{
+				if((*It)->Struct == StaticStruct<FNiagaraVariableAttributeBinding>())
+				{				
+					FNiagaraVariableAttributeBinding* ParameterBinding = (FNiagaraVariableAttributeBinding*) (*It)->ContainerPtrToValuePtr<void>(Renderer, 0);
 				
-				if(ParameterBinding->GetName().IsEqual(Parameter.GetName()))
-				{
-					AttributeBindings.Add({ParameterBinding, Renderer->GetCurrentSourceMode()});
+					if(ParameterBinding->GetName().IsEqual(Parameter.GetName()))
+					{
+						AttributeBindings.Add({ParameterBinding, Renderer->GetCurrentSourceMode()});
+					}
 				}
 			}
 		}
@@ -3640,7 +3658,14 @@ void FNiagaraEditorUtilities::SwitchParentEmitterVersion(TSharedRef<FNiagaraEmit
 {
 	// if we want to run the python upgrade script, we make a snapshot of the existing emitter viewmodel
 	UNiagaraPythonEmitter* OldPythonEmitter = NewObject<UNiagaraPythonEmitter>(GetTransientPackage());
-	FVersionedNiagaraEmitter Parent = EmitterViewModel->GetEmitter().GetEmitterData()->GetParent();
+	FVersionedNiagaraEmitterData* EmitterData = EmitterViewModel->GetEmitter().GetEmitterData();
+	ensureMsgf(EmitterData, TEXT("EmitterData missing when switching emitter version"));
+	if (EmitterData == nullptr)
+	{
+		return;
+	}
+
+	FVersionedNiagaraEmitter Parent = EmitterData->GetParent();
 	bool bRunPython = Parent.GetEmitterData() && Parent.Emitter->IsVersioningEnabled();
 	int32 ViewModelIndex = 0;
 	TSharedPtr<FNiagaraSystemViewModel> OldSystemViewModel;
diff --git a/Engine/Plugins/FX/Niagara/Source/NiagaraEditor/Private/Sequencer/NiagaraSequence/MovieSceneNiagaraEmitterTrack.cpp b/Engine/Plugins/FX/Niagara/Source/NiagaraEditor/Private/Sequencer/NiagaraSequence/MovieSceneNiagaraEmitterTrack.cpp
--- a/Engine/Plugins/FX/Niagara/Source/NiagaraEditor/Private/Sequencer/NiagaraSequence/MovieSceneNiagaraEmitterTrack.cpp
+++ b/Engine/Plugins/FX/Niagara/Source/NiagaraEditor/Private/Sequencer/NiagaraSequence/MovieSceneNiagaraEmitterTrack.cpp
@@ -272,13 +272,13 @@ void UMovieSceneNiagaraEmitterTrack::CreateSections(const FFrameRate& InFrameRes
 {
 	SectionInitializationErrors.Empty();
 
-	//-TODO:Stateless: We need to handle the stateless path here
-	if (GetEmitterHandleViewModel()->GetEmitterViewModel()->GetEmitter().GetEmitterData() == nullptr)
+	const FVersionedNiagaraEmitterData* EmitterData = GetEmitterHandleViewModel()->GetEmitterHandle()->GetEmitterData();
+	if (EmitterData == nullptr)
 	{
 		return;
 	}
 
-	UNiagaraScript* EmitterUpdateScript = GetEmitterHandleViewModel()->GetEmitterViewModel()->GetEmitter().GetEmitterData()->GetScript(ENiagaraScriptUsage::EmitterUpdateScript, FGuid());
+	UNiagaraScript* EmitterUpdateScript = EmitterData->GetScript(ENiagaraScriptUsage::EmitterUpdateScript, FGuid());
 	UNiagaraScriptSource* ScriptSource = CastChecked<UNiagaraScriptSource>(EmitterUpdateScript->GetLatestSource());
 	UNiagaraNodeOutput* OutputNode = ScriptSource->NodeGraph->FindOutputNode(ENiagaraScriptUsage::EmitterUpdateScript);
 
diff --git a/Engine/Plugins/FX/Niagara/Source/NiagaraEditor/Private/Sequencer/NiagaraSequence/Sections/MovieSceneNiagaraEmitterSection.cpp b/Engine/Plugins/FX/Niagara/Source/NiagaraEditor/Private/Sequencer/NiagaraSequence/Sections/MovieSceneNiagaraEmitterSection.cpp
--- a/Engine/Plugins/FX/Niagara/Source/NiagaraEditor/Private/Sequencer/NiagaraSequence/Sections/MovieSceneNiagaraEmitterSection.cpp
+++ b/Engine/Plugins/FX/Niagara/Source/NiagaraEditor/Private/Sequencer/NiagaraSequence/Sections/MovieSceneNiagaraEmitterSection.cpp
@@ -157,9 +157,16 @@ bool UMovieSceneNiagaraEmitterSection::TryAddTimeRangeModule(UNiagaraNodeFunctio
 		OutErrorMessage = LOCTEXT("CanNotBeShared", "Only one module can use the time range mode.");
 		return false;
 	}
+	
+	const FVersionedNiagaraEmitterData* EmitterData = GetEmitterHandleViewModel()->GetEmitterHandle()->GetEmitterData();
+	if (EmitterData == nullptr)
+	{
+		OutErrorMessage = LOCTEXT("EmiterNotSupport", "This emitter type is not supported for this operation.");
+		return false;
+	}
 
 	SectionTimingModule = &InModule;
-	UNiagaraScript* EmitterUpdateScript = GetEmitterHandleViewModel()->GetEmitterHandle()->GetEmitterData()->GetScript(ENiagaraScriptUsage::EmitterUpdateScript, FGuid());
+	UNiagaraScript* EmitterUpdateScript = EmitterData->GetScript(ENiagaraScriptUsage::EmitterUpdateScript, FGuid());
 	UNiagaraScript* SystemUpdateScript = GetSystemViewModel().GetSystem().GetSystemUpdateScript();
 	TArray<UNiagaraScript*> BinderDependentScripts;
 	BinderDependentScripts.Add(SystemUpdateScript);
@@ -243,7 +250,14 @@ bool UMovieSceneNiagaraEmitterSection::TryAddKeyModule(UNiagaraNodeFunctionCall&
 
 bool UMovieSceneNiagaraEmitterSection::TrySetupModuleAndBinders(UNiagaraNodeFunctionCall& InModule, FModuleAndBinders& InOutModuleAndBinders, FText& OutErrorMessage)
 {
-	UNiagaraScript* EmitterUpdateScript = GetEmitterHandleViewModel()->GetEmitterHandle()->GetEmitterData()->GetScript(ENiagaraScriptUsage::EmitterUpdateScript, FGuid());
+	const FVersionedNiagaraEmitterData* EmitterData = GetEmitterHandleViewModel()->GetEmitterHandle()->GetEmitterData();
+	if (EmitterData == nullptr)
+	{
+		OutErrorMessage = LOCTEXT("EmiterNotSupport", "This emitter type is not supported for this operation.");
+		return false;
+	}
+
+	UNiagaraScript* EmitterUpdateScript = EmitterData->GetScript(ENiagaraScriptUsage::EmitterUpdateScript, FGuid());
 	UNiagaraScript* SystemUpdateScript = GetSystemViewModel().GetSystem().GetSystemUpdateScript();
 	TArray<UNiagaraScript*> BinderDependentScripts;
 	BinderDependentScripts.Add(SystemUpdateScript);
@@ -458,9 +472,10 @@ void UMovieSceneNiagaraEmitterSection::UpdateKeyModulesFromSection(FChannelAndMo
 	}
 
 	// Create new modules for new keys.
-	if (NewKeysIndices.Num() > 0)
+	const FVersionedNiagaraEmitterData* EmitterData = GetEmitterHandleViewModel()->GetEmitterHandle()->GetEmitterData();
+	if (NewKeysIndices.Num() > 0 && EmitterData)
 	{
-		UNiagaraScript* EmitterUpdateScript = GetEmitterHandleViewModel()->GetEmitterHandle()->GetEmitterData()->GetScript(ENiagaraScriptUsage::EmitterUpdateScript, FGuid());
+		UNiagaraScript* EmitterUpdateScript = EmitterData->GetScript(ENiagaraScriptUsage::EmitterUpdateScript, FGuid());
 		UNiagaraScript* SystemUpdateScript = GetSystemViewModel().GetSystem().GetSystemUpdateScript();
 		TArray<UNiagaraScript*> BinderDependentScripts;
 		BinderDependentScripts.Add(SystemUpdateScript);
diff --git a/Engine/Plugins/FX/Niagara/Source/NiagaraEditor/Private/ViewModels/NiagaraCurveSelectionViewModel.cpp b/Engine/Plugins/FX/Niagara/Source/NiagaraEditor/Private/ViewModels/NiagaraCurveSelectionViewModel.cpp
--- a/Engine/Plugins/FX/Niagara/Source/NiagaraEditor/Private/ViewModels/NiagaraCurveSelectionViewModel.cpp
+++ b/Engine/Plugins/FX/Niagara/Source/NiagaraEditor/Private/ViewModels/NiagaraCurveSelectionViewModel.cpp
@@ -799,9 +799,15 @@ TSharedPtr<FNiagaraCurveSelectionTreeNode> UNiagaraCurveSelectionViewModel::Crea
 		OldEmitterChildNodes = OldEmitterNode->GetChildNodes();
 	}
 
+	// Early out for stateless emitters
 	FVersionedNiagaraEmitterData* EmitterData = EmitterHandle.GetEmitterData();
+	if (EmitterData == nullptr)
+	{
+		ensureMsgf(OldEmitterNode.IsValid(), TEXT("Old emitter node was valid and new one is not, this is not expected"));
+		return NewEmitterNode;
+	}
 
-	UNiagaraStackEditorData& StackEditorData = CastChecked<UNiagaraEmitterEditorData>(EmitterHandle.GetEmitterData()->GetEditorData())->GetStackEditorData();
+	UNiagaraStackEditorData& StackEditorData = CastChecked<UNiagaraEmitterEditorData>(EmitterData->GetEditorData())->GetStackEditorData();
 	TSharedPtr<FNiagaraCurveSelectionTreeNode> EmitterSpawnNode = CreateNodeForScript(
 		OldEmitterChildNodes, *EmitterData->EmitterSpawnScriptProps.Script, TEXT("Emitter Spawn"), StackEditorData,
 		UNiagaraStackEntry::FExecutionCategoryNames::Emitter, UNiagaraStackEntry::FExecutionSubcategoryNames::Spawn,
@@ -918,10 +924,8 @@ void UNiagaraCurveSelectionViewModel::Refresh()
 		for (TSharedRef<FNiagaraEmitterHandleViewModel> EmitterHandleViewModel : SystemViewModel->GetEmitterHandleViewModels())
 		{
 			// During transactions Emitters may not have their EditorData set if they were created during the transaction; skip if this is true.
-			//-TODO:Stateless:
 			FNiagaraEmitterHandle* EmitterHandle = EmitterHandleViewModel->GetEmitterHandle();
 			if (EmitterHandle && EmitterHandle->GetEmitterData() && EmitterHandle->GetEmitterData()->GetEditorData())
-			//-TODO:Stateless:
 			{
 				TSharedPtr<FNiagaraCurveSelectionTreeNode> EmitterNode = CreateNodeForEmitter(OldRootChildNodes, *EmitterHandle);
 				if (EmitterNode.IsValid())
diff --git a/Engine/Plugins/FX/Niagara/Source/NiagaraEditorWidgets/Private/DetailCustomizations/NiagaraDataInterfaceMeshRendererInfoDetails.cpp b/Engine/Plugins/FX/Niagara/Source/NiagaraEditorWidgets/Private/DetailCustomizations/NiagaraDataInterfaceMeshRendererInfoDetails.cpp
--- a/Engine/Plugins/FX/Niagara/Source/NiagaraEditorWidgets/Private/DetailCustomizations/NiagaraDataInterfaceMeshRendererInfoDetails.cpp
+++ b/Engine/Plugins/FX/Niagara/Source/NiagaraEditorWidgets/Private/DetailCustomizations/NiagaraDataInterfaceMeshRendererInfoDetails.cpp
@@ -186,7 +186,7 @@ FText FNiagaraDataInterfaceMeshRendererInfoDetails::CreateRendererTextLabel(cons
 	UNiagaraEmitter* Emitter = Properties->GetTypedOuter<UNiagaraEmitter>();
 	check(Emitter);
 
-	UNiagaraEmitterEditorData* EmitterEditorData = static_cast<UNiagaraEmitterEditorData*>(Properties->GetEmitterData()->GetEditorData());
+	UNiagaraEmitterEditorData* EmitterEditorData = Properties->GetEmitterData() ? static_cast<UNiagaraEmitterEditorData*>(Properties->GetEmitterData()->GetEditorData()) : nullptr;
 	if (EmitterEditorData == nullptr)
 	{
 		return LOCTEXT("NoneOption", "None");
diff --git a/Engine/Plugins/FX/Niagara/Source/NiagaraEditorWidgets/Private/DetailCustomizations/NiagaraDataInterfaceSpriteRendererInfoDetails.cpp b/Engine/Plugins/FX/Niagara/Source/NiagaraEditorWidgets/Private/DetailCustomizations/NiagaraDataInterfaceSpriteRendererInfoDetails.cpp
--- a/Engine/Plugins/FX/Niagara/Source/NiagaraEditorWidgets/Private/DetailCustomizations/NiagaraDataInterfaceSpriteRendererInfoDetails.cpp
+++ b/Engine/Plugins/FX/Niagara/Source/NiagaraEditorWidgets/Private/DetailCustomizations/NiagaraDataInterfaceSpriteRendererInfoDetails.cpp
@@ -186,7 +186,7 @@ FText FNiagaraDataInterfaceSpriteRendererInfoDetails::CreateRendererTextLabel(co
 	UNiagaraEmitter* Emitter = Properties->GetTypedOuter<UNiagaraEmitter>();
 	check(Emitter);
 
-	UNiagaraEmitterEditorData* EmitterEditorData = static_cast<UNiagaraEmitterEditorData*>(Properties->GetEmitterData()->GetEditorData());
+	UNiagaraEmitterEditorData* EmitterEditorData = Properties->GetEmitterData() ? static_cast<UNiagaraEmitterEditorData*>(Properties->GetEmitterData()->GetEditorData()) : nullptr;
 	if (EmitterEditorData == nullptr)
 	{
 		return LOCTEXT("NoneOption", "None");
