fix issue with clear before non iteration stage for rasterization grids

diff --git a/Engine/Plugins/FX/Niagara/Source/Niagara/Classes/NiagaraDataInterfaceRasterizationGrid3D.h b/Engine/Plugins/FX/Niagara/Source/Niagara/Classes/NiagaraDataInterfaceRasterizationGrid3D.h
--- a/Engine/Plugins/FX/Niagara/Source/Niagara/Classes/NiagaraDataInterfaceRasterizationGrid3D.h
+++ b/Engine/Plugins/FX/Niagara/Source/Niagara/Classes/NiagaraDataInterfaceRasterizationGrid3D.h
@@ -42,9 +42,6 @@ struct FNiagaraDataInterfaceProxyRasterizationGrid3D : public FNiagaraDataInterf
 	virtual void PreStage(const FNDIGpuComputePreStageContext& Context) override;
 	virtual void PostSimulate(const FNDIGpuComputePostSimulateContext& Context) override;
 
-	virtual void ConsumePerInstanceDataFromGameThread(void* PerInstanceData, const FNiagaraSystemInstanceID& Instance) override {}
-	virtual int32 PerInstanceDataPassedToRenderThreadSize() const override { return sizeof(RasterizationGrid3DRWInstanceData); }	
-
 	virtual void GetDispatchArgs(const FNDIGpuComputeDispatchArgsGenContext& Context) override;
 
 	/* List of proxy data for each system instances*/
diff --git a/Engine/Plugins/FX/Niagara/Source/Niagara/Private/NiagaraDataInterfaceRasterizationGrid3D.cpp b/Engine/Plugins/FX/Niagara/Source/Niagara/Private/NiagaraDataInterfaceRasterizationGrid3D.cpp
--- a/Engine/Plugins/FX/Niagara/Source/Niagara/Private/NiagaraDataInterfaceRasterizationGrid3D.cpp
+++ b/Engine/Plugins/FX/Niagara/Source/Niagara/Private/NiagaraDataInterfaceRasterizationGrid3D.cpp
@@ -1034,7 +1034,7 @@ bool UNiagaraDataInterfaceRasterizationGrid3D::PerInstanceTickPostSimulate(void*
 			[RT_Proxy, RT_NumCells = InstanceData->NumCells, RT_Precision = InstanceData->Precision, RT_ResetValue = InstanceData->ResetValue, InstanceID = SystemInstance->GetId(), RT_InstanceData = *InstanceData](FRHICommandListImmediate& RHICmdList)
 		{
 			check(RT_Proxy->SystemInstancesToProxyData.Contains(InstanceID));
-			RasterizationGrid3DRWInstanceData* TargetData = &RT_Proxy->SystemInstancesToProxyData.Add(InstanceID);
+			RasterizationGrid3DRWInstanceData* TargetData = RT_Proxy->SystemInstancesToProxyData.Find(InstanceID);
 
 			TargetData->NumTiles = RT_InstanceData.NumTiles;
 			TargetData->NumCells = RT_NumCells;					
