- Handle stateless in script stats view
- Clear out results when we get compilation notification, this fixes showing some stale data

#rnx

diff --git a/Engine/Plugins/FX/Niagara/Source/NiagaraEditor/Private/ViewModels/NiagaraScriptStatsViewModel.cpp b/Engine/Plugins/FX/Niagara/Source/NiagaraEditor/Private/ViewModels/NiagaraScriptStatsViewModel.cpp
--- a/Engine/Plugins/FX/Niagara/Source/NiagaraEditor/Private/ViewModels/NiagaraScriptStatsViewModel.cpp
+++ b/Engine/Plugins/FX/Niagara/Source/NiagaraEditor/Private/ViewModels/NiagaraScriptStatsViewModel.cpp
@@ -72,14 +72,18 @@ namespace NiagaraScriptStatsLocal
 			for (UNiagaraScript* Script : NiagaraScripts)
 			{
 				const FNiagaraVMExecutableData& VMData = Script->GetVMExecutableData();
+				if (!OutStatus.IsEmpty())
+				{
+					OutStatus.AppendChar(TEXT('\n'));
+				}
+				OutStatus.Appendf(TEXT("VM: %s\n"), *Script->GetName());
 				if (VMData.LastOpCount > 0)
 				{
-					if (!OutStatus.IsEmpty())
-					{
-						OutStatus.AppendChar(TEXT('\n'));
-					}
-					OutStatus.Appendf(TEXT("VM: %s\n"), *Script->GetName());
-					OutStatus.Appendf(TEXT("- Op Count = %u\n"), VMData.LastOpCount);
+					OutStatus.Appendf(TEXT("- Op Count = %u"), VMData.LastOpCount);
+				}
+				else
+				{
+					OutStatus.Append(TEXT("- Data missing please force a recompile"));
 				}
 			}
 		}
@@ -118,120 +122,137 @@ namespace NiagaraScriptStatsLocal
 
 	struct FNiagaraScriptStats
 	{
-		FNiagaraScriptStats(const FGuid& InGuid, FVersionedNiagaraEmitter NiagaraEmitter, EShaderPlatform ShaderPlatform)
+		FNiagaraScriptStats(const FGuid& InGuid, const FNiagaraEmitterHandle& EmitterHandle, EShaderPlatform ShaderPlatform)
 		{
+			FVersionedNiagaraEmitterData* EmitterData = EmitterHandle.GetEmitterData();
+
 			EmitterId = InGuid;
-			bIsGPU = NiagaraEmitter.GetEmitterData()->SimTarget == ENiagaraSimTarget::GPUComputeSim;
+			bIsGPU = EmitterData && EmitterData->SimTarget == ENiagaraSimTarget::GPUComputeSim;
 		}
 
-		void UpdateStatus(EShaderPlatform ShaderPlatform, FVersionedNiagaraEmitter NiagaraEmitter)
+		void UpdateStatus(EShaderPlatform ShaderPlatform, const FNiagaraEmitterHandle& EmitterHandle)
 		{
-			if ( bIsCompiling )
+			if ( bIsCompiling == false )
 			{
-				if ( bIsGPU )
-				{
-					bHasError = false;
+				return;
+			}
 
-					// Can we kick off our cache yet?
-					if ( ShaderScripts.Num() == 0 )
-					{
-						bool bScriptsValid = true;
-						NiagaraEmitter.GetEmitterData()->ForEachScript(
-							[&](UNiagaraScript* NiagaraScript)
-							{
-								if ( !NiagaraScript->GetVMExecutableDataCompilationId().CompilerVersionID.IsValid() || !NiagaraScript->GetVMExecutableDataCompilationId().BaseScriptCompileHash.IsValid() )
-								{
-									bScriptsValid = false;
-								}
-							}
-						);
+			FVersionedNiagaraEmitterData* EmitterData = EmitterHandle.GetEmitterData();
+			if (EmitterHandle.GetEmitterMode() != ENiagaraEmitterMode::Standard || EmitterData == nullptr)
+			{
+				ResultsString = TEXT("N/A");
+				bIsCompiling = false;
+				return;
+			}
+
+			if ( bIsGPU )
+			{
+				UpdateStatusGPU(ShaderPlatform, EmitterData);
+			}
+			else
+			{
+				TArray<UNiagaraScript*> NiagaraScripts;
+				EmitterData->GetScripts(NiagaraScripts);
+				GetVMScriptStatus(MakeArrayView(NiagaraScripts), bIsCompiling, bHasError, ResultsString);
+			}
+		}
 
-						if (!bScriptsValid)
+		void UpdateStatusGPU(EShaderPlatform ShaderPlatform, FVersionedNiagaraEmitterData* EmitterData)
+		{
+			bHasError = false;
+
+			// Can we kick off our cache yet?
+			if ( ShaderScripts.Num() == 0 )
+			{
+				bool bScriptsValid = true;
+				EmitterData->ForEachScript(
+					[&](UNiagaraScript* NiagaraScript)
+					{
+						if ( !NiagaraScript->GetVMExecutableDataCompilationId().CompilerVersionID.IsValid() || !NiagaraScript->GetVMExecutableDataCompilationId().BaseScriptCompileHash.IsValid() )
 						{
-							return;
+							bScriptsValid = false;
 						}
-
-						NiagaraEmitter.GetEmitterData()->ForEachScript(
-							[&](UNiagaraScript* NiagaraScript)
-							{
-								TArray<TUniquePtr<FNiagaraShaderScript>> NewResources;
-								NiagaraScript->CacheResourceShadersForCooking(ShaderPlatform, NewResources);
-								if (NewResources.Num() > 0)
-								{
-									ensure(NewResources.Num() == 1);
-									FNiagaraShaderScript* ShaderScript = NewResources[0].Release();
-									ShaderScripts.Emplace(ShaderScript);
-								}
-							}
-						);
-						return;
 					}
+				);
 
-					// Scripts are compiling
-					for ( const auto& ShaderScript : ShaderScripts)
+				if (!bScriptsValid)
+				{
+					return;
+				}
+
+				EmitterData->ForEachScript(
+					[&](UNiagaraScript* NiagaraScript)
 					{
-						if ( !ShaderScript->IsCompilationFinished() )
+						TArray<TUniquePtr<FNiagaraShaderScript>> NewResources;
+						NiagaraScript->CacheResourceShadersForCooking(ShaderPlatform, NewResources);
+						if (NewResources.Num() > 0)
 						{
-							return;
+							ensure(NewResources.Num() == 1);
+							FNiagaraShaderScript* ShaderScript = NewResources[0].Release();
+							ShaderScripts.Emplace(ShaderScript);
 						}
-						bHasError |= ShaderScript->GetCompileErrors().Num() > 0;
 					}
+				);
+				return;
+			}
 
-					// We have finished compiling, build results information
-					bIsCompiling = false;
+			// Scripts are compiling
+			for ( const auto& ShaderScript : ShaderScripts)
+			{
+				if ( !ShaderScript->IsCompilationFinished() )
+				{
+					return;
+				}
+				bHasError |= ShaderScript->GetCompileErrors().Num() > 0;
+			}
+
+			// We have finished compiling, build results information
+			bIsCompiling = false;
 
-					ResultsString.Reset();
-					if ( bHasError )
+			ResultsString.Reset();
+			if ( bHasError )
+			{
+				for (const auto& ShaderScript : ShaderScripts)
+				{
+					for ( const FString& ErrorString :  ShaderScript->GetCompileErrors() )
+					{
+						ResultsString.Append(ErrorString);
+						ResultsString.AppendChar(TEXT('\n'));
+					}
+				}
+			}
+			else
+			{
+				for (const auto& ShaderScript : ShaderScripts)
+				{
+					TConstArrayView<FSimulationStageMetaData> SimulationStageMetaData = ShaderScript->GetBaseVMScript()->GetSimulationStageMetaData();
+					for (int32 iSimStageIndex=0; iSimStageIndex < SimulationStageMetaData.Num(); ++iSimStageIndex)
 					{
-						for (const auto& ShaderScript : ShaderScripts)
+						FNiagaraShaderRef Shader = ShaderScript->GetShaderGameThread(iSimStageIndex);
+						if ( Shader.IsValid() )
 						{
-							for ( const FString& ErrorString :  ShaderScript->GetCompileErrors() )
+							if (!ResultsString.IsEmpty())
 							{
-								ResultsString.Append(ErrorString);
 								ResultsString.AppendChar(TEXT('\n'));
 							}
-						}
-					}
-					else
-					{
-						for (const auto& ShaderScript : ShaderScripts)
-						{
-							TConstArrayView<FSimulationStageMetaData> SimulationStageMetaData = ShaderScript->GetBaseVMScript()->GetSimulationStageMetaData();
-							for (int32 iSimStageIndex=0; iSimStageIndex < SimulationStageMetaData.Num(); ++iSimStageIndex)
+							ResultsString.Appendf(TEXT("GPU: %s"), *SimulationStageMetaData[iSimStageIndex].SimulationStageName.ToString());
+							if (Shader->GetNumInstructions() > 0)
 							{
-								FNiagaraShaderRef Shader = ShaderScript->GetShaderGameThread(iSimStageIndex);
-								if ( Shader.IsValid() )
-								{
-									if (!ResultsString.IsEmpty())
-									{
-										ResultsString.AppendChar(TEXT('\n'));
-									}
-									ResultsString.Appendf(TEXT("GPU: %s\n"), *SimulationStageMetaData[iSimStageIndex].SimulationStageName.ToString());
-									if (Shader->GetNumInstructions() > 0)
-									{
-										ResultsString.Appendf(TEXT("- Num Instructions = %d\n"), Shader->GetNumInstructions());
-									}
+								ResultsString.Appendf(TEXT("\n- Num Instructions = %d"), Shader->GetNumInstructions());
+							}
 #if WITH_EDITORONLY_DATA
-									TStringBuilder<2048> ShaderStatsBuilder;
-									const TArray<FGenericShaderStat> ShaderStats = ShaderScript->GetGameThreadShaderMap()->GetShaderStatistics(Shader->GetType(ShaderScript->GetGameThreadShaderMap()->GetPointerTable()));
-									for (const auto& ShaderStat : ShaderStats)
-									{
-										ShaderStatsBuilder << TEXT("- ") << ShaderStat.StatName << TEXT(" = ");
-										Visit([&ShaderStatsBuilder](auto& StoredValue) { ShaderStatsBuilder << StoredValue << "\n"; }, ShaderStat.Value);
-									}
-									ResultsString.Append(ShaderStatsBuilder.ToString());
-#endif // WITH_EDITORONLY_DATA
-								}
+							TStringBuilder<2048> ShaderStatsBuilder;
+							const TArray<FGenericShaderStat> ShaderStats = ShaderScript->GetGameThreadShaderMap()->GetShaderStatistics(Shader->GetType(ShaderScript->GetGameThreadShaderMap()->GetPointerTable()));
+							for (const auto& ShaderStat : ShaderStats)
+							{
+								ShaderStatsBuilder << TEXT("\n- ") << ShaderStat.StatName << TEXT(" = ");
+								Visit([&ShaderStatsBuilder](auto& StoredValue) { ShaderStatsBuilder << StoredValue; }, ShaderStat.Value);
 							}
+							ResultsString.Append(ShaderStatsBuilder.ToString());
+#endif // WITH_EDITORONLY_DATA
 						}
 					}
 				}
-				else
-				{
-					TArray<UNiagaraScript*> NiagaraScripts;
-					NiagaraEmitter.GetEmitterData()->GetScripts(NiagaraScripts);
-					GetVMScriptStatus(MakeArrayView(NiagaraScripts), bIsCompiling, bHasError, ResultsString);
-				}
 			}
 		}
 
@@ -818,11 +839,11 @@ FNiagaraScriptStatsViewModel::FGridCellDetails FNiagaraScriptStatsViewModel::Get
 	if (EmitterScriptsPtr == nullptr)
 	{
 		EmitterScriptsPtr = &AllEmitterScripts.AddDefaulted_GetRef();
-		EmitterScriptsPtr->Reset(new NiagaraScriptStatsLocal::FNiagaraScriptStats(RowID, EmitterHandle.GetInstance(), ShaderPlatform));
+		EmitterScriptsPtr->Reset(new NiagaraScriptStatsLocal::FNiagaraScriptStats(RowID, EmitterHandle, ShaderPlatform));
 	}
 	NiagaraScriptStatsLocal::FNiagaraScriptStats* EmitterScripts = EmitterScriptsPtr->Get();
 
-	EmitterScripts->UpdateStatus(ShaderPlatform, EmitterHandle.GetInstance());
+	EmitterScripts->UpdateStatus(ShaderPlatform, EmitterHandle);
 
 	Details.CellText = FText::FromString(EmitterScripts->GetCellString());
 	Details.CellColor = EmitterScripts->GetCellColor();
@@ -935,6 +956,7 @@ void FNiagaraScriptStatsViewModel::OnForceRecompile()
 
 void FNiagaraScriptStatsViewModel::OnSystemCompiled()
 {
+	PerPlatformScripts.Empty();
 	RefreshView();
 }
 
