- Fix for user variables being removed and renderers not updating correctly
- The variable was without namespace so failed to match

#tests Create a user bool, bind it to renderer enabled, delete the variable

diff --git a/Engine/Plugins/FX/Niagara/Source/NiagaraEditor/Private/Customizations/UserParameters/NiagaraSystemUserParametersBuilder.cpp b/Engine/Plugins/FX/Niagara/Source/NiagaraEditor/Private/Customizations/UserParameters/NiagaraSystemUserParametersBuilder.cpp
--- a/Engine/Plugins/FX/Niagara/Source/NiagaraEditor/Private/Customizations/UserParameters/NiagaraSystemUserParametersBuilder.cpp
+++ b/Engine/Plugins/FX/Niagara/Source/NiagaraEditor/Private/Customizations/UserParameters/NiagaraSystemUserParametersBuilder.cpp
@@ -332,12 +332,16 @@ void FNiagaraSystemUserParameterBuilder::DeleteParameter(FNiagaraVariable UserPa
 	FScopedTransaction Transaction(LOCTEXT("DeleteParameterTransaction", "Deleted user parameter"));
 	System->Modify();
 	
+	// UserParameter does not contain the User. namespace at this point
 	if(System->GetExposedParameters().IndexOf(UserParameter) != INDEX_NONE)
 	{
 		System->GetExposedParameters().RemoveParameter(UserParameter);		
 	}
 
 	Cast<UNiagaraSystemEditorData>(System->GetEditorData())->UserParameterHierarchy->Modify();
+
+	// Make sure we are in the User. namespace
+	FNiagaraUserRedirectionParameterStore::MakeUserVariable(UserParameter);
 	System->HandleVariableRemoved(UserParameter, true);
 
 	UserParamToWidgetMap.Empty();
