- Add missing null check for the graph source that can happen during undo

diff --git a/Engine/Plugins/FX/Niagara/Source/NiagaraEditor/Private/ViewModels/NiagaraSystemViewModel.cpp b/Engine/Plugins/FX/Niagara/Source/NiagaraEditor/Private/ViewModels/NiagaraSystemViewModel.cpp
--- a/Engine/Plugins/FX/Niagara/Source/NiagaraEditor/Private/ViewModels/NiagaraSystemViewModel.cpp
+++ b/Engine/Plugins/FX/Niagara/Source/NiagaraEditor/Private/ViewModels/NiagaraSystemViewModel.cpp
@@ -1523,27 +1523,26 @@ TArray<UNiagaraGraph*> FNiagaraSystemViewModel::GetAllGraphs()
 		for (const TSharedRef<FNiagaraEmitterHandleViewModel>& EmitterHandleViewModel : EmitterHandleViewModels)
 		{
 			FNiagaraEmitterHandle* EmitterHandle = EmitterHandleViewModel->GetEmitterHandle();
-			//-TODO:Stateless: Do we need stateless support here?
 			FVersionedNiagaraEmitterData* EmitterData = EmitterHandle ? EmitterHandle->GetEmitterData() : nullptr;
-			if (EmitterData == nullptr)
-			{
-				continue;
-			}
-			UNiagaraGraph* Graph = Cast<UNiagaraScriptSource>(EmitterData->GraphSource)->NodeGraph;
-			if (Graph)
+			if (EmitterData && EmitterData->GraphSource)
 			{
-				OutGraphs.Add(Graph);
+				if (UNiagaraGraph* Graph = Cast<UNiagaraScriptSource>(EmitterData->GraphSource)->NodeGraph)
+				{
+					OutGraphs.Add(Graph);
+				}
 			}
 		}
 	}
 	else
 	{
 		FNiagaraEmitterHandle* EmitterHandle = GetEmitterHandleViewModels()[0]->GetEmitterHandle();
-		//-TODO:Stateless: Do we need stateless support here?
 		FVersionedNiagaraEmitterData* EmitterData = EmitterHandle ? EmitterHandle->GetEmitterData() : nullptr;
-		if (EmitterData != nullptr)
+		if (EmitterData && EmitterData->GraphSource)
 		{
-			OutGraphs.Add(Cast<UNiagaraScriptSource>(EmitterData->GraphSource)->NodeGraph);
+			if (UNiagaraGraph* Graph = Cast<UNiagaraScriptSource>(EmitterData->GraphSource)->NodeGraph)
+			{
+				OutGraphs.Add(Graph);
+			}
 		}
 	}
 
