- Niagara Editor Preview Actor
- Used for previewing and tuning Niagara Systems without the need to run PIE
- Has some simple default motion options which can be overridden by a derived BP for more complex behavior

diff --git a/Engine/Plugins/FX/Niagara/Source/Niagara/Private/NiagaraEditorPreviewActor.cpp b/Engine/Plugins/FX/Niagara/Source/Niagara/Private/NiagaraEditorPreviewActor.cpp
new file mode 100644
--- /dev/null
+++ b/Engine/Plugins/FX/Niagara/Source/Niagara/Private/NiagaraEditorPreviewActor.cpp
@@ -0,0 +1,237 @@
+// Copyright Epic Games, Inc. All Rights Reserved.
+
+
+#include "NiagaraEditorPreviewActor.h"
+
+#include "NiagaraComponent.h"
+#include "NiagaraSystemInstanceController.h"
+
+#include "UObject/ConstructorHelpers.h"
+#include "Components/ArrowComponent.h"
+
+#include UE_INLINE_GENERATED_CPP_BY_NAME(NiagaraEditorPreviewActor)
+
+namespace NiagaraEditorPreviewActorPrivate
+{
+	TOptional<FVector> InterpolatePoints(float ShapeU, TConstArrayView<FVector> ShapePoints, ENiagaraEditorPreviewActorPlaybackType PlaybackType, double Tension)
+	{
+		const int32 NumPoints = ShapePoints.Num();
+		if (NumPoints == 0)
+		{
+			return {};
+		}
+
+		const int32 NumPointsMinusOne = ShapePoints.Num() - 1;
+
+		FVector Points[4];
+		double Interp = 0.0f;
+		{
+			switch (PlaybackType)
+			{
+				default:
+				case ENiagaraEditorPreviewActorPlaybackType::Once:
+				{
+					const double PointU = ShapeU * double(NumPointsMinusOne);
+					const int32 Point = FMath::FloorToInt(PointU);
+					Points[0] = ShapePoints[FMath::Clamp(Point - 1, 0, NumPointsMinusOne)];
+					Points[1] = ShapePoints[FMath::Clamp(Point    , 0, NumPointsMinusOne)];
+					Points[2] = ShapePoints[FMath::Clamp(Point + 1, 0, NumPointsMinusOne)];
+					Points[3] = ShapePoints[FMath::Clamp(Point + 2, 0, NumPointsMinusOne)];
+					Interp = FMath::Fractional(PointU);
+					break;
+				}
+				case ENiagaraEditorPreviewActorPlaybackType::Looping:
+				{
+					const double PointU = ShapeU * double(NumPoints);
+					const int32 Point = FMath::FloorToInt(PointU) + NumPoints;
+					Points[0] = ShapePoints[FMath::Modulo(Point - 1, NumPoints)];
+					Points[1] = ShapePoints[FMath::Modulo(Point    , NumPoints)];
+					Points[2] = ShapePoints[FMath::Modulo(Point + 1, NumPoints)];
+					Points[3] = ShapePoints[FMath::Modulo(Point + 2, NumPoints)];
+					Interp = FMath::Fractional(PointU);
+					break;
+				}
+				case ENiagaraEditorPreviewActorPlaybackType::PingPong:
+				{
+					const double PointU = ShapeU * double(NumPointsMinusOne);
+					const int32 Point = FMath::FloorToInt(PointU);
+					Points[0] = ShapePoints[FMath::Abs(FMath::Modulo(Point - 1, NumPoints))];
+					Points[1] = ShapePoints[FMath::Abs(FMath::Modulo(Point    , NumPoints))];
+					Points[2] = ShapePoints[FMath::Abs(FMath::Modulo(Point + 1, NumPoints))];
+					Points[3] = ShapePoints[FMath::Abs(FMath::Modulo(Point + 2, NumPoints))];
+					Interp = FMath::Fractional(PointU);
+					break;
+				}
+			}
+		}
+
+		const double T = FMath::Fractional(Interp);
+		const double T2 = T * T;
+		const double T3 = T2 * T;
+		const double C1 =  2.0 * T3 - 3.0 * T2 + 1.0;
+		const double C2 = -2.0 * T3 + 3.0 * T2;
+		const double C3 = T3 - 2.0 * T2 + T;
+		const double C4 = T3 - T2;
+
+		FVector Temp;
+		Temp  = (C1 * Points[1]) + (C2 * Points[2]);
+		Temp += Tension * C3 * (Points[2] - Points[0]);
+		Temp += Tension * C4 * (Points[3] - Points[1]);
+		return Temp * 0.5;
+	}
+}
+
+ANiagaraEditorPreviewActor::ANiagaraEditorPreviewActor(const FObjectInitializer& ObjectInitializer)
+	: Super(ObjectInitializer)
+{
+	ArrowComponent = CreateEditorOnlyDefaultSubobject<UArrowComponent>(TEXT("ArrowComponent0"));
+	RootComponent = ArrowComponent;
+
+	NiagaraComponent = CreateDefaultSubobject<UNiagaraComponent>(TEXT("NiagaraComponent0"));
+	NiagaraComponent->SetupAttachment(ArrowComponent);
+	NiagaraComponent->SetAllowScalability(false);
+
+	PrimaryActorTick.bCanEverTick = true;
+	PrimaryActorTick.bStartWithTickEnabled = true;
+
+	bIsEditorOnlyActor = true;
+
+	CustomShapePoints.Emplace(-200.0f, 0.0f, 0.0f);
+	CustomShapePoints.Emplace( 200.0f, 0.0f, 0.0f);
+}
+
+void ANiagaraEditorPreviewActor::PostRegisterAllComponents()
+{
+	Super::PostRegisterAllComponents();
+
+	//if (NiagaraComponent)
+	//{
+	//	NiagaraComponent->OnSystemFinished.AddUniqueDynamic(this, &ANiagaraActor::OnNiagaraSystemFinished);
+	//}
+}
+
+void ANiagaraEditorPreviewActor::PostUnregisterAllComponents()
+{
+	Super::PostUnregisterAllComponents();
+
+	//// Clear Notification Delegate
+	//if (NiagaraComponent)
+	//{
+	//	NiagaraComponent->OnSystemFinished.RemoveAll(this);
+	//}
+}
+
+void ANiagaraEditorPreviewActor::Tick(float DeltaSeconds)
+{
+	Super::Tick(DeltaSeconds);
+
+	UWorld* World = GetWorld();
+	if (World == nullptr)
+	{
+		return;
+	}
+
+	if ( NiagaraComponent->IsActive() == false )
+	{
+		NiagaraComponent->Activate();
+	}
+
+	FNiagaraSystemInstanceControllerConstPtr SystemInstanceController = NiagaraComponent->GetSystemInstanceController();
+	if (SystemInstanceController == nullptr)
+	{
+		return;
+	}
+
+	const float CurrentTime = SystemInstanceController->GetAge();
+	float MotionTime = CurrentTime / MotionDuration;
+
+	switch (PlaybackType)
+	{
+		case ENiagaraEditorPreviewActorPlaybackType::Once:
+			MotionTime = FMath::Clamp(MotionTime, 0.0f, 1.0f);
+			break;
+		case ENiagaraEditorPreviewActorPlaybackType::Looping:
+			MotionTime = FMath::Fractional(MotionTime);
+			break;
+		case ENiagaraEditorPreviewActorPlaybackType::PingPong:
+			MotionTime = ((FMath::FloorToInt(MotionTime) & 1) == 0) ? FMath::Fractional(MotionTime) : 1.0f - FMath::Fractional(MotionTime);
+			break;
+	}
+
+	TOptional<FVector> NewLocation;
+	switch (MotionType)
+	{
+		case ENiagaraEditorPreviewActorShapeType::Square:
+		{
+			const FVector2D HSize = SquareSize * 0.5;
+			const FVector Points[] =
+			{
+				FVector(-HSize.X, -HSize.Y, 0.0),
+				FVector( HSize.X, -HSize.Y, 0.0),
+				FVector( HSize.X,  HSize.Y, 0.0),
+				FVector(-HSize.X,  HSize.Y, 0.0),
+			};
+			NewLocation = NiagaraEditorPreviewActorPrivate::InterpolatePoints(MotionTime, Points, PlaybackType, ShapeTension);
+			break;
+		}
+
+		case ENiagaraEditorPreviewActorShapeType::Triangle:
+		{
+			const FVector2D HSize = TriangleSize * 0.5;
+			const FVector Points[] =
+			{
+				FVector(     0.0, -HSize.Y, 0.0),
+				FVector( HSize.X,  HSize.Y, 0.0),
+				FVector(-HSize.X,  HSize.Y, 0.0),
+			};
+			NewLocation = NiagaraEditorPreviewActorPrivate::InterpolatePoints(MotionTime, Points, PlaybackType, ShapeTension);
+			break;
+		}
+
+		case ENiagaraEditorPreviewActorShapeType::Circle:
+		{
+			const double R = CircleRotationRate.Get(1.0);
+			const double S = FMath::Sin(MotionTime * R * UE_TWO_PI);
+			const double C = FMath::Cos(MotionTime * R * UE_TWO_PI);
+
+			const double Radius = FMath::Lerp(CircleRadius, CircleEndRadius.Get(CircleRadius), MotionTime);
+			NewLocation = FVector(S, C, 0.0) * Radius;
+			break;
+		}
+
+		case ENiagaraEditorPreviewActorShapeType::Custom:
+		{
+			NewLocation = NiagaraEditorPreviewActorPrivate::InterpolatePoints(MotionTime, CustomShapePoints, PlaybackType, ShapeTension);
+			break;
+		}
+
+		case ENiagaraEditorPreviewActorShapeType::Blueprint:
+		{
+			FVector TempLocation = FVector::ZeroVector;
+			CalculateLocation(MotionTime, TempLocation);
+			NewLocation = TempLocation;
+			break;
+		}
+	}
+
+	if (NewLocation.IsSet())
+	{
+		const FVector L = ShapeRotation.RotateVector(NewLocation.GetValue()) * ShapeScale;
+		NiagaraComponent->SetRelativeLocation(L);
+	}
+}
+
+#if WITH_EDITOR
+bool ANiagaraEditorPreviewActor::GetReferencedContentObjects(TArray<UObject*>& Objects) const
+{
+	Super::GetReferencedContentObjects(Objects);
+
+	if (UNiagaraSystem* System = NiagaraComponent->GetAsset())
+	{
+		Objects.Add(System);
+	}
+
+	return true;
+}
+#endif // WITH_EDITOR
+
diff --git a/Engine/Plugins/FX/Niagara/Source/Niagara/Public/NiagaraEditorPreviewActor.h b/Engine/Plugins/FX/Niagara/Source/Niagara/Public/NiagaraEditorPreviewActor.h
new file mode 100644
--- /dev/null
+++ b/Engine/Plugins/FX/Niagara/Source/Niagara/Public/NiagaraEditorPreviewActor.h
@@ -0,0 +1,101 @@
+// Copyright Epic Games, Inc. All Rights Reserved.
+
+#pragma once
+
+#include "CoreMinimal.h"
+#include "UObject/ObjectMacros.h"
+#include "GameFramework/Actor.h"
+#include "NiagaraEditorPreviewActor.generated.h"
+
+UENUM()
+enum class ENiagaraEditorPreviewActorPlaybackType
+{
+	/** Playback motion will stop when the duation is hit. */
+	Once,
+	/** Playback motion will loop when the duration is hit. */
+	Looping,
+	/** Playback motion will go back and forth. */
+	PingPong,
+};
+
+UENUM()
+enum class ENiagaraEditorPreviewActorShapeType
+{
+	Circle,
+	Square,
+	Triangle,
+	Custom,
+	Blueprint,
+};
+
+/**
+* Niagara Particle System Actor for previewing in the editor will not exist in a cooked packagge or PIE
+*/
+UCLASS(ComponentWrapperClass, MinimalAPI, meta = (DisplayName = "Niagara Particle System Editor Preview Actor"))
+class ANiagaraEditorPreviewActor final : public AActor
+{
+	GENERATED_BODY()
+
+protected:
+	NIAGARA_API ANiagaraEditorPreviewActor(const FObjectInitializer& ObjectInitializer);
+
+public:
+	NIAGARA_API virtual void PostRegisterAllComponents() override;
+	NIAGARA_API virtual void PostUnregisterAllComponents() override;
+
+	// AActor interface
+	NIAGARA_API virtual void Tick(float DeltaSeconds) override;
+	virtual bool ShouldTickIfViewportsOnly() const override { return true; }
+#if WITH_EDITOR
+	NIAGARA_API virtual bool GetReferencedContentObjects(TArray<UObject*>& Objects) const override;
+#endif // WITH_EDITOR
+	// End of AActor interface
+
+	UFUNCTION(BlueprintImplementableEvent, BlueprintCallable, CallInEditor, Category = ActorMotion)
+	void CalculateLocation(float MotionTime, FVector& OutLocation);
+
+private:
+	/** The time it takes for us to complete a cycle of the motion in seconds. */
+	UPROPERTY(EditAnywhere, Category = ActorMotion)
+	double MotionDuration = 5.0f;
+
+	UPROPERTY(EditAnywhere, Category = ActorMotion, meta = (SegmentedDisplay))
+	ENiagaraEditorPreviewActorPlaybackType PlaybackType = ENiagaraEditorPreviewActorPlaybackType::Looping;
+
+	/** What motion type we want to preview view */
+	UPROPERTY(EditAnywhere, Category = ActorMotion, meta = (SegmentedDisplay))
+	ENiagaraEditorPreviewActorShapeType MotionType = ENiagaraEditorPreviewActorShapeType::Circle;
+
+	UPROPERTY(EditAnywhere, Category = ActorMotion, meta = (EditCondition = "MotionType != ENiagaraEditorPreviewActorShapeType::Circle", EditConditionHides))
+	double ShapeTension = 0.5;
+
+	UPROPERTY(EditAnywhere, Category = ActorMotion)
+	double ShapeScale = 1.0;
+
+	UPROPERTY(EditAnywhere, Category = ActorMotion)
+	FRotator ShapeRotation = FRotator::ZeroRotator;
+
+	UPROPERTY(EditAnywhere, Category = ActorMotion, meta = (EditCondition = "MotionType == ENiagaraEditorPreviewActorShapeType::Custom", EditConditionHides, DisplayName = "Shape Points"))
+	TArray<FVector> CustomShapePoints;
+
+	UPROPERTY(EditAnywhere, Category = ActorMotion, meta = (EditCondition = "MotionType == ENiagaraEditorPreviewActorShapeType::Circle", EditConditionHides))
+	double CircleRadius = 200.0;
+
+	UPROPERTY(EditAnywhere, Category = ActorMotion, meta = (EditCondition = "MotionType == ENiagaraEditorPreviewActorShapeType::Circle", EditConditionHides))
+	TOptional<double> CircleEndRadius;
+
+	UPROPERTY(EditAnywhere, Category = ActorMotion, meta = (EditCondition = "MotionType == ENiagaraEditorPreviewActorShapeType::Circle", EditConditionHides))
+	TOptional<double> CircleRotationRate;
+
+	UPROPERTY(EditAnywhere, Category = ActorMotion, meta = (EditCondition = "MotionType == ENiagaraEditorPreviewActorShapeType::Square", EditConditionHides))
+	FVector2D SquareSize = FVector2D(400.0, 400.0);
+
+	UPROPERTY(EditAnywhere, Category = ActorMotion, meta = (EditCondition = "MotionType == ENiagaraEditorPreviewActorShapeType::Triangle", EditConditionHides))
+	FVector2D TriangleSize = FVector2D(400.0, 400.0);
+
+	UPROPERTY(VisibleAnywhere, BlueprintReadOnly, Category = NiagaraActor, meta = (AllowPrivateAccess = "true"))
+	TObjectPtr<class UNiagaraComponent> NiagaraComponent;
+
+	UPROPERTY()
+	TObjectPtr<class UArrowComponent> ArrowComponent;
+};
diff --git a/Engine/Plugins/FX/Niagara/Source/NiagaraEditor/Private/Customizations/NiagaraEditorPreviewActorCustomization.cpp b/Engine/Plugins/FX/Niagara/Source/NiagaraEditor/Private/Customizations/NiagaraEditorPreviewActorCustomization.cpp
new file mode 100644
--- /dev/null
+++ b/Engine/Plugins/FX/Niagara/Source/NiagaraEditor/Private/Customizations/NiagaraEditorPreviewActorCustomization.cpp
@@ -0,0 +1,38 @@
+// Copyright Epic Games, Inc. All Rights Reserved.
+
+#include "NiagaraEditorPreviewActorCustomization.h"
+#include "NiagaraEditorPreviewActor.h"
+
+#include "DetailLayoutBuilder.h"
+
+TSharedRef<IDetailCustomization> FNiagaraEditorPreviewActorCustomization::MakeInstance()
+{
+	return MakeShared<FNiagaraEditorPreviewActorCustomization>();
+}
+
+FName FNiagaraEditorPreviewActorCustomization::GetCustomizationTypeName()
+{
+	return ANiagaraEditorPreviewActor::StaticClass()->GetFName();
+}
+
+void FNiagaraEditorPreviewActorCustomization::CustomizeDetails(IDetailLayoutBuilder& DetailLayout)
+{
+	const TArray<FName> CategoryOrder =
+	{
+		FName("ActorMotion"),
+		FName("Niagara"),
+		FName("NiagaraComponent_Parameters"),
+		FName("Niagara Utilities"),
+	};
+
+	TArray<FName> CategoryNames;
+	DetailLayout.GetCategoryNames(CategoryNames);
+
+	for (FName CategoryName : CategoryNames)
+	{
+		if (!CategoryOrder.Contains(CategoryName))
+		{
+			DetailLayout.HideCategory(CategoryName);
+		}
+	}
+}
diff --git a/Engine/Plugins/FX/Niagara/Source/NiagaraEditor/Private/Customizations/NiagaraEditorPreviewActorCustomization.h b/Engine/Plugins/FX/Niagara/Source/NiagaraEditor/Private/Customizations/NiagaraEditorPreviewActorCustomization.h
new file mode 100644
--- /dev/null
+++ b/Engine/Plugins/FX/Niagara/Source/NiagaraEditor/Private/Customizations/NiagaraEditorPreviewActorCustomization.h
@@ -0,0 +1,17 @@
+// Copyright Epic Games, Inc. All Rights Reserved.
+
+#pragma once
+
+#include "IDetailCustomization.h"
+
+class IDetailLayoutBuilder;
+
+class FNiagaraEditorPreviewActorCustomization : public IDetailCustomization
+{
+public:
+	static TSharedRef<IDetailCustomization> MakeInstance();
+	static FName GetCustomizationTypeName();
+
+	virtual void CustomizeDetails(IDetailLayoutBuilder& InDetailLayout) override;
+};
+
diff --git a/Engine/Plugins/FX/Niagara/Source/NiagaraEditor/Private/NiagaraEditorModule.cpp b/Engine/Plugins/FX/Niagara/Source/NiagaraEditor/Private/NiagaraEditorModule.cpp
--- a/Engine/Plugins/FX/Niagara/Source/NiagaraEditor/Private/NiagaraEditorModule.cpp
+++ b/Engine/Plugins/FX/Niagara/Source/NiagaraEditor/Private/NiagaraEditorModule.cpp
@@ -110,6 +110,7 @@
 #include "Sections/ColorPropertySection.h"
 
 #include "Customizations/NiagaraComponentDetails.h"
+#include "Customizations/NiagaraEditorPreviewActorCustomization.h"
 #include "Customizations/NiagaraFunctionCallNodeDetails.h"
 #include "Customizations/NiagaraParameterBindingCustomization.h"
 #include "Customizations/NiagaraPlatformSetCustomization.h"
@@ -1184,6 +1185,10 @@ void FNiagaraEditorModule::StartupModule()
 		UNiagaraComponent::StaticClass()->GetFName(),
 		FOnGetDetailCustomizationInstance::CreateStatic(&FNiagaraComponentDetails::MakeInstance));
 
+	PropertyModule.RegisterCustomClassLayout(
+		FNiagaraEditorPreviewActorCustomization::GetCustomizationTypeName(),
+		FOnGetDetailCustomizationInstance::CreateStatic(&FNiagaraEditorPreviewActorCustomization::MakeInstance));
+
 	PropertyModule.RegisterCustomClassLayout(
 		UNiagaraNodeStaticSwitch::StaticClass()->GetFName(),
 		FOnGetDetailCustomizationInstance::CreateStatic(&FNiagaraStaticSwitchNodeDetails::MakeInstance));
