- Fix for null emitter data / script access

#tests Create a point force module, change the strength to a data interface, also make sure to have a LWE in the system
#rnx

diff --git a/Engine/Plugins/FX/Niagara/Source/NiagaraEditor/Private/ViewModels/Stack/NiagaraStackGraphUtilities.cpp b/Engine/Plugins/FX/Niagara/Source/NiagaraEditor/Private/ViewModels/Stack/NiagaraStackGraphUtilities.cpp
--- a/Engine/Plugins/FX/Niagara/Source/NiagaraEditor/Private/ViewModels/Stack/NiagaraStackGraphUtilities.cpp
+++ b/Engine/Plugins/FX/Niagara/Source/NiagaraEditor/Private/ViewModels/Stack/NiagaraStackGraphUtilities.cpp
@@ -4457,31 +4457,40 @@ void FNiagaraStackGraphUtilities::GetEmitterHandleAndCompiledScriptsForStackNode
 		}
 		else
 		{
-			OutEmitterHandle = OwningSystem.GetEmitterHandles().FindByPredicate([&StackNode](const FNiagaraEmitterHandle& EmitterHandle)
+			OutEmitterHandle = OwningSystem.GetEmitterHandles().FindByPredicate(
+				[&StackNode](const FNiagaraEmitterHandle& EmitterHandle)
 				{
-					return CastChecked<UNiagaraScriptSource>(EmitterHandle.GetEmitterData()->GraphSource)->NodeGraph == StackNode.GetNiagaraGraph();
-				});
+					FVersionedNiagaraEmitterData* EmitterData = EmitterHandle.GetEmitterData();
+					UNiagaraScriptSource* ScriptSource = EmitterData && EmitterData->GraphSource ? CastChecked<UNiagaraScriptSource>(EmitterData->GraphSource) : nullptr;
+					return ScriptSource ? ScriptSource->NodeGraph == StackNode.GetNiagaraGraph() : false;
+				}
+			);
 
 			if (OutEmitterHandle != nullptr)
 			{
 				switch (OutputNode->GetUsage())
 				{
-				case ENiagaraScriptUsage::EmitterSpawnScript:
-					OutCompiledScripts.Add(OwningSystem.GetSystemSpawnScript());
-					break;
-				case ENiagaraScriptUsage::EmitterUpdateScript:
-					OutCompiledScripts.Add(OwningSystem.GetSystemUpdateScript());
-					break;
-				case ENiagaraScriptUsage::ParticleSpawnScript:
-				case ENiagaraScriptUsage::ParticleUpdateScript:
-				case ENiagaraScriptUsage::ParticleEventScript:
-				case ENiagaraScriptUsage::ParticleSimulationStageScript:
-					OutCompiledScripts.Add(OutEmitterHandle->GetEmitterData()->GetScript(OutputNode->GetUsage(), OutputNode->GetUsageId()));
-					if (OutEmitterHandle->GetEmitterData()->SimTarget == ENiagaraSimTarget::GPUComputeSim)
+					case ENiagaraScriptUsage::EmitterSpawnScript:
+						OutCompiledScripts.Add(OwningSystem.GetSystemSpawnScript());
+						break;
+					case ENiagaraScriptUsage::EmitterUpdateScript:
+						OutCompiledScripts.Add(OwningSystem.GetSystemUpdateScript());
+						break;
+					case ENiagaraScriptUsage::ParticleSpawnScript:
+					case ENiagaraScriptUsage::ParticleUpdateScript:
+					case ENiagaraScriptUsage::ParticleEventScript:
+					case ENiagaraScriptUsage::ParticleSimulationStageScript:
 					{
-						OutCompiledScripts.Add(OutEmitterHandle->GetEmitterData()->GetScript(ENiagaraScriptUsage::ParticleGPUComputeScript, FGuid()));
+						if ( FVersionedNiagaraEmitterData* EmitterData = OutEmitterHandle->GetEmitterData() )
+						{
+							OutCompiledScripts.Add(EmitterData->GetScript(OutputNode->GetUsage(), OutputNode->GetUsageId()));
+							if (EmitterData->SimTarget == ENiagaraSimTarget::GPUComputeSim)
+							{
+								OutCompiledScripts.Add(EmitterData->GetScript(ENiagaraScriptUsage::ParticleGPUComputeScript, FGuid()));
+							}
+						}
+						break;
 					}
-					break;
 				}
 			}
 		}
