- Fix for color curves not working inside NPCs when using GPU emitters

diff --git a/Engine/Plugins/FX/Niagara/Shaders/Private/NiagaraDataInterfaceCurveTemplate.ush b/Engine/Plugins/FX/Niagara/Shaders/Private/NiagaraDataInterfaceCurveTemplate.ush
--- a/Engine/Plugins/FX/Niagara/Shaders/Private/NiagaraDataInterfaceCurveTemplate.ush
+++ b/Engine/Plugins/FX/Niagara/Shaders/Private/NiagaraDataInterfaceCurveTemplate.ush
@@ -17,7 +17,7 @@ void GetCurveLUTIndices_{ParameterName}(float Time, out uint IndexA, out uint In
 
 float SampleCurveLUT_{ParameterName}(uint index)
 {
-#if {UseStaticBuffer} == 1
+#if {UseUniqueBuffer} == 1
 	return {ParameterName}_CurveLUT[index];
 #else
 	return StaticInputFloat[{ParameterName}_LUTOffset + index];
diff --git a/Engine/Plugins/FX/Niagara/Source/Niagara/Private/NiagaraDataInterfaceCurveBase.cpp b/Engine/Plugins/FX/Niagara/Source/Niagara/Private/NiagaraDataInterfaceCurveBase.cpp
--- a/Engine/Plugins/FX/Niagara/Source/Niagara/Private/NiagaraDataInterfaceCurveBase.cpp
+++ b/Engine/Plugins/FX/Niagara/Source/Niagara/Private/NiagaraDataInterfaceCurveBase.cpp
@@ -452,7 +452,7 @@ bool UNiagaraDataInterfaceCurveBase::Equals(const UNiagaraDataInterface* Other)
 void UNiagaraDataInterfaceCurveBase::CacheStaticBuffers(struct FNiagaraSystemStaticBuffers& StaticBuffers, const FNiagaraVariable& ResolvedVariable, bool bUsedByCPU, bool bUsedByGPU)
 {
 	LUTOffset = INDEX_NONE;
-	if (bUsedByGPU && ResolvedVariable.IsInNameSpace(FNiagaraConstants::UserNamespaceString) == false)
+	if (bUsedByGPU && !ResolvedVariable.IsInExternalNamespace())
 	{
 		LUTOffset = StaticBuffers.AddGpuData(ShaderLUT);
 	}
@@ -500,7 +500,7 @@ void UNiagaraDataInterfaceCurveBase::GetParameterDefinitionHLSL(const FNiagaraDa
 	const TMap<FString, FStringFormatArg> TemplateArgs =
 	{
 		{TEXT("ParameterName"),				ParamInfo.DataInterfaceHLSLSymbol},
-		{TEXT("UseStaticBuffer"),			ParamInfo.IsUserParameter() ? TEXT("1") : TEXT("0")},
+		{TEXT("UseUniqueBuffer"),			ParamInfo.IsExternalParameter() ? TEXT("1") : TEXT("0")},
 		{TEXT("NumElements"),				FString::FromInt(NumElements)},
 		{TEXT("CurveSampleFunctionName"),	GetCurveSampleFunctionName().ToString()},
 	};
diff --git a/Engine/Plugins/FX/Niagara/Source/Niagara/Private/NiagaraTypes.cpp b/Engine/Plugins/FX/Niagara/Source/Niagara/Private/NiagaraTypes.cpp
--- a/Engine/Plugins/FX/Niagara/Source/Niagara/Private/NiagaraTypes.cpp
+++ b/Engine/Plugins/FX/Niagara/Source/Niagara/Private/NiagaraTypes.cpp
@@ -2,6 +2,7 @@
 
 #include "NiagaraTypes.h"
 
+#include "NiagaraConstants.h"
 #include "NiagaraDataInterface.h"
 #include "NiagaraEmitter.h"
 #include "Misc/StringBuilder.h"
@@ -59,6 +60,21 @@ bool FNiagaraVariableBase::ReplaceRootNamespace(const FStringView& ExpectedNames
 	return false;
 }
 
+bool FNiagaraVariableBase::IsInExternalNamespace() const
+{
+	const FNameBuilder NameString(Name);
+	const FStringView NameStringView(NameString);
+
+	for (const FString& Namespace : { FNiagaraConstants::UserNamespaceString, FNiagaraConstants::ParameterCollectionNamespaceString })
+	{
+		if ((NameStringView.Len() > Namespace.Len() + 1) && (NameStringView[Namespace.Len()] == '.') && NameStringView.StartsWith(Namespace))
+		{
+			return true;
+		}
+	}
+	return false;
+}
+
 bool FNiagaraVariableBase::SerializeFromMismatchedTag(const struct FPropertyTag& Tag, FStructuredArchive::FSlot Slot)
 {
 	if (Tag.GetType().IsStruct(FNiagaraVariable::StaticStruct()->GetFName()))
diff --git a/Engine/Plugins/FX/Niagara/Source/Niagara/Public/NiagaraTypes.h b/Engine/Plugins/FX/Niagara/Source/Niagara/Public/NiagaraTypes.h
--- a/Engine/Plugins/FX/Niagara/Source/Niagara/Public/NiagaraTypes.h
+++ b/Engine/Plugins/FX/Niagara/Source/Niagara/Public/NiagaraTypes.h
@@ -1395,6 +1395,12 @@ struct FNiagaraVariableBase
 		return IsInNameSpace(NamespaceString.ToView());
 	}
 #endif
+	
+	/**
+	Check to see if this is an external variable or not.
+	An external variable is defined as one that is not contained within the Niagara system, i.e. a user or parameter collection parameter.
+	*/
+	bool IsInExternalNamespace() const;
 
 	bool Serialize(FArchive& Ar);
 #if WITH_EDITORONLY_DATA
diff --git a/Engine/Plugins/FX/Niagara/Source/NiagaraShader/Private/NiagaraShader.cpp b/Engine/Plugins/FX/Niagara/Source/NiagaraShader/Private/NiagaraShader.cpp
--- a/Engine/Plugins/FX/Niagara/Source/NiagaraShader/Private/NiagaraShader.cpp
+++ b/Engine/Plugins/FX/Niagara/Source/NiagaraShader/Private/NiagaraShader.cpp
@@ -1284,6 +1284,11 @@ void FNiagaraDataInterfaceParamRef::WriteFrozenParameters(FMemoryImageWriter& Wr
 	InParameters.WriteMemoryImageWithDerivedType(Writer, Base->GetShaderStorageType());
 }
 
+bool FNiagaraDataInterfaceGPUParamInfo::IsExternalParameter() const
+{
+	return DataInterfaceHLSLSymbol.StartsWith(TEXT("User_")) || DataInterfaceHLSLSymbol.StartsWith(TEXT("NPC_"));
+}
+
 bool FNiagaraDataInterfaceGPUParamInfo::IsUserParameter() const
 {
 	return DataInterfaceHLSLSymbol.StartsWith(TEXT("User_"));
diff --git a/Engine/Plugins/FX/Niagara/Source/NiagaraShader/Public/NiagaraShared.h b/Engine/Plugins/FX/Niagara/Source/NiagaraShader/Public/NiagaraShared.h
--- a/Engine/Plugins/FX/Niagara/Source/NiagaraShader/Public/NiagaraShared.h
+++ b/Engine/Plugins/FX/Niagara/Source/NiagaraShader/Public/NiagaraShared.h
@@ -233,6 +233,13 @@ struct FNiagaraDataInterfaceGPUParamInfo
 	UPROPERTY()
 	TArray<FNiagaraDataInterfaceGeneratedFunction> GeneratedFunctions;
 
+	/**
+	Check to see if this is an external parameter or not.
+	An external parameter is defined as one that is not contained within the Niagara system, i.e. a user or parameter collection parameter.
+	*/
+	NIAGARASHADER_API bool IsExternalParameter() const;
+
+	UE_DEPRECATED(5.7, "Please use IsExternalParameter as IsUserParameter does not cover NPCs and can result in incorrect behavior")
 	NIAGARASHADER_API bool IsUserParameter() const;
 
 	NIAGARASHADER_API bool Serialize(FArchive& Ar);
