- Add rotation options to preview actor

diff --git a/Engine/Plugins/FX/Niagara/Source/Niagara/Private/NiagaraEditorPreviewActor.cpp b/Engine/Plugins/FX/Niagara/Source/Niagara/Private/NiagaraEditorPreviewActor.cpp
--- a/Engine/Plugins/FX/Niagara/Source/Niagara/Private/NiagaraEditorPreviewActor.cpp
+++ b/Engine/Plugins/FX/Niagara/Source/Niagara/Private/NiagaraEditorPreviewActor.cpp
@@ -142,23 +142,97 @@ void ANiagaraEditorPreviewActor::Tick(float DeltaSeconds)
 		return;
 	}
 
+	// Calculate our playback time
 	const float CurrentTime = SystemInstanceController->GetAge();
-	float MotionTime = CurrentTime / MotionDuration;
+	float CurrMotionTime = CurrentTime / MotionDuration;
+	float NextMotionTime = (CurrentTime + 0.1f) / MotionDuration;
 
 	switch (PlaybackType)
 	{
 		case ENiagaraEditorPreviewActorPlaybackType::Once:
-			MotionTime = FMath::Clamp(MotionTime, 0.0f, 1.0f);
+			CurrMotionTime = FMath::Clamp(CurrMotionTime, 0.0f, 1.0f);
+			NextMotionTime = FMath::Clamp(NextMotionTime, 0.0f, 1.0f);
 			break;
 		case ENiagaraEditorPreviewActorPlaybackType::Looping:
-			MotionTime = FMath::Fractional(MotionTime);
+			CurrMotionTime = FMath::Fractional(CurrMotionTime);
+			NextMotionTime = FMath::Fractional(NextMotionTime);
 			break;
 		case ENiagaraEditorPreviewActorPlaybackType::PingPong:
-			MotionTime = ((FMath::FloorToInt(MotionTime) & 1) == 0) ? FMath::Fractional(MotionTime) : 1.0f - FMath::Fractional(MotionTime);
+			CurrMotionTime = ((FMath::FloorToInt(CurrMotionTime) & 1) == 0) ? FMath::Fractional(CurrMotionTime) : 1.0f - FMath::Fractional(CurrMotionTime);
+			NextMotionTime = ((FMath::FloorToInt(NextMotionTime) & 1) == 0) ? FMath::Fractional(NextMotionTime) : 1.0f - FMath::Fractional(NextMotionTime);
 			break;
 	}
 
-	TOptional<FVector> NewLocation;
+	// Calculate the location of the component
+	TOptional<FVector> CurrLocation = InternalCalculateLocation(CurrMotionTime);
+
+	// Calculate the rotation of the component
+	TOptional<FQuat> CurrRotation;
+	switch ( RotationMode )
+	{
+		case ENiagaraEditorPreviewActorRotationMode::None:
+		{
+			CurrRotation = FQuat::Identity;
+			break;
+		}
+
+		case ENiagaraEditorPreviewActorRotationMode::DirectionOfTravel:
+		{
+			const TOptional<FVector> NextLocation = InternalCalculateLocation(NextMotionTime);
+			if (CurrLocation.IsSet() && NextLocation.IsSet())
+			{
+				FVector XAxis = NextLocation.GetValue() - CurrLocation.GetValue();
+				if (XAxis.Normalize())
+				{
+					FVector ZAxis = FVector::ZAxisVector;
+					FVector YAxis = FVector::CrossProduct(ZAxis, XAxis);
+					XAxis = FVector::CrossProduct(YAxis, ZAxis);
+
+					CurrRotation = FMatrix(XAxis, YAxis, ZAxis, FVector::ZeroVector).ToQuat();
+				}
+			}
+			break;
+		}
+
+		case ENiagaraEditorPreviewActorRotationMode::Blueprint:
+		{
+			FQuat TempRotation = FQuat::Identity;
+			CalculateRotation(NextMotionTime, TempRotation);
+			CurrRotation = TempRotation;
+			break;
+		}
+	}
+
+	// Set on the component
+	if (CurrLocation.IsSet() || CurrRotation.IsSet())
+	{
+		FVector FinalLocation;
+		if (CurrLocation.IsSet())
+		{
+			FinalLocation = ShapeRotation.RotateVector(CurrLocation.GetValue()) * ShapeScale;
+		}
+		else
+		{
+			FinalLocation = NiagaraComponent->GetRelativeLocation();
+		}
+
+		FRotator FinalRotation;
+		if (CurrRotation.IsSet())
+		{
+			FinalRotation = (ShapeRotation.Quaternion() * CurrRotation.GetValue()).Rotator();
+		}
+		else
+		{
+			FinalRotation = NiagaraComponent->GetRelativeRotation();
+		}
+
+		NiagaraComponent->SetRelativeLocationAndRotation(FinalLocation, FinalRotation);
+	}
+}
+
+TOptional<FVector> ANiagaraEditorPreviewActor::InternalCalculateLocation(float MotionTime)
+{
+	TOptional<FVector> OutLocation;
 	switch (MotionType)
 	{
 		case ENiagaraEditorPreviewActorShapeType::Square:
@@ -171,7 +245,7 @@ void ANiagaraEditorPreviewActor::Tick(float DeltaSeconds)
 				FVector( HSize.X,  HSize.Y, 0.0),
 				FVector(-HSize.X,  HSize.Y, 0.0),
 			};
-			NewLocation = NiagaraEditorPreviewActorPrivate::InterpolatePoints(MotionTime, Points, PlaybackType, ShapeTension);
+			OutLocation = NiagaraEditorPreviewActorPrivate::InterpolatePoints(MotionTime, Points, PlaybackType, ShapeTension);
 			break;
 		}
 
@@ -184,7 +258,7 @@ void ANiagaraEditorPreviewActor::Tick(float DeltaSeconds)
 				FVector( HSize.X,  HSize.Y, 0.0),
 				FVector(-HSize.X,  HSize.Y, 0.0),
 			};
-			NewLocation = NiagaraEditorPreviewActorPrivate::InterpolatePoints(MotionTime, Points, PlaybackType, ShapeTension);
+			OutLocation = NiagaraEditorPreviewActorPrivate::InterpolatePoints(MotionTime, Points, PlaybackType, ShapeTension);
 			break;
 		}
 
@@ -195,13 +269,13 @@ void ANiagaraEditorPreviewActor::Tick(float DeltaSeconds)
 			const double C = FMath::Cos(MotionTime * R * UE_TWO_PI);
 
 			const double Radius = FMath::Lerp(CircleRadius, CircleEndRadius.Get(CircleRadius), MotionTime);
-			NewLocation = FVector(S, C, 0.0) * Radius;
+			OutLocation = FVector(S, C, 0.0) * Radius;
 			break;
 		}
 
 		case ENiagaraEditorPreviewActorShapeType::Custom:
 		{
-			NewLocation = NiagaraEditorPreviewActorPrivate::InterpolatePoints(MotionTime, CustomShapePoints, PlaybackType, ShapeTension);
+			OutLocation = NiagaraEditorPreviewActorPrivate::InterpolatePoints(MotionTime, CustomShapePoints, PlaybackType, ShapeTension);
 			break;
 		}
 
@@ -209,16 +283,12 @@ void ANiagaraEditorPreviewActor::Tick(float DeltaSeconds)
 		{
 			FVector TempLocation = FVector::ZeroVector;
 			CalculateLocation(MotionTime, TempLocation);
-			NewLocation = TempLocation;
+			OutLocation = TempLocation;
 			break;
 		}
 	}
 
-	if (NewLocation.IsSet())
-	{
-		const FVector L = ShapeRotation.RotateVector(NewLocation.GetValue()) * ShapeScale;
-		NiagaraComponent->SetRelativeLocation(L);
-	}
+	return OutLocation;
 }
 
 #if WITH_EDITOR
diff --git a/Engine/Plugins/FX/Niagara/Source/Niagara/Public/NiagaraEditorPreviewActor.h b/Engine/Plugins/FX/Niagara/Source/Niagara/Public/NiagaraEditorPreviewActor.h
--- a/Engine/Plugins/FX/Niagara/Source/Niagara/Public/NiagaraEditorPreviewActor.h
+++ b/Engine/Plugins/FX/Niagara/Source/Niagara/Public/NiagaraEditorPreviewActor.h
@@ -28,6 +28,17 @@ enum class ENiagaraEditorPreviewActorShapeType
 	Blueprint,
 };
 
+UENUM()
+enum class ENiagaraEditorPreviewActorRotationMode
+{
+	// Do not apply any orientation changes to the component
+	None,
+	// Orient the component towards the direction of travel
+	DirectionOfTravel,
+	// Orient the component based on a blueprint callback
+	Blueprint,
+};
+
 /**
 * Niagara Particle System Actor for previewing in the editor will not exist in a cooked packagge or PIE
 */
@@ -54,6 +65,12 @@ public:
 	UFUNCTION(BlueprintImplementableEvent, BlueprintCallable, CallInEditor, Category = ActorMotion)
 	void CalculateLocation(float MotionTime, FVector& OutLocation);
 
+	UFUNCTION(BlueprintImplementableEvent, BlueprintCallable, CallInEditor, Category = ActorMotion)
+	void CalculateRotation(float MotionTime, FQuat& OutRotation);
+
+private:
+	TOptional<FVector> InternalCalculateLocation(float MotionTime);
+
 private:
 	/** The time it takes for us to complete a cycle of the motion in seconds. */
 	UPROPERTY(EditAnywhere, Category = ActorMotion)
@@ -93,6 +110,9 @@ private:
 	UPROPERTY(EditAnywhere, Category = ActorMotion, meta = (EditCondition = "MotionType == ENiagaraEditorPreviewActorShapeType::Triangle", EditConditionHides))
 	FVector2D TriangleSize = FVector2D(400.0, 400.0);
 
+	UPROPERTY(EditAnywhere, Category = ActorMotion, meta = (SegmentedDisplay))
+	ENiagaraEditorPreviewActorRotationMode RotationMode = ENiagaraEditorPreviewActorRotationMode::DirectionOfTravel;
+
 	UPROPERTY(VisibleAnywhere, BlueprintReadOnly, Category = NiagaraActor, meta = (AllowPrivateAccess = "true"))
 	TObjectPtr<class UNiagaraComponent> NiagaraComponent;
 
