- Calculate spawn probability per particle for spawn rate, rather than per loop

#virtualized

diff --git a/Engine/Plugins/FX/Niagara/Shaders/Private/Stateless/Template/GenerateCS_PostModules.ush b/Engine/Plugins/FX/Niagara/Shaders/Private/Stateless/Template/GenerateCS_PostModules.ush
--- a/Engine/Plugins/FX/Niagara/Shaders/Private/Stateless/Template/GenerateCS_PostModules.ush
+++ b/Engine/Plugins/FX/Niagara/Shaders/Private/Stateless/Template/GenerateCS_PostModules.ush
@@ -15,6 +15,7 @@ DECLARE_SCALAR_ARRAY(uint,	SpawnInfo_ParticleOffset, NIAGARA_MAX_GPU_SPAWN_INFOS
 DECLARE_SCALAR_ARRAY(uint,	SpawnInfo_UniqueOffset, NIAGARA_MAX_GPU_SPAWN_INFOS);
 DECLARE_SCALAR_ARRAY(float,	SpawnInfo_Time, NIAGARA_MAX_GPU_SPAWN_INFOS);
 DECLARE_SCALAR_ARRAY(float,	SpawnInfo_Rate, NIAGARA_MAX_GPU_SPAWN_INFOS);
+DECLARE_SCALAR_ARRAY(float,	SpawnInfo_Probability, NIAGARA_MAX_GPU_SPAWN_INFOS);
 DECLARE_SCALAR_ARRAY(float,	SpawnInfo_LifetimeScale, NIAGARA_MAX_GPU_SPAWN_INFOS);
 DECLARE_SCALAR_ARRAY(float,	SpawnInfo_LifetimeBias, NIAGARA_MAX_GPU_SPAWN_INFOS);
 
@@ -41,6 +42,7 @@ void StatelessMain(uint3 GroupId : SV_GroupID, int GroupThreadIndex : SV_GroupIn
 			const uint SpawnInfoUniqueOffset	= GET_SCALAR_ARRAY_ELEMENT(SpawnInfo_UniqueOffset, i);
 			const float SpawnInfoTime			= GET_SCALAR_ARRAY_ELEMENT(SpawnInfo_Time, i);
 			const float SpawnInfoRate			= GET_SCALAR_ARRAY_ELEMENT(SpawnInfo_Rate, i);
+			const float SpawnInfoProbability	= GET_SCALAR_ARRAY_ELEMENT(SpawnInfo_Probability, i);
 			const float SpawnInfoLifetimeScale	= GET_SCALAR_ARRAY_ELEMENT(SpawnInfo_LifetimeScale, i);
 			const float SpawnInfoLifetimeBias	= GET_SCALAR_ARRAY_ELEMENT(SpawnInfo_LifetimeBias, i);
 			if ( SpawnInfoIndex < SpawnInfoNumActive )
@@ -50,6 +52,7 @@ void StatelessMain(uint3 GroupId : SV_GroupID, int GroupThreadIndex : SV_GroupIn
 				Particle.Age			= Common_SimulationTime - (SpawnInfoTime + float(SpawnParticleIndex) * SpawnInfoRate);
 				GRandomSeedInternal		= FNiagaraStatelessDefinitions::MakeRandomSeed(Common_RandomSeed, Particle.UniqueID, ModuleSeedOffset++, 0);
 				Particle.Lifetime		= RandomScaleBiasFloat(0, SpawnInfoLifetimeScale, SpawnInfoLifetimeBias);
+				Particle.Lifetime		= RandomFloat(2) <= SpawnInfoProbability ? Particle.Lifetime : -1.0f;
 				break;
 			}
 
diff --git a/Engine/Plugins/FX/Niagara/Source/Niagara/Internal/Stateless/NiagaraStatelessSpawnInfo.h b/Engine/Plugins/FX/Niagara/Source/Niagara/Internal/Stateless/NiagaraStatelessSpawnInfo.h
--- a/Engine/Plugins/FX/Niagara/Source/Niagara/Internal/Stateless/NiagaraStatelessSpawnInfo.h
+++ b/Engine/Plugins/FX/Niagara/Source/Niagara/Internal/Stateless/NiagaraStatelessSpawnInfo.h
@@ -72,6 +72,7 @@ struct FNiagaraStatelessRuntimeSpawnInfo
 	float	SpawnTimeStart	= 0.0f;
 	float	SpawnTimeEnd	= 0.0f;
 	float	Rate			= 0.0f;
+	float	Probability		= 1.0f;
 	int32	Amount			= 0;		// Note: When a burst this is the absolute amount burst, when rate this is the amount over the spawn duration
 	float	LifetimeMin		= 0.0f;
 	float	LifetimeMax		= 0.0f;
diff --git a/Engine/Plugins/FX/Niagara/Source/Niagara/Private/Stateless/NiagaraStatelessEmitterData.cpp b/Engine/Plugins/FX/Niagara/Source/Niagara/Private/Stateless/NiagaraStatelessEmitterData.cpp
--- a/Engine/Plugins/FX/Niagara/Source/Niagara/Private/Stateless/NiagaraStatelessEmitterData.cpp
+++ b/Engine/Plugins/FX/Niagara/Source/Niagara/Private/Stateless/NiagaraStatelessEmitterData.cpp
@@ -123,6 +123,7 @@ uint32 FNiagaraStatelessEmitterData::CalculateActiveParticles(int32 InRandomSeed
 			GET_SCALAR_ARRAY_ELEMENT(SpawnParameters->SpawnInfo_UniqueOffset, GpuSpawnIndex)	= SpawnInfo.UniqueOffset;
 			GET_SCALAR_ARRAY_ELEMENT(SpawnParameters->SpawnInfo_Time, GpuSpawnIndex)			= SpawnTimeStart;
 			GET_SCALAR_ARRAY_ELEMENT(SpawnParameters->SpawnInfo_Rate, GpuSpawnIndex)			= SpawnRate;
+			GET_SCALAR_ARRAY_ELEMENT(SpawnParameters->SpawnInfo_Probability, GpuSpawnIndex)		= SpawnInfo.Probability;
 			GET_SCALAR_ARRAY_ELEMENT(SpawnParameters->SpawnInfo_LifetimeScale, GpuSpawnIndex)	= SpawnInfo.LifetimeMax - SpawnInfo.LifetimeMin;
 			GET_SCALAR_ARRAY_ELEMENT(SpawnParameters->SpawnInfo_LifetimeBias, GpuSpawnIndex)	= SpawnInfo.LifetimeMin;
 			++GpuSpawnIndex;
@@ -139,6 +140,7 @@ uint32 FNiagaraStatelessEmitterData::CalculateActiveParticles(int32 InRandomSeed
 			GET_SCALAR_ARRAY_ELEMENT(SpawnParameters->SpawnInfo_UniqueOffset, GpuSpawnIndex)	= 0;
 			GET_SCALAR_ARRAY_ELEMENT(SpawnParameters->SpawnInfo_Time, GpuSpawnIndex)			= 0.0f;
 			GET_SCALAR_ARRAY_ELEMENT(SpawnParameters->SpawnInfo_Rate, GpuSpawnIndex)			= 0.0f;
+			GET_SCALAR_ARRAY_ELEMENT(SpawnParameters->SpawnInfo_Probability, GpuSpawnIndex)		= 0.0f;
 			GET_SCALAR_ARRAY_ELEMENT(SpawnParameters->SpawnInfo_LifetimeScale, GpuSpawnIndex)	= 0.0f;
 			GET_SCALAR_ARRAY_ELEMENT(SpawnParameters->SpawnInfo_LifetimeBias, GpuSpawnIndex)	= 0.0f;
 			++GpuSpawnIndex;
diff --git a/Engine/Plugins/FX/Niagara/Source/Niagara/Private/Stateless/NiagaraStatelessEmitterInstance.cpp b/Engine/Plugins/FX/Niagara/Source/Niagara/Private/Stateless/NiagaraStatelessEmitterInstance.cpp
--- a/Engine/Plugins/FX/Niagara/Source/Niagara/Private/Stateless/NiagaraStatelessEmitterInstance.cpp
+++ b/Engine/Plugins/FX/Niagara/Source/Niagara/Private/Stateless/NiagaraStatelessEmitterInstance.cpp
@@ -694,11 +694,12 @@ void FNiagaraStatelessEmitterInstance::InitSpawnInfosForLoop(float Initializatio
 	for (FActiveSpawnRate& SpawnInfo : ActiveSpawnRates)
 	{
 		// Unlike stateful emitters we evaluate the spawn probability & rate per loop
+		float SpawnProbability = 1.0f;
 		if (SpawnInfo.SpawnProbability.IsSet())
 		{
-			float SpawnProbability = EvaluateDistribution(SpawnInfo.SpawnProbability.GetValue(), RandomStream, RendererBindings, DefaultSpawnProbability);
+			SpawnProbability = EvaluateDistribution(SpawnInfo.SpawnProbability.GetValue(), RandomStream, RendererBindings, DefaultSpawnProbability);
 			SpawnProbability = FMath::Clamp(SpawnProbability, 0.0f, 1.0f);
-			if (SpawnProbability < RandomStream.FRand())
+			if (FMath::IsNearlyZero(SpawnProbability))
 			{
 				continue;
 			}
@@ -733,8 +734,9 @@ void FNiagaraStatelessEmitterInstance::InitSpawnInfosForLoop(float Initializatio
 				if ( (ExistingInfo.Type == ENiagaraStatelessSpawnInfoType::Rate) &&
 					 (ExistingInfo.Rate == SpawnRate) &&
 					 (ExistingInfo.SpawnTimeEnd == SpawnAgeStart) &&
-					 (ExistingInfo.LifetimeMin == LifetimeMin) &&
-					 (ExistingInfo.LifetimeMax == LifetimeMax)
+					 FMath::IsNearlyEqual(ExistingInfo.Probability, SpawnProbability) &&
+					 FMath::IsNearlyEqual(ExistingInfo.LifetimeMin, LifetimeMin) &&
+					 FMath::IsNearlyEqual(ExistingInfo.LifetimeMax, LifetimeMax)
 				)
 				{
 					if (ExistingInfo.UniqueOffset + ExistingInfo.Amount == UniqueIndexOffset)
@@ -756,6 +758,7 @@ void FNiagaraStatelessEmitterInstance::InitSpawnInfosForLoop(float Initializatio
 				NewSpawnInfo.SpawnTimeEnd	= SpawnAgeEnd;
 				NewSpawnInfo.Rate			= SpawnRate;
 				NewSpawnInfo.Amount			= NumSpawned;
+				NewSpawnInfo.Probability	= SpawnProbability;	
 				NewSpawnInfo.LifetimeMin	= LifetimeMin;
 				NewSpawnInfo.LifetimeMax	= LifetimeMax;
 			}
diff --git a/Engine/Plugins/FX/Niagara/Source/Niagara/Private/Stateless/NiagaraStatelessParticleSimContext.cpp b/Engine/Plugins/FX/Niagara/Source/Niagara/Private/Stateless/NiagaraStatelessParticleSimContext.cpp
--- a/Engine/Plugins/FX/Niagara/Source/Niagara/Private/Stateless/NiagaraStatelessParticleSimContext.cpp
+++ b/Engine/Plugins/FX/Niagara/Source/Niagara/Private/Stateless/NiagaraStatelessParticleSimContext.cpp
@@ -178,6 +178,7 @@ namespace NiagaraStateless
 					const uint32 SpawnInfoUniqueOffset		= GET_SCALAR_ARRAY_ELEMENT(SpawnParameters.SpawnInfo_UniqueOffset, GpuSpawnIndex);
 					const float  SpawnInfoTime				= GET_SCALAR_ARRAY_ELEMENT(SpawnParameters.SpawnInfo_Time, GpuSpawnIndex);
 					const float  SpawnInfoRate				= GET_SCALAR_ARRAY_ELEMENT(SpawnParameters.SpawnInfo_Rate, GpuSpawnIndex);
+					const float  SpawnInfoProbability		= GET_SCALAR_ARRAY_ELEMENT(SpawnParameters.SpawnInfo_Probability, GpuSpawnIndex);
 					const float  SpawnInfoLifetimeScale		= GET_SCALAR_ARRAY_ELEMENT(SpawnParameters.SpawnInfo_LifetimeScale, GpuSpawnIndex);
 					const float  SpawnInfoLifetimeBias		= GET_SCALAR_ARRAY_ELEMENT(SpawnParameters.SpawnInfo_LifetimeBias, GpuSpawnIndex);
 					if (SpawnInfoIndex < SpawnInfoNumActive)
@@ -189,7 +190,8 @@ namespace NiagaraStateless
 						// Write unique index as it's needed for the random operation
 						GetParticleUniqueID()[NumInstances] = Particle_UniqueID;
 
-						Particle_Lifetime		= RandomScaleBiasFloat(NumInstances, 0, SpawnInfoLifetimeScale, SpawnInfoLifetimeBias);
+						Particle_Lifetime = RandomScaleBiasFloat(NumInstances, 0, SpawnInfoLifetimeScale, SpawnInfoLifetimeBias);
+						Particle_Lifetime = RandomFloat(NumInstances, 2) <= SpawnInfoProbability ? Particle_Lifetime : -1.0f;
 						break;
 					}
 
diff --git a/Engine/Plugins/FX/Niagara/Source/NiagaraShader/Internal/Stateless/NiagaraStatelessSimulationShader.h b/Engine/Plugins/FX/Niagara/Source/NiagaraShader/Internal/Stateless/NiagaraStatelessSimulationShader.h
--- a/Engine/Plugins/FX/Niagara/Source/NiagaraShader/Internal/Stateless/NiagaraStatelessSimulationShader.h
+++ b/Engine/Plugins/FX/Niagara/Source/NiagaraShader/Internal/Stateless/NiagaraStatelessSimulationShader.h
@@ -17,6 +17,7 @@ namespace NiagaraStateless
 		SHADER_PARAMETER_SCALAR_ARRAY(uint32,	SpawnInfo_UniqueOffset, [MaxGpuSpawnInfos])
 		SHADER_PARAMETER_SCALAR_ARRAY(float,	SpawnInfo_Time, [MaxGpuSpawnInfos])
 		SHADER_PARAMETER_SCALAR_ARRAY(float,	SpawnInfo_Rate, [MaxGpuSpawnInfos])
+		SHADER_PARAMETER_SCALAR_ARRAY(float,	SpawnInfo_Probability, [MaxGpuSpawnInfos])
 		SHADER_PARAMETER_SCALAR_ARRAY(float,	SpawnInfo_LifetimeScale, [MaxGpuSpawnInfos])
 		SHADER_PARAMETER_SCALAR_ARRAY(float,	SpawnInfo_LifetimeBias, [MaxGpuSpawnInfos])
 	END_SHADER_PARAMETER_STRUCT()
