Unified some parameter deletion and renaming code paths, made User Parameter deletion handle parameter references better

diff --git a/Engine/Plugins/FX/Niagara/Source/Niagara/Public/NiagaraTypes.h b/Engine/Plugins/FX/Niagara/Source/Niagara/Public/NiagaraTypes.h
--- a/Engine/Plugins/FX/Niagara/Source/Niagara/Public/NiagaraTypes.h
+++ b/Engine/Plugins/FX/Niagara/Source/Niagara/Public/NiagaraTypes.h
@@ -1322,7 +1322,7 @@ struct FNiagaraVariableBase
 		return Name;
 	}
 
-	void SetNamespacedName(const FString& InNamespace, FName InVariableName);
+	NIAGARA_API void SetNamespacedName(const FString& InNamespace, FName InVariableName);
 
 	void SetType(const FNiagaraTypeDefinition& InTypeDef)
 	{
diff --git a/Engine/Plugins/FX/Niagara/Source/NiagaraEditor/Private/Customizations/UserParameters/NiagaraSystemUserParametersBuilder.cpp b/Engine/Plugins/FX/Niagara/Source/NiagaraEditor/Private/Customizations/UserParameters/NiagaraSystemUserParametersBuilder.cpp
--- a/Engine/Plugins/FX/Niagara/Source/NiagaraEditor/Private/Customizations/UserParameters/NiagaraSystemUserParametersBuilder.cpp
+++ b/Engine/Plugins/FX/Niagara/Source/NiagaraEditor/Private/Customizations/UserParameters/NiagaraSystemUserParametersBuilder.cpp
@@ -329,20 +329,18 @@ FReply FNiagaraSystemUserParameterBuilder::SummonHierarchyEditor() const
 
 void FNiagaraSystemUserParameterBuilder::DeleteParameter(FNiagaraVariable UserParameter)
 {
+	// We re-add the User namespace if necessary
+	if(!UserParameter.IsInNameSpace(FNiagaraConstants::UserNamespace))
+	{
+		UserParameter.SetNamespacedName(FNiagaraConstants::UserNamespaceString, UserParameter.GetName());
+	}
+	
 	FScopedTransaction Transaction(LOCTEXT("DeleteParameterTransaction", "Deleted user parameter"));
 	System->Modify();
 	
-	// UserParameter does not contain the User. namespace at this point
-	if(System->GetExposedParameters().IndexOf(UserParameter) != INDEX_NONE)
-	{
-		System->GetExposedParameters().RemoveParameter(UserParameter);		
-	}
-
 	Cast<UNiagaraSystemEditorData>(System->GetEditorData())->UserParameterHierarchy->Modify();
 
-	// Make sure we are in the User. namespace
-	FNiagaraUserRedirectionParameterStore::MakeUserVariable(UserParameter);
-	System->HandleVariableRemoved(UserParameter, true);
+	SystemViewModel.Pin()->DeleteParameter(UserParameter);
 
 	UserParamToWidgetMap.Empty();
 	
diff --git a/Engine/Plugins/FX/Niagara/Source/NiagaraEditor/Private/NiagaraEditorUtilities.cpp b/Engine/Plugins/FX/Niagara/Source/NiagaraEditor/Private/NiagaraEditorUtilities.cpp
--- a/Engine/Plugins/FX/Niagara/Source/NiagaraEditor/Private/NiagaraEditorUtilities.cpp
+++ b/Engine/Plugins/FX/Niagara/Source/NiagaraEditor/Private/NiagaraEditorUtilities.cpp
@@ -84,8 +84,11 @@
 #include "ViewModels/NiagaraEmitterHandleViewModel.h"
 #include "ViewModels/NiagaraEmitterViewModel.h"
 #include "ViewModels/NiagaraParameterPanelViewModel.h"
+#include "ViewModels/NiagaraSystemScriptViewModel.h"
 #include "Widgets/SDataHierarchyEditor.h"
 #include "ViewModels/HierarchyEditor/NiagaraScriptParametersHierarchyViewModel.h"
+#include "ViewModels/Stack/NiagaraStackFunctionInput.h"
+#include "ViewModels/Stack/NiagaraStackViewModel.h"
 #include "Widgets/SToolTip.h"
 
 #define LOCTEXT_NAMESPACE "FNiagaraEditorUtilities"
@@ -2266,7 +2269,7 @@ bool FNiagaraEditorUtilities::AddParameter(FNiagaraVariable& NewParameterVariabl
 	return bSuccess;
 }
 
-TArray<FNiagaraVariable> FNiagaraEditorUtilities::GetReferencedUserParametersFromEmitter(TSharedRef<FNiagaraEmitterViewModel> EmitterViewModel)
+TArray<FNiagaraVariable> FNiagaraEditorUtilities::UserParameters::GetReferencedUserParametersFromEmitter(TSharedRef<FNiagaraEmitterViewModel> EmitterViewModel)
 {
 	TArray<FNiagaraVariable> UserParameters;
 	
@@ -2353,61 +2356,96 @@ TArray<FNiagaraVariable> FNiagaraEditorUtilities::GetReferencedUserParametersFro
 	return UserParameters;
 }
 
-TArray<UNiagaraNodeParameterMapGet*> FNiagaraEditorUtilities::GetParameterMapGetNodesWithUserParameter(TSharedRef<FNiagaraEmitterViewModel> EmitterViewModel, FNiagaraVariable UserParameter)
+TArray<UNiagaraNodeParameterMapGet*> FNiagaraEditorUtilities::UserParameters::GetParameterMapGetNodesWithUserParameter(TSharedRef<FNiagaraSystemViewModel> SystemViewModel, FNiagaraVariable UserParameter, const FNiagaraActionScopeConfig& Config)
 {
+	if(Config.SpecificEmitterGuid.IsSet())
+	{
+		ensureMsgf(Config.bIncludeEmitters, TEXT("If you are searching for a specific emitter, bIncludeEmitters needs to be true"));
+	}
+	
 	TArray<UNiagaraNodeParameterMapGet*> Result;
 
-	TSharedRef<FNiagaraScriptGraphViewModel> GraphViewModel = EmitterViewModel->GetSharedScriptViewModel()->GetGraphViewModel();
+	TArray<UEdGraph*> GraphsToInclude;
+	if(Config.bIncludeSystem)
+	{
+		GraphsToInclude.Add(SystemViewModel->GetSystemScriptViewModel()->GetGraphViewModel()->GetGraph());
+	}
 
-	TArray<UNiagaraNodeParameterMapGet*> MapGetNodes;
-	GraphViewModel->GetGraph()->GetNodesOfClass<UNiagaraNodeParameterMapGet>(MapGetNodes);
+	if(Config.bIncludeEmitters)
+	{
+		for(TSharedPtr<FNiagaraEmitterHandleViewModel> EmitterHandleViewModel : SystemViewModel->GetEmitterHandleViewModels())
+		{
+			if(Config.SpecificEmitterGuid.IsSet() && EmitterHandleViewModel->GetId() != Config.SpecificEmitterGuid.GetValue())
+			{
+				continue;
+			}
+			
+			GraphsToInclude.Add(EmitterHandleViewModel->GetEmitterViewModel()->GetSharedScriptViewModel()->GetGraphViewModel()->GetGraph());
+		}
+	}
 
-	for(UNiagaraNodeParameterMapGet* MapGetNode : MapGetNodes)
+	for(UEdGraph* Graph : GraphsToInclude)
 	{
-		TArray<UEdGraphPin*> OutputPins;
-		MapGetNode->GetOutputPins(OutputPins);
+		TArray<UNiagaraNodeParameterMapGet*> MapGetNodes;
+		Graph->GetNodesOfClass<UNiagaraNodeParameterMapGet>(MapGetNodes);
 
-		for(UEdGraphPin* Pin : OutputPins)
+		for(UNiagaraNodeParameterMapGet* MapGetNode : MapGetNodes)
 		{
-			if(!Pin->bOrphanedPin && Pin->PinType.PinSubCategory == UNiagaraNodeParameterMapBase::ParameterPinSubCategory)
+			TArray<UEdGraphPin*> OutputPins;
+			MapGetNode->GetOutputPins(OutputPins);
+
+			for(UEdGraphPin* Pin : OutputPins)
 			{
-				FNiagaraVariable Variable = UEdGraphSchema_Niagara::PinToNiagaraVariable(Pin);
-				FNiagaraParameterHandle Handle(Variable.GetName());
-				if(Handle.IsUserHandle() && Variable == UserParameter)
+				if(!Pin->bOrphanedPin && Pin->PinType.PinSubCategory == UNiagaraNodeParameterMapBase::ParameterPinSubCategory)
 				{
-					Result.AddUnique(MapGetNode);
-					break;
-				}
-			}			
+					FNiagaraVariable Variable = UEdGraphSchema_Niagara::PinToNiagaraVariable(Pin);
+					FNiagaraParameterHandle Handle(Variable.GetName());
+					if(Handle.IsUserHandle() && Variable == UserParameter)
+					{
+						Result.AddUnique(MapGetNode);
+						break;
+					}
+				}			
+			}
 		}
 	}
-
+	
 	return Result;
 }
 
-TArray<FNiagaraUserParameterBinding*> FNiagaraEditorUtilities::GetUserParameterBindingsForUserParameter(TSharedRef<FNiagaraEmitterViewModel> EmitterViewModel, FNiagaraVariable UserParameter)
+TArray<FNiagaraEditorUtilities::UserParameters::FUserParameterBindingInfo> FNiagaraEditorUtilities::UserParameters::GetUserParameterBindingsForUserParameter(TSharedRef<FNiagaraSystemViewModel> SystemViewModel, FNiagaraVariable UserParameter, const FNiagaraActionScopeConfig& Config)
 {
-	TArray<FNiagaraUserParameterBinding*> Bindings;
+	if(Config.SpecificEmitterGuid.IsSet())
+	{
+		ensureMsgf(Config.bIncludeEmitters, TEXT("If you are searching for a specific emitter, bIncludeEmitters needs to be true"));
+	}
+	
+	TArray<FUserParameterBindingInfo> Bindings;
 
-	// First we gather the bindings on renderers of an emitter
-	//-TODO:Stateless
-	if (FVersionedNiagaraEmitterData* EmitterData = EmitterViewModel->GetEmitter().GetEmitterData())
+	if(Config.bIncludeSystem)
 	{
-		const TArray<UNiagaraRendererProperties*>& Renderers = EmitterData->GetRenderers();
-		for(UNiagaraRendererProperties* Renderer : Renderers)
+		TSharedRef<FNiagaraScriptGraphViewModel> GraphViewModel = SystemViewModel->GetSystemScriptViewModel()->GetGraphViewModel();
+
+		TArray<UNiagaraNodeInput*> InputNodes;
+		GraphViewModel->GetGraph()->GetNodesOfClass<UNiagaraNodeInput>(InputNodes);
+
+		for(UNiagaraNodeInput* Input : InputNodes)
 		{
-			for(TFieldIterator<FStructProperty> It(Renderer->GetClass()); It; ++It)
+			if(Input->GetDataInterface() != nullptr)
 			{
-				if((*It)->Struct == StaticStruct<FNiagaraUserParameterBinding>())
-				{				
-					FNiagaraUserParameterBinding* ParameterBinding = (FNiagaraUserParameterBinding*) (*It)->ContainerPtrToValuePtr<void>(Renderer, 0);
-				
-					if(ParameterBinding->Parameter.IsValid())
-					{
-						FNiagaraParameterHandle ReferencedParameter(ParameterBinding->Parameter.GetName());
-						if(ReferencedParameter.IsUserHandle() && ParameterBinding->Parameter == UserParameter)
+				for(TFieldIterator<FStructProperty> It(Input->GetDataInterface()->GetClass()); It; ++It)
+				{
+					if((*It)->Struct == StaticStruct<FNiagaraUserParameterBinding>())
+					{					
+						FNiagaraUserParameterBinding* ParameterBinding = (FNiagaraUserParameterBinding*) (*It)->ContainerPtrToValuePtr<void>(Input->GetDataInterface(), 0);
+
+						if(ParameterBinding->Parameter.IsValid())
 						{
-							Bindings.Add(ParameterBinding);
+							FNiagaraParameterHandle ReferencedParameter(ParameterBinding->Parameter.GetName());
+							if(ReferencedParameter.IsUserHandle() && ParameterBinding->Parameter == UserParameter)
+							{
+								Bindings.Emplace(ParameterBinding, Input->GetDataInterface());
+							}
 						}
 					}
 				}
@@ -2415,28 +2453,65 @@ TArray<FNiagaraUserParameterBinding*> FNiagaraEditorUtilities::GetUserParameterB
 		}
 	}
 
-	// then we get the bindings on DIs which live on input nodes in the emitter graph
-	TSharedRef<FNiagaraScriptGraphViewModel> GraphViewModel = EmitterViewModel->GetSharedScriptViewModel()->GetGraphViewModel();
-
-	TArray<UNiagaraNodeInput*> InputNodes;
-	GraphViewModel->GetGraph()->GetNodesOfClass<UNiagaraNodeInput>(InputNodes);
-
-	for(UNiagaraNodeInput* Input : InputNodes)
+	if(Config.bIncludeEmitters)
 	{
-		if(Input->GetDataInterface() != nullptr)
+		for(TSharedPtr<FNiagaraEmitterHandleViewModel> EmitterHandleViewModel : SystemViewModel->GetEmitterHandleViewModels())
 		{
-			for(TFieldIterator<FStructProperty> It(Input->GetDataInterface()->GetClass()); It; ++It)
+			if(Config.SpecificEmitterGuid.IsSet() && EmitterHandleViewModel->GetId() != Config.SpecificEmitterGuid.GetValue())
 			{
-				if((*It)->Struct == StaticStruct<FNiagaraUserParameterBinding>())
-				{					
-					FNiagaraUserParameterBinding* ParameterBinding = (FNiagaraUserParameterBinding*) (*It)->ContainerPtrToValuePtr<void>(Input->GetDataInterface(), 0);
+				continue;
+			}
+			
+			// First we gather the bindings on renderers of an emitter
+			//-TODO:Stateless
+			if (FVersionedNiagaraEmitterData* EmitterData = EmitterHandleViewModel->GetEmitterViewModel()->GetEmitter().GetEmitterData())
+			{
+				const TArray<UNiagaraRendererProperties*>& Renderers = EmitterData->GetRenderers();
+				for(UNiagaraRendererProperties* Renderer : Renderers)
+				{
+					for(TFieldIterator<FStructProperty> It(Renderer->GetClass()); It; ++It)
+					{
+						if((*It)->Struct == StaticStruct<FNiagaraUserParameterBinding>())
+						{				
+							FNiagaraUserParameterBinding* ParameterBinding = (FNiagaraUserParameterBinding*) (*It)->ContainerPtrToValuePtr<void>(Renderer, 0);
+						
+							if(ParameterBinding->Parameter.IsValid())
+							{
+								FNiagaraParameterHandle ReferencedParameter(ParameterBinding->Parameter.GetName());
+								if(ReferencedParameter.IsUserHandle() && ParameterBinding->Parameter == UserParameter)
+								{
+									Bindings.Emplace(ParameterBinding, Renderer);
+								}
+							}
+						}
+					}
+				}
+			}
 
-					if(ParameterBinding->Parameter.IsValid())
+			// then we get the bindings on DIs which live on input nodes in the emitter graph
+			TSharedRef<FNiagaraScriptGraphViewModel> GraphViewModel = EmitterHandleViewModel->GetEmitterViewModel()->GetSharedScriptViewModel()->GetGraphViewModel();
+
+			TArray<UNiagaraNodeInput*> InputNodes;
+			GraphViewModel->GetGraph()->GetNodesOfClass<UNiagaraNodeInput>(InputNodes);
+
+			for(UNiagaraNodeInput* Input : InputNodes)
+			{
+				if(Input->GetDataInterface() != nullptr)
+				{
+					for(TFieldIterator<FStructProperty> It(Input->GetDataInterface()->GetClass()); It; ++It)
 					{
-						FNiagaraParameterHandle ReferencedParameter(ParameterBinding->Parameter.GetName());
-						if(ReferencedParameter.IsUserHandle() && ParameterBinding->Parameter == UserParameter)
-						{
-							Bindings.Add(ParameterBinding);
+						if((*It)->Struct == StaticStruct<FNiagaraUserParameterBinding>())
+						{					
+							FNiagaraUserParameterBinding* ParameterBinding = (FNiagaraUserParameterBinding*) (*It)->ContainerPtrToValuePtr<void>(Input->GetDataInterface(), 0);
+
+							if(ParameterBinding->Parameter.IsValid())
+							{
+								FNiagaraParameterHandle ReferencedParameter(ParameterBinding->Parameter.GetName());
+								if(ReferencedParameter.IsUserHandle() && ParameterBinding->Parameter == UserParameter)
+								{
+									Bindings.Emplace(ParameterBinding, Input->GetDataInterface());
+								}
+							}
 						}
 					}
 				}
@@ -2447,32 +2522,93 @@ TArray<FNiagaraUserParameterBinding*> FNiagaraEditorUtilities::GetUserParameterB
 	return Bindings;
 }
 
-TArray<TPair<FNiagaraVariableAttributeBinding*, ENiagaraRendererSourceDataMode>> FNiagaraEditorUtilities::GetVariableAttributeBindingsForParameter(TSharedRef<FNiagaraEmitterViewModel> EmitterViewModel, FNiagaraVariable Parameter)
+TArray<UNiagaraStackFunctionInput*> FNiagaraEditorUtilities::UserParameters::GetStackFunctionInputsWithLinkedParameter(TSharedRef<FNiagaraSystemViewModel> SystemViewModel, FNiagaraVariable Parameter, const FNiagaraActionScopeConfig& Config)
 {
-	TArray<TPair<FNiagaraVariableAttributeBinding*, ENiagaraRendererSourceDataMode>> AttributeBindings;
+	if(Config.SpecificEmitterGuid.IsSet())
+	{
+		ensureMsgf(Config.bIncludeEmitters, TEXT("If you are searching for a specific emitter, bIncludeEmitters needs to be true"));
+	}
 	
-	//-TODO:Stateless
-	if (FVersionedNiagaraEmitterData* EmitterData = EmitterViewModel->GetEmitter().GetEmitterData())
+	TArray<UNiagaraStackFunctionInput*> Result;
+
+	if(Config.bIncludeSystem)
 	{
-		const TArray<UNiagaraRendererProperties*>& Renderers = EmitterData->GetRenderers();
-		for(UNiagaraRendererProperties* Renderer : Renderers)
+		SystemViewModel->GetSystemStackViewModel()->GetRootEntry()->GetUnfilteredChildrenOfType(Result, true);
+	}
+	
+	if(Config.bIncludeEmitters)
+	{
+		//-TODO:Stateless
+		for(TSharedPtr<FNiagaraEmitterHandleViewModel> EmitterHandleViewModel : SystemViewModel->GetEmitterHandleViewModels())
 		{
-			for(TFieldIterator<FStructProperty> It(Renderer->GetClass()); It; ++It)
+			if(Config.SpecificEmitterGuid.IsSet() && EmitterHandleViewModel->GetId() != Config.SpecificEmitterGuid.GetValue())
 			{
-				if((*It)->Struct == StaticStruct<FNiagaraVariableAttributeBinding>())
-				{				
-					FNiagaraVariableAttributeBinding* ParameterBinding = (FNiagaraVariableAttributeBinding*) (*It)->ContainerPtrToValuePtr<void>(Renderer, 0);
-				
-					if(ParameterBinding->GetName().IsEqual(Parameter.GetName()))
+				continue;
+			}
+
+			EmitterHandleViewModel->GetEmitterStackViewModel()->GetRootEntry()->GetUnfilteredChildrenOfType(Result, true);
+		}
+	}
+
+	return Result;
+}
+
+FNiagaraEditorUtilities::FNiagaraActionScopeConfig::FNiagaraActionScopeConfig(FGuid InSpecificEmitterGuid)
+{
+	ensureMsgf(InSpecificEmitterGuid.IsValid(), TEXT("Invalid emitter ID passed into ctor"));
+	
+	bIncludeSystem = false;
+	bIncludeEmitters = true;
+	SpecificEmitterGuid = InSpecificEmitterGuid;
+}
+
+TArray<FNiagaraEditorUtilities::FVariableAttributeBindingInfo> FNiagaraEditorUtilities::GetVariableAttributeBindingsForParameter(TSharedRef<FNiagaraSystemViewModel> SystemViewModel, FNiagaraVariable Parameter, const FNiagaraActionScopeConfig& Config)
+{
+	if(Config.SpecificEmitterGuid.IsSet())
+	{
+		ensureMsgf(Config.bIncludeEmitters, TEXT("If you are searching for a specific emitter, bIncludeEmitters needs to be true"));
+	}
+	
+	TArray<FVariableAttributeBindingInfo> Result;
+
+	if(Config.bIncludeSystem)
+	{
+		// No attribute bindings in systems at this point
+	}
+	if(Config.bIncludeEmitters)
+	{
+		//-TODO:Stateless
+		for(TSharedPtr<FNiagaraEmitterHandleViewModel> EmitterHandleViewModel : SystemViewModel->GetEmitterHandleViewModels())
+		{
+			if(Config.SpecificEmitterGuid.IsSet() && EmitterHandleViewModel->GetId() != Config.SpecificEmitterGuid.GetValue())
+			{
+				continue;
+			}
+			
+			if (FVersionedNiagaraEmitterData* EmitterData = EmitterHandleViewModel->GetEmitterViewModel()->GetEmitter().GetEmitterData())
+			{
+				const TArray<UNiagaraRendererProperties*>& Renderers = EmitterData->GetRenderers();
+				for(UNiagaraRendererProperties* Renderer : Renderers)
+				{
+					for(TFieldIterator<FStructProperty> It(Renderer->GetClass()); It; ++It)
 					{
-						AttributeBindings.Add({ParameterBinding, Renderer->GetCurrentSourceMode()});
+						if((*It)->Struct == StaticStruct<FNiagaraVariableAttributeBinding>())
+						{				
+							FNiagaraVariableAttributeBinding* ParameterBinding = (FNiagaraVariableAttributeBinding*) (*It)->ContainerPtrToValuePtr<void>(Renderer, 0);
+						
+							if(ParameterBinding->GetName().IsEqual(Parameter.GetName()))
+							{
+								FNiagaraVariableAttributeBinding* DefaultParameterBinding = (FNiagaraVariableAttributeBinding*) (*It)->ContainerPtrToValuePtr<void>(Renderer->GetClass()->GetDefaultObject(), 0);
+								Result.Emplace(ParameterBinding, DefaultParameterBinding, Renderer->GetCurrentSourceMode(), Renderer, EmitterHandleViewModel->GetEmitterViewModel()->GetEmitter());
+							}
+						}
 					}
 				}
 			}
 		}
 	}
 
-	return AttributeBindings;
+	return Result;
 }
 
 TObjectPtr<UNiagaraScriptVariable> FNiagaraEditorUtilities::UserParameters::GetScriptVariableForUserParameter(const FNiagaraVariable& UserParameter, TSharedPtr<FNiagaraSystemViewModel> SystemViewModel)
@@ -2490,24 +2626,29 @@ const UNiagaraScriptVariable* FNiagaraEditorUtilities::UserParameters::FindScrip
 	return Cast<UNiagaraSystemEditorData>(System.GetEditorData())->FindUserScriptVariable(UserParameterGuid);
 }
 
-void FNiagaraEditorUtilities::UserParameters::ReplaceUserParameterReferences(TSharedRef<FNiagaraEmitterViewModel> EmitterViewModel,	FNiagaraVariable OldUserParameter, FNiagaraVariable NewUserParameter)
+void FNiagaraEditorUtilities::UserParameters::ReplaceUserParameterReferences(TSharedRef<FNiagaraSystemViewModel> SystemViewModel, FNiagaraVariable OldUserParameter, FNiagaraVariable NewUserParameter, const FNiagaraActionScopeConfig& Config)
 {
-	TArray<FNiagaraUserParameterBinding*> ReferencingBindings = FNiagaraEditorUtilities::GetUserParameterBindingsForUserParameter(EmitterViewModel, OldUserParameter);
-
-	for(FNiagaraUserParameterBinding* ReferencingBinding : ReferencingBindings)
+	using namespace FNiagaraEditorUtilities::UserParameters;
+	
+	for(const FUserParameterBindingInfo& UserParameterBindingInfo : GetUserParameterBindingsForUserParameter(SystemViewModel, OldUserParameter, Config))
 	{
-		ReferencingBinding->Parameter = NewUserParameter;
+		if(UserParameterBindingInfo.Owner)
+		{
+			UserParameterBindingInfo.Owner->Modify();
+		}
+		UserParameterBindingInfo.Binding->Parameter = NewUserParameter;
 	}
 
-	TArray<TPair<FNiagaraVariableAttributeBinding*, ENiagaraRendererSourceDataMode>> ReferencingAttributeBindings = FNiagaraEditorUtilities::GetVariableAttributeBindingsForParameter(EmitterViewModel, OldUserParameter);
-	for(const auto& ReferencingBinding : ReferencingAttributeBindings)
+	for(const FVariableAttributeBindingInfo& AttributeBindingInfo : GetVariableAttributeBindingsForParameter(SystemViewModel, OldUserParameter, Config))
 	{
-		ReferencingBinding.Key->SetValue(NewUserParameter.GetName(), EmitterViewModel->GetEmitter().ToBase(), ReferencingBinding.Value);
+		if(AttributeBindingInfo.Owner)
+		{
+			AttributeBindingInfo.Owner->Modify();
+		}
+		AttributeBindingInfo.Binding->SetValue(NewUserParameter.GetName(), AttributeBindingInfo.OwningEmitter.ToBase(), AttributeBindingInfo.SourceDataMode);
 	}
 	
-	TArray<UNiagaraNodeParameterMapGet*> ReferencingMapGetNodes = FNiagaraEditorUtilities::GetParameterMapGetNodesWithUserParameter(EmitterViewModel, OldUserParameter);
-
-	for(UNiagaraNodeParameterMapGet* MapGet : ReferencingMapGetNodes)
+	for(UNiagaraNodeParameterMapGet* MapGet : GetParameterMapGetNodesWithUserParameter(SystemViewModel, OldUserParameter, Config))
 	{
 		TArray<UEdGraphPin*> OutputPins;
 		MapGet->GetOutputPins(OutputPins);
@@ -2520,6 +2661,7 @@ void FNiagaraEditorUtilities::UserParameters::ReplaceUserParameterReferences(TSh
 
 				if(RepresentingVariable == OldUserParameter)
 				{
+					GraphPin->Modify();
 					GraphPin->PinName = NewUserParameter.GetName();
 				}
 			}
@@ -2527,6 +2669,65 @@ void FNiagaraEditorUtilities::UserParameters::ReplaceUserParameterReferences(TSh
 	}
 }
 
+void FNiagaraEditorUtilities::UserParameters::DeleteUserParameterReferences(TSharedRef<FNiagaraSystemViewModel> SystemViewModel, FNiagaraVariable UserParameterToDelete, const FNiagaraActionScopeConfig& Config)
+{
+	using namespace FNiagaraEditorUtilities::UserParameters;
+	
+	for(const FUserParameterBindingInfo& UserParameterBindingInfo : GetUserParameterBindingsForUserParameter(SystemViewModel, UserParameterToDelete, Config))
+	{
+		if(UserParameterBindingInfo.Owner)
+		{
+			UserParameterBindingInfo.Owner->Modify();
+		}
+		UserParameterBindingInfo.Binding->Parameter = FNiagaraVariable(UserParameterBindingInfo.Binding->Parameter.GetType(), NAME_None);
+	}
+
+	for(const auto& AttributeBindingInfo : GetVariableAttributeBindingsForParameter(SystemViewModel, UserParameterToDelete, Config))
+	{
+		if(AttributeBindingInfo.Owner)
+		{
+			AttributeBindingInfo.Owner->Modify();
+		}
+
+		if(AttributeBindingInfo.OwningEmitter.Emitter != nullptr)
+		{
+			AttributeBindingInfo.Binding->ResetToDefault(*AttributeBindingInfo.DefaultBinding, AttributeBindingInfo.OwningEmitter.ToBase(), AttributeBindingInfo.SourceDataMode);
+		}
+	}
+
+	for(UNiagaraStackFunctionInput* LinkedStackFunctionInput : GetStackFunctionInputsWithLinkedParameter(SystemViewModel, UserParameterToDelete, Config))
+	{
+		if(LinkedStackFunctionInput->GetValueMode() == UNiagaraStackFunctionInput::EValueMode::Linked && LinkedStackFunctionInput->GetLinkedParameterValue() == UserParameterToDelete)
+		{
+			LinkedStackFunctionInput->Reset();
+				
+			if(LinkedStackFunctionInput->GetInputType().IsDataInterface())
+			{
+				UNiagaraDataInterface* CurrentUserDataInterface = SystemViewModel->GetSystem().GetExposedParameters().GetDataInterface(UserParameterToDelete);
+				if(CurrentUserDataInterface->GetClass() == LinkedStackFunctionInput->GetInputType().GetClass())
+				{
+					LinkedStackFunctionInput->SetDataInterfaceValue(CurrentUserDataInterface->GetClass());
+				}
+			}
+			else if (LinkedStackFunctionInput->GetInputType().IsUObject() == false)
+			{
+				const UScriptStruct* LocalValueStruct = LinkedStackFunctionInput->GetInputType().GetScriptStruct();
+				if (LocalValueStruct != nullptr)
+				{
+					TSharedRef<FStructOnScope> LocalValue = MakeShared<FStructOnScope>(LocalValueStruct);
+						
+					const uint8* Memory = SystemViewModel->GetSystem().GetExposedParameters().GetParameterData(UserParameterToDelete);
+					if(Memory)
+					{
+						FMemory::Memcpy(LocalValue->GetStructMemory(), Memory, LocalValueStruct->GetStructureSize());
+						LinkedStackFunctionInput->SetLocalValue(LocalValue);
+					}
+				}
+			}
+		}
+	}
+}
+
 void FNiagaraEditorUtilities::ShowParentEmitterInContentBrowser(TSharedRef<FNiagaraEmitterViewModel> Emitter)
 {
 	FContentBrowserModule& ContentBrowserModule = FModuleManager::Get().LoadModuleChecked<FContentBrowserModule>(TEXT("ContentBrowser"));
diff --git a/Engine/Plugins/FX/Niagara/Source/NiagaraEditor/Private/ViewModels/NiagaraSystemViewModel.cpp b/Engine/Plugins/FX/Niagara/Source/NiagaraEditor/Private/ViewModels/NiagaraSystemViewModel.cpp
--- a/Engine/Plugins/FX/Niagara/Source/NiagaraEditor/Private/ViewModels/NiagaraSystemViewModel.cpp
+++ b/Engine/Plugins/FX/Niagara/Source/NiagaraEditor/Private/ViewModels/NiagaraSystemViewModel.cpp
@@ -800,7 +800,7 @@ void FNiagaraSystemViewModel::DuplicateEmitters(TArray<FEmitterHandleToDuplicate
 		}
 		
 		TSharedPtr<FNiagaraEmitterHandleViewModel> EmitterHandleViewModel = GetEmitterHandleViewModelById(Element.Key);
-		TArray<FNiagaraVariable> ReferencedUserParameters = FNiagaraEditorUtilities::GetReferencedUserParametersFromEmitter(EmitterHandleViewModel->GetEmitterViewModel());
+		TArray<FNiagaraVariable> ReferencedUserParameters = FNiagaraEditorUtilities::UserParameters::GetReferencedUserParametersFromEmitter(EmitterHandleViewModel->GetEmitterViewModel());
 
 		UNiagaraSystem* OriginalSystem = Element.Value;
 		
@@ -834,13 +834,15 @@ void FNiagaraSystemViewModel::DuplicateEmitters(TArray<FEmitterHandleToDuplicate
 					FNiagaraVariable NewUserParameter = FNiagaraVariable(UserParameter.GetType(), FNiagaraUtilities::GetUniqueName(UserParameter.GetName(), ExistingUserParameterNames));
 					TargetParameterStore.AddParameter(NewUserParameter);
 					OriginalSystem->GetExposedParameters().CopyParameterData(TargetParameterStore, UserParameter, NewUserParameter);
-					FNiagaraEditorUtilities::UserParameters::ReplaceUserParameterReferences(EmitterHandleViewModel->GetEmitterViewModel(), UserParameter, NewUserParameter);
+					FNiagaraEditorUtilities::FNiagaraActionScopeConfig Config(Element.Key);
+					FNiagaraEditorUtilities::UserParameters::ReplaceUserParameterReferences(AsShared(), UserParameter, NewUserParameter, Config);
 					RenamedUserParameters.Add(UserParameter, NewUserParameter);
 				}
 				else
 				{
 					FNiagaraVariable RenamedUserParameter = RenamedUserParameters[UserParameter];
-					FNiagaraEditorUtilities::UserParameters::ReplaceUserParameterReferences(EmitterHandleViewModel->GetEmitterViewModel(), UserParameter, RenamedUserParameter);
+					FNiagaraEditorUtilities::FNiagaraActionScopeConfig Config(Element.Key);
+					FNiagaraEditorUtilities::UserParameters::ReplaceUserParameterReferences(AsShared(), UserParameter, RenamedUserParameter, Config);
 				}
 			}
 			// if the parameter doesn't exist at all, we just add it, copy the data and we're done
@@ -1423,10 +1425,12 @@ bool FNiagaraSystemViewModel::RenameParameter(const FNiagaraVariable TargetParam
 
 		FNiagaraVariable NewVariable = TargetParameter;
 		NewVariable.SetName(NewName);
-		for(TSharedPtr<FNiagaraEmitterHandleViewModel> EmitterHandleViewModel : GetEmitterHandleViewModels())
-		{
-			FNiagaraEditorUtilities::UserParameters::ReplaceUserParameterReferences(EmitterHandleViewModel->GetEmitterViewModel(), TargetParameter, NewVariable);
-		}
+
+		FNiagaraEditorUtilities::FNiagaraActionScopeConfig Config;
+		Config.bIncludeSystem = true;
+		Config.bIncludeEmitters = true;
+		FNiagaraEditorUtilities::UserParameters::ReplaceUserParameterReferences(AsShared(), TargetParameter, NewVariable, Config);
+		
 		bExposedParametersRename = true;
 	}
 
@@ -1500,6 +1504,28 @@ bool FNiagaraSystemViewModel::RenameParameter(const FNiagaraVariable TargetParam
 	return false;
 }
 
+bool FNiagaraSystemViewModel::DeleteParameter(const FNiagaraVariableBase Parameter)
+{
+	FScopedTransaction RenameParameterTransaction(LOCTEXT("DeleteParameter", "Delete parameter"));
+
+	System->Modify();
+
+	// Delete the parameter if it's in the parameter store for user parameters.
+	if (System->GetExposedParameters().IndexOf(Parameter) != INDEX_NONE)
+	{
+		FNiagaraEditorUtilities::FNiagaraActionScopeConfig Config;
+		Config.bIncludeSystem = true;
+		Config.bIncludeEmitters = true;
+		FNiagaraEditorUtilities::UserParameters::DeleteUserParameterReferences(AsShared(), Parameter, Config);
+		
+		System->GetExposedParameters().RemoveParameter(Parameter);
+	}
+
+	System->HandleVariableRemoved(Parameter, true);
+
+	return true;
+}
+
 TArray<UNiagaraGraph*> FNiagaraSystemViewModel::GetAllGraphs()
 {
 	TArray<UNiagaraGraph*> OutGraphs;
diff --git a/Engine/Plugins/FX/Niagara/Source/NiagaraEditor/Private/ViewModels/Stack/NiagaraStackFunctionInput.cpp b/Engine/Plugins/FX/Niagara/Source/NiagaraEditor/Private/ViewModels/Stack/NiagaraStackFunctionInput.cpp
--- a/Engine/Plugins/FX/Niagara/Source/NiagaraEditor/Private/ViewModels/Stack/NiagaraStackFunctionInput.cpp
+++ b/Engine/Plugins/FX/Niagara/Source/NiagaraEditor/Private/ViewModels/Stack/NiagaraStackFunctionInput.cpp
@@ -2274,11 +2274,11 @@ void UNiagaraStackFunctionInput::SetLocalValue(TSharedRef<FStructOnScope> InLoca
 	}
 }
 
-void UNiagaraStackFunctionInput::SetDataInterfaceValue(const UNiagaraDataInterface& InDataInterface)
+void UNiagaraStackFunctionInput::SetDataInterfaceValue(TSubclassOf<UNiagaraDataInterface> InDataInterfaceType)
 {
-	checkf(InDataInterface.GetClass() == InputType.GetClass(), TEXT("Can not set an input to an unrelated type."));
+	checkf(InDataInterfaceType->IsChildOf(InputType.GetClass()), TEXT("Can not set an input to an unrelated type."));
 
-	if (InputValues.Mode == EValueMode::Data && InputValues.DataObject->Equals(&InDataInterface))
+	if (InputValues.Mode == EValueMode::Data && InputValues.DataObject->Equals(GetDefault<UNiagaraDataInterface>(InDataInterfaceType)))
 	{
 		// The value matches the current value so noop.
 		return;
@@ -2302,8 +2302,7 @@ void UNiagaraStackFunctionInput::SetDataInterfaceValue(const UNiagaraDataInterfa
 	}
 
 	UNiagaraDataInterface* InputDataInterface = nullptr;
-	FNiagaraStackGraphUtilities::SetDataInterfaceValueForFunctionInput(*OverridePin, InDataInterface.GetClass(), AliasedInputParameterHandle.GetParameterHandleString().ToString(), InputDataInterface);
-
+	FNiagaraStackGraphUtilities::SetDataInterfaceValueForFunctionInput(*OverridePin, InDataInterfaceType, AliasedInputParameterHandle.GetParameterHandleString().ToString(), InputDataInterface);
 }
 
 bool UNiagaraStackFunctionInput::CanReset() const
diff --git a/Engine/Plugins/FX/Niagara/Source/NiagaraEditor/Public/NiagaraEditorUtilities.h b/Engine/Plugins/FX/Niagara/Source/NiagaraEditor/Public/NiagaraEditorUtilities.h
--- a/Engine/Plugins/FX/Niagara/Source/NiagaraEditor/Public/NiagaraEditorUtilities.h
+++ b/Engine/Plugins/FX/Niagara/Source/NiagaraEditor/Public/NiagaraEditorUtilities.h
@@ -16,6 +16,7 @@
 
 #define UE_API NIAGARAEDITOR_API
 
+class UNiagaraStackFunctionInput;
 class UNiagaraClipboardContent;
 class UNiagaraNodeInput;
 class UNiagaraNodeOutput;
@@ -294,11 +295,33 @@ namespace FNiagaraEditorUtilities
 	 * @returns Bool for whether adding the parameter succeeded.
 	 */
 	bool AddParameter(FNiagaraVariable& NewParameterVariable, FNiagaraParameterStore& TargetParameterStore, UObject& ParameterStoreOwner, UNiagaraStackEditorData* StackEditorData);
+
+	struct FVariableAttributeBindingInfo
+	{
+		/** The actual binding that was found during traversal */
+		FNiagaraVariableAttributeBinding* Binding = nullptr;
+		/** The default binding as provided by the CDO. */
+		FNiagaraVariableAttributeBinding* DefaultBinding = nullptr;
+		ENiagaraRendererSourceDataMode SourceDataMode = ENiagaraRendererSourceDataMode::Particles;
+		/** Immediate owning object, used to call Modify() */
+		UObject* Owner = nullptr;
+		/** The owning emitter, if it exists. */
+		FVersionedNiagaraEmitter OwningEmitter;
+	};
+
+	/** Config struct to set what an action should be targeting. Can target the system, all emitters, or a specific emitter.
+	 * If SpecificEmitterGuid is set, bIncludeEmitters should be true also.  */
+	struct FNiagaraActionScopeConfig
+	{
+		FNiagaraActionScopeConfig() = default;
+		FNiagaraActionScopeConfig(FGuid InSpecificEmitterGuid);
+		
+		bool bIncludeSystem = true;
+		bool bIncludeEmitters = true;
+		TOptional<FGuid> SpecificEmitterGuid;
+	};
 	
-	TArray<FNiagaraVariable> GetReferencedUserParametersFromEmitter(TSharedRef<FNiagaraEmitterViewModel> EmitterViewModel);
-	TArray<UNiagaraNodeParameterMapGet*> GetParameterMapGetNodesWithUserParameter(TSharedRef<FNiagaraEmitterViewModel> EmitterViewModel, FNiagaraVariable UserParameter);
-	TArray<FNiagaraUserParameterBinding*> GetUserParameterBindingsForUserParameter(TSharedRef<FNiagaraEmitterViewModel> EmitterViewModel, FNiagaraVariable UserParameter);
-	TArray<TPair<FNiagaraVariableAttributeBinding*, ENiagaraRendererSourceDataMode>> GetVariableAttributeBindingsForParameter(TSharedRef<FNiagaraEmitterViewModel> EmitterViewModel, FNiagaraVariable Parameter);
+	TArray<FVariableAttributeBindingInfo> GetVariableAttributeBindingsForParameter(TSharedRef<FNiagaraSystemViewModel> SystemViewModel, FNiagaraVariable Parameter, const FNiagaraActionScopeConfig& Config);
 
 	NIAGARAEDITOR_API bool AddEmitterContextMenuActions(FMenuBuilder& MenuBuilder, const TSharedPtr<FNiagaraEmitterHandleViewModel>& EmitterHandleViewModel);
 
@@ -434,10 +457,26 @@ namespace FNiagaraEditorUtilities
 	namespace UserParameters
 	{
 		FNiagaraVariable DuplicateUserParameter(FNiagaraVariable ParameterToDuplicate, UNiagaraSystem& System);
+
+		TArray<FNiagaraVariable> GetReferencedUserParametersFromEmitter(TSharedRef<FNiagaraEmitterViewModel> EmitterViewModel);
+		
+		TArray<UNiagaraNodeParameterMapGet*> GetParameterMapGetNodesWithUserParameter(TSharedRef<FNiagaraSystemViewModel> SystemViewModel, FNiagaraVariable UserParameter, const FNiagaraActionScopeConfig& Config);
+
+		struct FUserParameterBindingInfo
+		{
+			FNiagaraUserParameterBinding* Binding = nullptr;
+			UObject* Owner = nullptr;
+		};
+		
+		TArray<FUserParameterBindingInfo> GetUserParameterBindingsForUserParameter(TSharedRef<FNiagaraSystemViewModel> SystemViewModel, FNiagaraVariable UserParameter, const FNiagaraActionScopeConfig& Config);
+		TArray<UNiagaraStackFunctionInput*> GetStackFunctionInputsWithLinkedParameter(TSharedRef<FNiagaraSystemViewModel> SystemViewModel, FNiagaraVariable Parameter, const FNiagaraActionScopeConfig& Config);
+		
 		NIAGARAEDITOR_API TObjectPtr<UNiagaraScriptVariable> GetScriptVariableForUserParameter(const FNiagaraVariable& UserParameter, TSharedPtr<FNiagaraSystemViewModel> SystemViewModel);
 		NIAGARAEDITOR_API TObjectPtr<UNiagaraScriptVariable> GetScriptVariableForUserParameter(const FNiagaraVariable& UserParameter, UNiagaraSystem& System);
 		NIAGARAEDITOR_API const UNiagaraScriptVariable* FindScriptVariableForUserParameter(const FGuid& UserParameterGuid, const UNiagaraSystem& System);
-		void ReplaceUserParameterReferences(TSharedRef<FNiagaraEmitterViewModel> EmitterViewModel, FNiagaraVariable OldUserParameter, FNiagaraVariable NewUserParameter);
+		
+		void ReplaceUserParameterReferences(TSharedRef<FNiagaraSystemViewModel> SystemViewModel, FNiagaraVariable OldUserParameter, FNiagaraVariable NewUserParameter, const FNiagaraActionScopeConfig& Config);
+		void DeleteUserParameterReferences(TSharedRef<FNiagaraSystemViewModel> SystemViewModel, FNiagaraVariable UserParameterToDelete, const FNiagaraActionScopeConfig& Config);
 	}
 
 	namespace Preview
diff --git a/Engine/Plugins/FX/Niagara/Source/NiagaraEditor/Public/ViewModels/NiagaraSystemViewModel.h b/Engine/Plugins/FX/Niagara/Source/NiagaraEditor/Public/ViewModels/NiagaraSystemViewModel.h
--- a/Engine/Plugins/FX/Niagara/Source/NiagaraEditor/Public/ViewModels/NiagaraSystemViewModel.h
+++ b/Engine/Plugins/FX/Niagara/Source/NiagaraEditor/Public/ViewModels/NiagaraSystemViewModel.h
@@ -495,6 +495,8 @@ public:
 	/** Rename a parameter within the RenameScope and update associated tracking metadata such as Parameter Definition synchronization. */
 	bool RenameParameter(const FNiagaraVariable TargetParameter, const FName NewName, ENiagaraGetGraphParameterReferencesMode RenameScopeMode);
 
+	bool DeleteParameter(const FNiagaraVariableBase Parameter);
+
 	const TSharedPtr<FNiagaraSystemGraphSelectionViewModel>& GetSystemGraphSelectionViewModel() const { return SystemGraphSelectionViewModel; };
 
 	TSharedPtr<FNiagaraUserParameterPanelViewModel> GetUserParameterPanelViewModel() const { return UserParameterPanelViewModel; }
diff --git a/Engine/Plugins/FX/Niagara/Source/NiagaraEditor/Public/ViewModels/Stack/NiagaraStackFunctionInput.h b/Engine/Plugins/FX/Niagara/Source/NiagaraEditor/Public/ViewModels/Stack/NiagaraStackFunctionInput.h
--- a/Engine/Plugins/FX/Niagara/Source/NiagaraEditor/Public/ViewModels/Stack/NiagaraStackFunctionInput.h
+++ b/Engine/Plugins/FX/Niagara/Source/NiagaraEditor/Public/ViewModels/Stack/NiagaraStackFunctionInput.h
@@ -228,8 +228,8 @@ public:
 	/** Sets this input's local value. */
 	NIAGARAEDITOR_API void SetLocalValue(TSharedRef<FStructOnScope> InLocalValue);
 
-	/** Sets this input's data interface value. */
-	NIAGARAEDITOR_API void SetDataInterfaceValue(const UNiagaraDataInterface& InDataInterface);
+	/** Sets this input's data interface value by creating a new default Data Interface. */
+	NIAGARAEDITOR_API void SetDataInterfaceValue(TSubclassOf<UNiagaraDataInterface> InDataInterfaceType);
 	
 	/** Returns whether or not the value or handle of this input has been overridden and can be reset. */
 	NIAGARAEDITOR_API bool CanReset() const;
diff --git a/Engine/Plugins/FX/Niagara/Source/NiagaraEditorWidgets/Private/Stack/SNiagaraStackFunctionInputValue.cpp b/Engine/Plugins/FX/Niagara/Source/NiagaraEditorWidgets/Private/Stack/SNiagaraStackFunctionInputValue.cpp
--- a/Engine/Plugins/FX/Niagara/Source/NiagaraEditorWidgets/Private/Stack/SNiagaraStackFunctionInputValue.cpp
+++ b/Engine/Plugins/FX/Niagara/Source/NiagaraEditorWidgets/Private/Stack/SNiagaraStackFunctionInputValue.cpp
@@ -534,8 +534,7 @@ void SNiagaraStackFunctionInputValue::SetToLocalValue()
 {
 	if (FunctionInput->GetInputType().IsDataInterface())
 	{
-		const UNiagaraDataInterface* DefaultDataInterfaceValue = GetDefault<UNiagaraDataInterface>(FunctionInput->GetInputType().GetClass());
-		FunctionInput->SetDataInterfaceValue(*DefaultDataInterfaceValue);
+		FunctionInput->SetDataInterfaceValue(FunctionInput->GetInputType().GetClass());
 	}
 	else if (FunctionInput->GetInputType().IsUObject() == false)
 	{
