- Scene capture DI update
- Add option to hide owner component (enabled by default)
- Add hiden component list

#virtualized

diff --git a/Engine/Plugins/FX/Niagara/Source/Niagara/Private/DataInterface/NiagaraDataInterfaceSceneCapture2D.cpp b/Engine/Plugins/FX/Niagara/Source/Niagara/Private/DataInterface/NiagaraDataInterfaceSceneCapture2D.cpp
--- a/Engine/Plugins/FX/Niagara/Source/Niagara/Private/DataInterface/NiagaraDataInterfaceSceneCapture2D.cpp
+++ b/Engine/Plugins/FX/Niagara/Source/Niagara/Private/DataInterface/NiagaraDataInterfaceSceneCapture2D.cpp
@@ -341,7 +341,9 @@ bool UNiagaraDataInterfaceSceneCapture2D::CopyToInternal(UNiagaraDataInterface*
 	Destination->ManagedOrthoWidth = ManagedOrthoWidth;
 	Destination->bManagedCaptureEveryFrame = bManagedCaptureEveryFrame;
 	Destination->bManagedCaptureOnMovement = bManagedCaptureOnMovement;
+	Destination->bManagedHideAttachComponent = bManagedHideAttachComponent;
 	Destination->ManagedShowOnlyActors = ManagedShowOnlyActors;
+	Destination->ManagedHidenActors = ManagedHidenActors;
 	Destination->bAutoMoveWithComponent = bAutoMoveWithComponent;
 	Destination->AutoMoveOffsetLocationMode = AutoMoveOffsetLocationMode;
 	Destination->AutoMoveOffsetLocation = AutoMoveOffsetLocation;
@@ -369,7 +371,9 @@ bool UNiagaraDataInterfaceSceneCapture2D::Equals(const UNiagaraDataInterface* In
 		Other->ManagedOrthoWidth == ManagedOrthoWidth &&
 		Other->bManagedCaptureEveryFrame == bManagedCaptureEveryFrame &&
 		Other->bManagedCaptureOnMovement == bManagedCaptureOnMovement &&
+		Other->bManagedHideAttachComponent == bManagedHideAttachComponent &&
 		Other->ManagedShowOnlyActors == ManagedShowOnlyActors &&
+		Other->ManagedHidenActors == ManagedHidenActors &&
 		Other->bAutoMoveWithComponent == bAutoMoveWithComponent &&
 		Other->AutoMoveOffsetLocationMode == AutoMoveOffsetLocationMode &&
 		Other->AutoMoveOffsetLocation == AutoMoveOffsetLocation &&
@@ -625,6 +629,17 @@ bool UNiagaraDataInterfaceSceneCapture2D::PerInstanceTick(void* PerInstanceData,
 				CaptureComponent->ShowOnlyActors = ManagedShowOnlyActors;
 				CaptureComponent->PrimitiveRenderMode = ESceneCapturePrimitiveRenderMode::PRM_UseShowOnlyList;
 			}
+			else
+			{
+				CaptureComponent->HiddenActors = ManagedHidenActors;
+				if (bManagedHideAttachComponent)
+				{
+					if (UPrimitiveComponent* SceneComponent = Cast<UPrimitiveComponent>(SystemInstance->GetAttachComponent()))
+					{
+						CaptureComponent->HiddenComponents.Add(SceneComponent);
+					}
+				}
+			}
 
 			CaptureComponent->RegisterComponentWithWorld(SystemInstance->GetWorld());
 
diff --git a/Engine/Plugins/FX/Niagara/Source/Niagara/Private/DataInterface/NiagaraDataInterfaceSceneCapture2D.h b/Engine/Plugins/FX/Niagara/Source/Niagara/Private/DataInterface/NiagaraDataInterfaceSceneCapture2D.h
--- a/Engine/Plugins/FX/Niagara/Source/Niagara/Private/DataInterface/NiagaraDataInterfaceSceneCapture2D.h
+++ b/Engine/Plugins/FX/Niagara/Source/Niagara/Private/DataInterface/NiagaraDataInterfaceSceneCapture2D.h
@@ -94,6 +94,12 @@ public:
 	UPROPERTY(EditAnywhere, Category = "SceneCaptureManaged", meta = (EditConditionHides, EditCondition = "SourceMode == ENDISceneCapture2DSourceMode::Managed"))
 	bool bManagedCaptureOnMovement = false;
 
+	UPROPERTY(EditAnywhere, Category = "SceneCaptureManaged", meta = (EditConditionHides, EditCondition = "SourceMode == ENDISceneCapture2DSourceMode::Managed"))
+	bool bManagedHideAttachComponent = true;
+
+	UPROPERTY(EditAnywhere, Category = "SceneCaptureManaged", meta = (EditConditionHides, EditCondition = "SourceMode == ENDISceneCapture2DSourceMode::Managed"))
+	TArray<TObjectPtr<AActor>> ManagedHidenActors;
+
 	UPROPERTY(EditAnywhere, Category = "SceneCaptureManaged", meta = (EditConditionHides, EditCondition = "SourceMode == ENDISceneCapture2DSourceMode::Managed"))
 	TArray<TObjectPtr<AActor>> ManagedShowOnlyActors;
 
