[Niagara]
- Fix wrong threading assumption when destroying FNiagaraShaderScript from a serialize function since it can come from the async loading thread

diff --git a/Engine/Plugins/FX/Niagara/Source/NiagaraShader/Private/NiagaraShared.cpp b/Engine/Plugins/FX/Niagara/Source/NiagaraShader/Private/NiagaraShared.cpp
--- a/Engine/Plugins/FX/Niagara/Source/NiagaraShader/Private/NiagaraShared.cpp
+++ b/Engine/Plugins/FX/Niagara/Source/NiagaraShader/Private/NiagaraShared.cpp
@@ -71,8 +71,10 @@ FNiagaraShaderScript::FNiagaraShaderScript()
 FNiagaraShaderScript::~FNiagaraShaderScript()
 {
 #if WITH_EDITOR
-	check(IsInGameThread());
-	CancelCompilation();
+	if (OutstandingCompileShaderMapIds.Num() > 0)
+	{
+		CancelCompilation();
+	}
 #endif
 }
 
