- Fix incorrect stride with FRibbonAccumulationValues when using different features

diff --git a/Engine/Plugins/FX/Niagara/Shaders/Private/Ribbons/NiagaraRibbonCommon.ush b/Engine/Plugins/FX/Niagara/Shaders/Private/Ribbons/NiagaraRibbonCommon.ush
--- a/Engine/Plugins/FX/Niagara/Shaders/Private/Ribbons/NiagaraRibbonCommon.ush
+++ b/Engine/Plugins/FX/Niagara/Shaders/Private/Ribbons/NiagaraRibbonCommon.ush
@@ -2,6 +2,8 @@
 
 #pragma once
 
+#include "/Plugin/FX/Niagara/Shared/NiagaraRibbonShared.h"
+
 #define INDEX_NONE (-1)
 #define INT_MIN         (-2147483647 - 1)
 #define UE_SMALL_NUMBER			(1.e-8f)
@@ -250,20 +252,3 @@ float AcosFast(float InX)
 	return (InX >= 0) ? Res : PI - Res;
 }
 
-struct FRibbonAccumulationValues
-{
-	float RibbonDistance;
-	uint SegmentCount;
-#if HAS_RIBBON_ID
-	uint MultiRibbonCount;
-#endif // HAS_RIBBON_ID
-#if RIBBONS_WANTS_AUTOMATIC_TESSELLATION
-	float TessTotalLength;
-	float TessAvgSegmentLength;
-	float TessAvgSegmentAngle;
-#if RIBBON_HAS_TWIST
-	float TessTwistAvgAngle;
-	float TessTwistAvgWidth;
-#endif // RIBBON_HAS_TWIST
-#endif // RIBBONS_WANTS_AUTOMATIC_TESSELLATION
-};
diff --git a/Engine/Plugins/FX/Niagara/Shaders/Shared/NiagaraRibbonShared.h b/Engine/Plugins/FX/Niagara/Shaders/Shared/NiagaraRibbonShared.h
new file mode 100644
--- /dev/null
+++ b/Engine/Plugins/FX/Niagara/Shaders/Shared/NiagaraRibbonShared.h
@@ -0,0 +1,24 @@
+// Copyright Epic Games, Inc. All Rights Reserved.
+// This file is shared between C++ & HLSL
+
+#pragma once
+
+#ifndef GPU_SIMULATION
+	#define uint	uint32
+#endif
+
+struct FRibbonAccumulationValues
+{
+	float RibbonDistance;
+	uint SegmentCount;
+	uint MultiRibbonCount;			// Only valid when HAS_RIBBON_ID
+	float TessTotalLength;			// Only valid when RIBBONS_WANTS_AUTOMATIC_TESSELLATION
+	float TessAvgSegmentLength;		// Only valid when RIBBONS_WANTS_AUTOMATIC_TESSELLATION
+	float TessAvgSegmentAngle;		// Only valid when RIBBONS_WANTS_AUTOMATIC_TESSELLATION
+	float TessTwistAvgAngle;		// Only valid when RIBBONS_WANTS_AUTOMATIC_TESSELLATION && RIBBON_HAS_TWIST
+	float TessTwistAvgWidth;		// Only valid when RIBBONS_WANTS_AUTOMATIC_TESSELLATION && RIBBON_HAS_TWIST
+};
+
+#ifndef GPU_SIMULATION
+	#undef uint
+#endif
diff --git a/Engine/Plugins/FX/Niagara/Source/Niagara/Private/NiagaraRendererRibbons.cpp b/Engine/Plugins/FX/Niagara/Source/Niagara/Private/NiagaraRendererRibbons.cpp
--- a/Engine/Plugins/FX/Niagara/Source/Niagara/Private/NiagaraRendererRibbons.cpp
+++ b/Engine/Plugins/FX/Niagara/Source/Niagara/Private/NiagaraRendererRibbons.cpp
@@ -5,6 +5,7 @@
 #include "GlobalRenderResources.h"
 #include "GPUSortManager.h"
 #include "NiagaraRibbonVertexFactory.h"
+#include "NiagaraRibbonShared.h"
 #include "NiagaraDataSet.h"
 #include "NiagaraDataSetAccessor.h"
 #include "NiagaraEmitterInstance.h"
@@ -552,17 +553,6 @@ struct FNiagaraRibbonGPUInitComputeBuffers
 	
 	FNiagaraRibbonGPUInitComputeBuffers() { }
 
-	uint32 GetAccumulationStructSize(bool bWantsMultiRibbon, bool bWantsTessellation, bool bWantsTessellationTwist)
-	{
-		const uint32 BaseStructSize = 2;
-		const uint32 MultiRibbonSize = 1;
-		const uint32 TessellationSize = bWantsTessellationTwist ? 5 : 3 ;
-
-		return BaseStructSize
-			+ (bWantsMultiRibbon ? MultiRibbonSize : 0)
-			+ (bWantsTessellation ? TessellationSize : 0);
-	}
-
 	void InitOrUpdateBuffers(FRHICommandListBase& RHICmdList, int32 NeededSize, bool bWantsMultiRibbon, bool bWantsTessellation, bool bWantsTessellationTwist)
 	{
 		// TODO: Downsize these when we haven't needed the size for a bit
@@ -579,7 +569,7 @@ struct FNiagaraRibbonGPUInitComputeBuffers
 			TransientTessellationStats.Initialize(RHICmdList, TEXT("NiagaraGPUInit-Tessellation-0"), sizeof(float), TessellationBufferSize, EPixelFormat::PF_R32_FLOAT, InitialAccess, BUF_Static);
 		}
 
-		const uint32 AccumulationBufferStructSize = GetAccumulationStructSize(bWantsMultiRibbon, bWantsTessellation, bWantsTessellationTwist) * sizeof(float);
+		const uint32 AccumulationBufferStructSize = sizeof(FRibbonAccumulationValues);
 		if (TransientAccumulation[0].NumBytes < (AccumulationBufferStructSize * NeededSize))
 		{
 			TransientAccumulation[0].Initialize(RHICmdList, TEXT("NiagaraGPUInit-Accumulation-0"), AccumulationBufferStructSize, NeededSize, BUF_Static, false, false, InitialAccess);
diff --git a/Engine/Plugins/FX/Niagara/Source/NiagaraShader/Private/NiagaraRibbonCompute.cpp b/Engine/Plugins/FX/Niagara/Source/NiagaraShader/Private/NiagaraRibbonCompute.cpp
--- a/Engine/Plugins/FX/Niagara/Source/NiagaraShader/Private/NiagaraRibbonCompute.cpp
+++ b/Engine/Plugins/FX/Niagara/Source/NiagaraShader/Private/NiagaraRibbonCompute.cpp
@@ -14,6 +14,7 @@ void FNiagaraRibbonComputeCommon::ModifyCompilationEnvironment(const FGlobalShad
 {
 	FGlobalShader::ModifyCompilationEnvironment(Parameters, OutEnvironment);
 	//OutEnvironment.SetDefine(TEXT("NiagaraVFLooseParameters"), TEXT("NiagaraRibbonVFLooseParameters"));
+	OutEnvironment.SetDefine(TEXT("GPU_SIMULATION"), true);
 	if (ThreadGroupSize > 0)
 	{
 		OutEnvironment.SetDefine(TEXT("THREADGROUP_SIZE"), ThreadGroupSize);
