- Add rotator distribution to stateless & use in shape location to allow the rotation to be bound to an attribute
- Fixed the axis not properly being used on the display name for each widget
- Fixed the same binding being added multiple times

diff --git a/Engine/Plugins/FX/Niagara/Source/Niagara/Internal/Stateless/Modules/NiagaraStatelessModule_ShapeLocation.h b/Engine/Plugins/FX/Niagara/Source/Niagara/Internal/Stateless/Modules/NiagaraStatelessModule_ShapeLocation.h
--- a/Engine/Plugins/FX/Niagara/Source/Niagara/Internal/Stateless/Modules/NiagaraStatelessModule_ShapeLocation.h
+++ b/Engine/Plugins/FX/Niagara/Source/Niagara/Internal/Stateless/Modules/NiagaraStatelessModule_ShapeLocation.h
@@ -76,8 +76,8 @@ public:
 	UPROPERTY(EditAnywhere, Category = "Parameters", meta = (SegmentedDisplay))
 	ENiagaraCoordinateSpace CoordinateSpace = ENiagaraCoordinateSpace::Simulation;
 
-	UPROPERTY(EditAnywhere, Category = "Parameters")
-	FRotator3f ShapeRotation = FRotator3f::ZeroRotator;
+	UPROPERTY(EditAnywhere, Category = "Parameters", meta = (DisableRangeDistribution, DisableUniformDistribution))
+	FNiagaraDistributionRangeRotator ShapeRotation;
 
 private:
 #if WITH_EDITORONLY_DATA
diff --git a/Engine/Plugins/FX/Niagara/Source/Niagara/Internal/Stateless/NiagaraStatelessEmitterDataBuildContext.h b/Engine/Plugins/FX/Niagara/Source/Niagara/Internal/Stateless/NiagaraStatelessEmitterDataBuildContext.h
--- a/Engine/Plugins/FX/Niagara/Source/Niagara/Internal/Stateless/NiagaraStatelessEmitterDataBuildContext.h
+++ b/Engine/Plugins/FX/Niagara/Source/Niagara/Internal/Stateless/NiagaraStatelessEmitterDataBuildContext.h
@@ -182,6 +182,7 @@ public:
 	FNiagaraStatelessRangeVector2 ConvertDistributionToRange(const FNiagaraDistributionRangeVector2& Distribution, const FVector2f& DefaultValue) const { return ConvertDistributionToRangeHelper<FNiagaraStatelessRangeVector2>(Distribution, DefaultValue); }
 	FNiagaraStatelessRangeVector3 ConvertDistributionToRange(const FNiagaraDistributionRangeVector3& Distribution, const FVector3f& DefaultValue) const { return ConvertDistributionToRangeHelper<FNiagaraStatelessRangeVector3>(Distribution, DefaultValue); }
 	FNiagaraStatelessRangeColor   ConvertDistributionToRange(const FNiagaraDistributionRangeColor& Distribution, const FLinearColor& DefaultValue) const { return ConvertDistributionToRangeHelper<FNiagaraStatelessRangeColor>(Distribution, DefaultValue); }
+	FNiagaraStatelessRangeRotator ConvertDistributionToRange(const FNiagaraDistributionRangeRotator& Distribution, const FRotator3f& DefaultValue) const { return ConvertDistributionToRangeHelper<FNiagaraStatelessRangeRotator>(Distribution, DefaultValue); }
 	FNiagaraStatelessRangeInt     ConvertDistributionToRange(const FNiagaraDistributionRangeInt& Distribution, int32 DefaultValue) const { return ConvertDistributionToRangeHelper<FNiagaraStatelessRangeInt>(Distribution, DefaultValue); }
 
 	NIAGARA_API void AddParticleSimulationExecSimulate(TFunction<void(const NiagaraStateless::FParticleSimulationContext&)> Func) const;
diff --git a/Engine/Plugins/FX/Niagara/Source/Niagara/Private/Stateless/Modules/NiagaraStatelessModule_ShapeLocation.cpp b/Engine/Plugins/FX/Niagara/Source/Niagara/Private/Stateless/Modules/NiagaraStatelessModule_ShapeLocation.cpp
--- a/Engine/Plugins/FX/Niagara/Source/Niagara/Private/Stateless/Modules/NiagaraStatelessModule_ShapeLocation.cpp
+++ b/Engine/Plugins/FX/Niagara/Source/Niagara/Private/Stateless/Modules/NiagaraStatelessModule_ShapeLocation.cpp
@@ -12,11 +12,11 @@ namespace NSMShapeLocationPrivate
 {
 	struct FModuleBuiltData
 	{
-		ENSM_ShapePrimitive		ShapePrimitive = ENSM_ShapePrimitive::Max;
-		ENiagaraCoordinateSpace	CoordinateSpace = ENiagaraCoordinateSpace::Simulation;
-		FQuat4f					ShapeRotation = FQuat4f::Identity;
-		int32					PositionVariableOffset = INDEX_NONE;
-		int32					PreviousPositionVariableOffset = INDEX_NONE;
+		ENSM_ShapePrimitive				ShapePrimitive = ENSM_ShapePrimitive::Max;
+		ENiagaraCoordinateSpace			CoordinateSpace = ENiagaraCoordinateSpace::Simulation;
+		FNiagaraStatelessRangeRotator	ShapeRotation;
+		int32							PositionVariableOffset = INDEX_NONE;
+		int32							PreviousPositionVariableOffset = INDEX_NONE;
 	};
 
 	struct FModuleBuiltDataBox
@@ -189,7 +189,7 @@ void UNiagaraStatelessModule_ShapeLocation::BuildEmitterData(const FNiagaraState
 
 	BuiltData->ShapePrimitive	= ShapePrimitive;
 	BuiltData->CoordinateSpace	= CoordinateSpace;
-	BuiltData->ShapeRotation	= ShapeRotation.Quaternion();
+	BuiltData->ShapeRotation	= BuildContext.ConvertDistributionToRange(ShapeRotation, FRotator3f::ZeroRotator);
 
 	switch (ShapePrimitive)
 	{
@@ -253,7 +253,8 @@ void UNiagaraStatelessModule_ShapeLocation::SetShaderParameters(const FNiagaraSt
 	FParameters* Parameters = SetShaderParameterContext.GetParameterNestedStruct<FParameters>();
 	const FModuleBuiltData* ModuleBuiltData = SetShaderParameterContext.ReadBuiltData<FModuleBuiltData>();
 
-	const FQuat4f Rotation = SetShaderParameterContext.GetSpaceTransforms().TransformRotation(ModuleBuiltData->CoordinateSpace, ModuleBuiltData->ShapeRotation);
+	const FQuat4f ResolvedShapeRotation = SetShaderParameterContext.GetRendererParameterValue(ModuleBuiltData->ShapeRotation.ParameterOffset, ModuleBuiltData->ShapeRotation.Min.Quaternion());
+	const FQuat4f Rotation = SetShaderParameterContext.GetSpaceTransforms().TransformRotation(ModuleBuiltData->CoordinateSpace, ResolvedShapeRotation);
 
 	Parameters->ShapeLocation_Mode			= FUintVector4::ZeroValue;
 	Parameters->ShapeLocation_Parameters0	= FVector4f::Zero();
@@ -344,8 +345,14 @@ void UNiagaraStatelessModule_ShapeLocation::SetShaderParameters(const FNiagaraSt
 #if WITH_EDITOR
 void UNiagaraStatelessModule_ShapeLocation::DrawDebug(const FNiagaraStatelessDrawDebugContext& DrawDebugContext) const
 {
+	if (ShapeRotation.IsBinding())
+	{
+		return;
+	}
+
+	const FQuat4f ResolvedRotation = ShapeRotation.CalculateRange().Min.Quaternion();
 	const FVector DrawPosition = DrawDebugContext.TransformPosition(FVector3f::ZeroVector);
-	const FQuat DrawRotation = DrawDebugContext.TransformRotation(CoordinateSpace, ShapeRotation.Quaternion());
+	const FQuat DrawRotation = DrawDebugContext.TransformRotation(CoordinateSpace, ResolvedRotation);
 	switch (ShapePrimitive)
 	{
 		case ENSM_ShapePrimitive::Box:
diff --git a/Engine/Plugins/FX/Niagara/Source/Niagara/Private/Stateless/NiagaraStatelessDistribution.cpp b/Engine/Plugins/FX/Niagara/Source/Niagara/Private/Stateless/NiagaraStatelessDistribution.cpp
--- a/Engine/Plugins/FX/Niagara/Source/Niagara/Private/Stateless/NiagaraStatelessDistribution.cpp
+++ b/Engine/Plugins/FX/Niagara/Source/Niagara/Private/Stateless/NiagaraStatelessDistribution.cpp
@@ -414,6 +414,20 @@ FNiagaraStatelessRangeColor FNiagaraDistributionRangeColor::CalculateRange(const
 	return Mode == ENiagaraDistributionMode::UniformConstant ? FNiagaraStatelessRangeColor(Min, Min) : FNiagaraStatelessRangeColor(Min, Max);
 }
 
+void FNiagaraDistributionRangeRotator::InitConstant(const FRotator3f& Value)
+{
+	Mode = ENiagaraDistributionMode::NonUniformConstant;
+	Min = Value;
+	Max = Value;
+#if WITH_EDITORONLY_DATA
+	ChannelConstantsAndRanges = TArray<float>({ Value.Roll, Value.Pitch, Value.Yaw });
+#endif
+}
+
+FNiagaraStatelessRangeRotator FNiagaraDistributionRangeRotator::CalculateRange(const FRotator3f& Default) const
+{
+	return Mode == ENiagaraDistributionMode::UniformConstant ? FNiagaraStatelessRangeRotator(Min, Min) : FNiagaraStatelessRangeRotator(Min, Max);
+}
 
 void FNiagaraDistributionFloat::InitConstant(float Value)
 {
@@ -755,6 +769,31 @@ void FNiagaraDistributionRangeColor::UpdateValuesFromDistribution()
 	);
 }
 
+void FNiagaraDistributionRangeRotator::UpdateValuesFromDistribution()
+{
+	FVector2f ValuesTimeRange;
+	NiagaraStatelessDistributionPrivate::UpdateDistributionValues(
+		Mode,
+		ChannelConstantsAndRanges,
+		ChannelCurves,
+		3,
+		ValuesTimeRange,
+		[this](int32 Num) {},
+		[this](int32 ValueIndex, int32 ChannelIndex) -> float&
+		{
+			FRotator3f& Target = ValueIndex == 0 ? Min : Max;
+			switch (ChannelIndex)
+			{
+				default:	checkNoEntry(); // Intentional fall through
+				case 0:		return Target.Roll;
+				case 1:		return Target.Pitch;
+				case 2:		return Target.Yaw;
+			}
+		},
+		MaxLutSampleCount
+	);
+}
+
 void FNiagaraDistributionFloat::UpdateValuesFromDistribution()
 {
 	NiagaraStatelessDistributionPrivate::UpdateDistributionValues(
@@ -881,3 +920,16 @@ bool FNiagaraDistributionRangeVector3::SerializeFromMismatchedTag(const struct F
 	return false;
 }
 
+bool FNiagaraDistributionRangeRotator::SerializeFromMismatchedTag(const struct FPropertyTag& Tag, FStructuredArchive::FSlot Slot)
+{
+	if (Tag.GetType().IsStruct(NAME_Rotator3f))
+	{
+		FRotator3f Value;
+		Slot << Value;
+		*this = FNiagaraDistributionRangeRotator(Value);
+		return true;
+	}
+
+	return false;
+}
+
diff --git a/Engine/Plugins/FX/Niagara/Source/Niagara/Public/Stateless/NiagaraStatelessDistribution.h b/Engine/Plugins/FX/Niagara/Source/Niagara/Public/Stateless/NiagaraStatelessDistribution.h
--- a/Engine/Plugins/FX/Niagara/Source/Niagara/Public/Stateless/NiagaraStatelessDistribution.h
+++ b/Engine/Plugins/FX/Niagara/Source/Niagara/Public/Stateless/NiagaraStatelessDistribution.h
@@ -225,6 +225,32 @@ struct FNiagaraDistributionRangeColor : public FNiagaraDistributionBase
 #endif
 };
 
+USTRUCT()
+struct FNiagaraDistributionRangeRotator : public FNiagaraDistributionBase
+{
+	GENERATED_BODY()
+
+	FNiagaraDistributionRangeRotator() { InitConstant(FRotator3f::ZeroRotator); }
+	explicit FNiagaraDistributionRangeRotator(const FRotator3f& ConstantValue) { InitConstant(ConstantValue); }
+
+	UPROPERTY(EditAnywhere, Category = "Parameters")
+	FRotator3f Min = FRotator3f::ZeroRotator;
+
+	UPROPERTY(EditAnywhere, Category = "Parameters")
+	FRotator3f Max = FRotator3f::ZeroRotator;
+
+	NIAGARA_API void InitConstant(const FRotator3f& Value);
+	NIAGARA_API FNiagaraStatelessRangeRotator CalculateRange(const FRotator3f& Default = FRotator3f::ZeroRotator) const;
+
+#if WITH_EDITORONLY_DATA
+	virtual bool AllowCurves() const override { return false; }
+	virtual int32 GetBaseNumberOfChannels() const override { return 3; }
+	virtual void UpdateValuesFromDistribution() override;
+	virtual FNiagaraTypeDefinition GetBindingTypeDef() const { return FNiagaraTypeDefinition::GetQuatDef(); }
+#endif
+	bool SerializeFromMismatchedTag(const struct FPropertyTag& Tag, FStructuredArchive::FSlot Slot);
+};
+
 USTRUCT()
 struct FNiagaraDistributionFloat : public FNiagaraDistributionBase
 {
@@ -440,3 +466,11 @@ struct TStructOpsTypeTraits<FNiagaraDistributionRangeVector3> : public TStructOp
 	};
 };
 
+template<>
+struct TStructOpsTypeTraits<FNiagaraDistributionRangeRotator> : public TStructOpsTypeTraitsBase2<FNiagaraDistributionRangeRotator>
+{
+	enum
+	{
+		WithStructuredSerializeFromMismatchedTag = true,
+	};
+};
diff --git a/Engine/Plugins/FX/Niagara/Source/Niagara/Public/Stateless/NiagaraStatelessRange.h b/Engine/Plugins/FX/Niagara/Source/Niagara/Public/Stateless/NiagaraStatelessRange.h
--- a/Engine/Plugins/FX/Niagara/Source/Niagara/Public/Stateless/NiagaraStatelessRange.h
+++ b/Engine/Plugins/FX/Niagara/Source/Niagara/Public/Stateless/NiagaraStatelessRange.h
@@ -11,6 +11,7 @@ template<> struct FNiagaraStatelessRangeDefaultValue<FVector2f> { static const F
 template<> struct FNiagaraStatelessRangeDefaultValue<FVector3f> { static const FVector3f Zero() { return FVector3f(0.0f, 0.0f, 0.0f); } };
 template<> struct FNiagaraStatelessRangeDefaultValue<FVector4f> { static const FVector4f Zero() { return FVector4f(0.0f, 0.0f, 0.0f, 0.0f); } };
 template<> struct FNiagaraStatelessRangeDefaultValue<FLinearColor> { static const FLinearColor Zero() { return FLinearColor(0.0f, 0.0f, 0.0f, 0.0f); } };
+template<> struct FNiagaraStatelessRangeDefaultValue<FRotator3f> { static const FRotator3f Zero() { return FRotator3f::ZeroRotator; } };
 
 template<typename TType>
 struct FNiagaraStatelessRange
@@ -35,3 +36,4 @@ using FNiagaraStatelessRangeVector2 = FNiagaraStatelessRange<FVector2f>;
 using FNiagaraStatelessRangeVector3 = FNiagaraStatelessRange<FVector3f>;
 using FNiagaraStatelessRangeVector4 = FNiagaraStatelessRange<FVector4f>;
 using FNiagaraStatelessRangeColor = FNiagaraStatelessRange<FLinearColor>;
+using FNiagaraStatelessRangeRotator = FNiagaraStatelessRange<FRotator3f>;
diff --git a/Engine/Plugins/FX/Niagara/Source/NiagaraEditor/Private/Stateless/NiagaraDistributionPropertyCustomization.cpp b/Engine/Plugins/FX/Niagara/Source/NiagaraEditor/Private/Stateless/NiagaraDistributionPropertyCustomization.cpp
--- a/Engine/Plugins/FX/Niagara/Source/NiagaraEditor/Private/Stateless/NiagaraDistributionPropertyCustomization.cpp
+++ b/Engine/Plugins/FX/Niagara/Source/NiagaraEditor/Private/Stateless/NiagaraDistributionPropertyCustomization.cpp
@@ -49,13 +49,12 @@ public:
 			const FName DisableUniformDistributionName("DisableUniformDistribution");
 			const FName DisableNonUniformDistributionName("DisableNonUniformDistribution");
 			const FName DisableRangeDistributionName("DisableRangeDistribution");
-			const FName DisplayAsColorDistributionName("DisplayAsColorDistribution");
 			bAllowBinding = InDistribution->AllowBinding() && InDistribution->GetBindingTypeDef().IsValid() && (PropertyHandle ? PropertyHandle->HasMetaData(DisableBindingDistributionName) == false : false);
 			bAllowUniform = InDistribution->AllowConstant() && PropertyHandle->HasMetaData(DisableUniformDistributionName) == false;
 			bAllowNonUniform = InDistribution->AllowConstant() && PropertyHandle->HasMetaData(DisableNonUniformDistributionName) == false;
 			bAllowRange = InDistribution->AllowConstant() && PropertyHandle->HasMetaData(DisableRangeDistributionName) == false;
 			bAllowCurves = InDistribution->AllowCurves() && PropertyHandle->HasMetaData(DisableCurveDistributionName) == false;
-			bDisplayAsColor = InDistribution->DisplayAsColor() || (InNumChannels >= 3 && PropertyHandle->HasMetaData(DisplayAsColorDistributionName));
+			bDisplayAsColor = InDistribution->DisplayAsColor();
 			PropertyHandle->SetOnPropertyValueChanged(FSimpleDelegate::CreateSP(this->AsShared(), &FNiagaraDistributionAdapter::DistributionPropertyChanged));
 			EditorMode = GetDistributionEditorModeFromSourceMode(SourceNumChannels, bDisplayAsColor, SourceDistribution->Mode);
 			EditorNumChannels = FNiagaraDistributionEditorUtilities::IsUniform(EditorMode) ? 1 : SourceNumChannels;
@@ -564,7 +563,7 @@ public:
 					{
 						if (Variable.GetType() == AllowedTypeDef && Variable.IsInNameSpace(FNiagaraConstants::SystemNamespaceString))
 						{
-							AvailableBindings.Add(Variable);
+							AvailableBindings.AddUnique(Variable);
 						}
 					}
 				}
@@ -574,7 +573,7 @@ public:
 			{
 				if (Variable.GetType() == AllowedTypeDef)
 				{
-					AvailableBindings.Add(Variable);
+					AvailableBindings.AddUnique(Variable);
 				}
 			}
 
@@ -586,7 +585,7 @@ public:
 				{
 					if (NPCVariable.GetType() == AllowedTypeDef)
 					{
-						AvailableBindings.Add(NPCVariable);
+						AvailableBindings.AddUnique(NPCVariable);
 					}
 				}
 			}
@@ -941,6 +940,20 @@ private:
 	bool bUpdatingDistributionModeInternally = false;
 };
 
+class FNiagaraDistributionAdapterRotator : public FNiagaraDistributionAdapter
+{
+	virtual FText GetChannelDisplayName(int32 ChannelIndex) const override
+	{
+		switch (ChannelIndex)
+		{
+			case 0:	return LOCTEXT("RollChannelName", "Roll");
+			case 1:	return LOCTEXT("PitchChannelName", "Pitch");
+			case 2:	return LOCTEXT("YawChannelName", "Yaw");
+		}
+		return {};
+	}
+};
+
 class SNiagaraDistributionPropertyWidget : public SCompoundWidget, public FEditorUndoClient
 {
 	SLATE_BEGIN_ARGS(SNiagaraDistributionPropertyWidget) { }
@@ -1091,6 +1104,27 @@ TSharedRef<IPropertyTypeCustomization> FNiagaraDistributionPropertyCustomization
 	return MakeShareable<FNiagaraDistributionPropertyCustomization>(new FNiagaraDistributionPropertyCustomization(ColorDistributionPropertyHandleToDistributionAdapter));
 }
 
+TSharedRef<IPropertyTypeCustomization> FNiagaraDistributionPropertyCustomization::MakeRotatorInstance()
+{
+	FPropertyHandleToDistributionAdapter DistributionPropertyHandleToDistributionAdapter = FPropertyHandleToDistributionAdapter::CreateLambda(
+		[](TSharedRef<IPropertyHandle> DistributionPropertyHandle)
+		{
+			void* ValueData = nullptr;
+			TArray<UObject*> OuterObjects;
+			DistributionPropertyHandle->GetOuterObjects(OuterObjects);
+			if (OuterObjects.Num() == 1 && DistributionPropertyHandle->GetValueData(ValueData) == FPropertyAccess::Success)
+			{
+				FNiagaraDistributionBase* Distribution = static_cast<FNiagaraDistributionBase*>(ValueData);
+				TSharedRef<FNiagaraDistributionAdapter> DistributionAdapter = MakeShared<FNiagaraDistributionAdapterRotator>();
+				DistributionAdapter->Initialize(DistributionPropertyHandle, OuterObjects[0], Distribution, Distribution->GetBaseNumberOfChannels());
+				return DistributionAdapter;
+			}
+			return MakeShared<FNiagaraDistributionAdapter>();
+		});
+
+	return MakeShareable<FNiagaraDistributionPropertyCustomization>(new FNiagaraDistributionPropertyCustomization(DistributionPropertyHandleToDistributionAdapter));
+}
+
 void FNiagaraDistributionPropertyCustomization::CustomizeHeader(TSharedRef<IPropertyHandle> PropertyHandle, FDetailWidgetRow& HeaderRow, IPropertyTypeCustomizationUtils& CustomizationUtils)
 {
 	HeaderRow.NameContent()
diff --git a/Engine/Plugins/FX/Niagara/Source/NiagaraEditor/Private/Stateless/NiagaraDistributionPropertyCustomization.h b/Engine/Plugins/FX/Niagara/Source/NiagaraEditor/Private/Stateless/NiagaraDistributionPropertyCustomization.h
--- a/Engine/Plugins/FX/Niagara/Source/NiagaraEditor/Private/Stateless/NiagaraDistributionPropertyCustomization.h
+++ b/Engine/Plugins/FX/Niagara/Source/NiagaraEditor/Private/Stateless/NiagaraDistributionPropertyCustomization.h
@@ -15,6 +15,7 @@ public:
 	static TSharedRef<IPropertyTypeCustomization> MakeVector3Instance();
 	static TSharedRef<IPropertyTypeCustomization> MakePositionInstance();
 	static TSharedRef<IPropertyTypeCustomization> MakeColorInstance();
+	static TSharedRef<IPropertyTypeCustomization> MakeRotatorInstance();
 
 	virtual void CustomizeHeader(TSharedRef<IPropertyHandle> PropertyHandle, FDetailWidgetRow& HeaderRow, IPropertyTypeCustomizationUtils& CustomizationUtils) override;
 
diff --git a/Engine/Plugins/FX/Niagara/Source/NiagaraEditor/Private/ViewModels/Stack/NiagaraStackStatelessEmitterSimulateGroup.cpp b/Engine/Plugins/FX/Niagara/Source/NiagaraEditor/Private/ViewModels/Stack/NiagaraStackStatelessEmitterSimulateGroup.cpp
--- a/Engine/Plugins/FX/Niagara/Source/NiagaraEditor/Private/ViewModels/Stack/NiagaraStackStatelessEmitterSimulateGroup.cpp
+++ b/Engine/Plugins/FX/Niagara/Source/NiagaraEditor/Private/ViewModels/Stack/NiagaraStackStatelessEmitterSimulateGroup.cpp
@@ -407,6 +407,7 @@ void UNiagaraStackStatelessModuleItem::RefreshChildrenInternal(const TArray<UNia
 			ModuleObject->RegisterInstancedCustomPropertyTypeLayout(FNiagaraDistributionRangeVector2::StaticStruct()->GetFName(), FOnGetPropertyTypeCustomizationInstance::CreateStatic(&FNiagaraDistributionPropertyCustomization::MakeVector2Instance));
 			ModuleObject->RegisterInstancedCustomPropertyTypeLayout(FNiagaraDistributionRangeVector3::StaticStruct()->GetFName(), FOnGetPropertyTypeCustomizationInstance::CreateStatic(&FNiagaraDistributionPropertyCustomization::MakeVector3Instance));
 			ModuleObject->RegisterInstancedCustomPropertyTypeLayout(FNiagaraDistributionRangeColor::StaticStruct()->GetFName(), FOnGetPropertyTypeCustomizationInstance::CreateStatic(&FNiagaraDistributionPropertyCustomization::MakeColorInstance));
+			ModuleObject->RegisterInstancedCustomPropertyTypeLayout(FNiagaraDistributionRangeRotator::StaticStruct()->GetFName(), FOnGetPropertyTypeCustomizationInstance::CreateStatic(&FNiagaraDistributionPropertyCustomization::MakeRotatorInstance));
 			ModuleObject->RegisterInstancedCustomPropertyTypeLayout(FNiagaraDistributionRangeInt::StaticStruct()->GetFName(), FOnGetPropertyTypeCustomizationInstance::CreateStatic(&FNiagaraDistributionIntPropertyCustomization::MakeIntInstance));
 			ModuleObject->RegisterInstancedCustomPropertyTypeLayout(FNiagaraDistributionCurveFloat::StaticStruct()->GetFName(), FOnGetPropertyTypeCustomizationInstance::CreateStatic(&FNiagaraDistributionPropertyCustomization::MakeFloatInstance));
 			ModuleObject->RegisterInstancedCustomPropertyTypeLayout(FNiagaraDistributionCurveVector3::StaticStruct()->GetFName(), FOnGetPropertyTypeCustomizationInstance::CreateStatic(&FNiagaraDistributionPropertyCustomization::MakeVector3Instance));
diff --git a/Engine/Plugins/FX/Niagara/Source/NiagaraEditor/Private/Widgets/SNiagaraDistributionEditor.cpp b/Engine/Plugins/FX/Niagara/Source/NiagaraEditor/Private/Widgets/SNiagaraDistributionEditor.cpp
--- a/Engine/Plugins/FX/Niagara/Source/NiagaraEditor/Private/Widgets/SNiagaraDistributionEditor.cpp
+++ b/Engine/Plugins/FX/Niagara/Source/NiagaraEditor/Private/Widgets/SNiagaraDistributionEditor.cpp
@@ -388,9 +388,11 @@ private:
 					.AllowSpin(true)
 					.bColorAxisLabels(true)
 					.X(this, &SNiagaraDistributionValueEditor::GetValue, 0, ValueIndex)
+					.XDisplayName(this, &SNiagaraDistributionValueEditor::GetChannelDisplayName, 0)
 					.OnXChanged(this, &SNiagaraDistributionValueEditor::ValueChanged, 0, ValueIndex)
 					.OnXCommitted(this, &SNiagaraDistributionValueEditor::ValueCommitted, 0, ValueIndex)
 					.Y(this, &SNiagaraDistributionValueEditor::GetValue, 1, ValueIndex)
+					.YDisplayName(this, &SNiagaraDistributionValueEditor::GetChannelDisplayName, 1)
 					.OnYChanged(this, &SNiagaraDistributionValueEditor::ValueChanged, 1, ValueIndex)
 					.OnYCommitted(this, &SNiagaraDistributionValueEditor::ValueCommitted, 1, ValueIndex)
 					.OnBeginSliderMovement(this, &SNiagaraDistributionValueEditor::BeginValueChange)
@@ -408,12 +410,15 @@ private:
 					.AllowSpin(true)
 					.bColorAxisLabels(true)
 					.X(this, &SNiagaraDistributionValueEditor::GetValue, 0, ValueIndex)
+					.XDisplayName(this, &SNiagaraDistributionValueEditor::GetChannelDisplayName, 0)
 					.OnXChanged(this, &SNiagaraDistributionValueEditor::ValueChanged, 0, ValueIndex)
 					.OnXCommitted(this, &SNiagaraDistributionValueEditor::ValueCommitted, 0, ValueIndex)
 					.Y(this, &SNiagaraDistributionValueEditor::GetValue, 1, ValueIndex)
+					.YDisplayName(this, &SNiagaraDistributionValueEditor::GetChannelDisplayName, 1)
 					.OnYChanged(this, &SNiagaraDistributionValueEditor::ValueChanged, 1, ValueIndex)
 					.OnYCommitted(this, &SNiagaraDistributionValueEditor::ValueCommitted, 1, ValueIndex)
 					.Z(this, &SNiagaraDistributionValueEditor::GetValue, 2, ValueIndex)
+					.ZDisplayName(this, &SNiagaraDistributionValueEditor::GetChannelDisplayName, 2)
 					.OnZChanged(this, &SNiagaraDistributionValueEditor::ValueChanged, 2, ValueIndex)
 					.OnZCommitted(this, &SNiagaraDistributionValueEditor::ValueCommitted, 2, ValueIndex)
 					.OnBeginSliderMovement(this, &SNiagaraDistributionValueEditor::BeginValueChange)
@@ -431,15 +436,19 @@ private:
 					.AllowSpin(true)
 					.bColorAxisLabels(true)
 					.X(this, &SNiagaraDistributionValueEditor::GetValue, 0, ValueIndex)
+					.XDisplayName(this, &SNiagaraDistributionValueEditor::GetChannelDisplayName, 0)
 					.OnXChanged(this, &SNiagaraDistributionValueEditor::ValueChanged, 0, ValueIndex)
 					.OnXCommitted(this, &SNiagaraDistributionValueEditor::ValueCommitted, 0, ValueIndex)
 					.Y(this, &SNiagaraDistributionValueEditor::GetValue, 1, ValueIndex)
+					.YDisplayName(this, &SNiagaraDistributionValueEditor::GetChannelDisplayName, 1)
 					.OnYChanged(this, &SNiagaraDistributionValueEditor::ValueChanged, 1, ValueIndex)
 					.OnYCommitted(this, &SNiagaraDistributionValueEditor::ValueCommitted, 1, ValueIndex)
 					.Z(this, &SNiagaraDistributionValueEditor::GetValue, 2, ValueIndex)
+					.ZDisplayName(this, &SNiagaraDistributionValueEditor::GetChannelDisplayName, 2)
 					.OnZChanged(this, &SNiagaraDistributionValueEditor::ValueChanged, 2, ValueIndex)
 					.OnZCommitted(this, &SNiagaraDistributionValueEditor::ValueCommitted, 2, ValueIndex)
 					.W(this, &SNiagaraDistributionValueEditor::GetValue, 3, ValueIndex)
+					.WDisplayName(this, &SNiagaraDistributionValueEditor::GetChannelDisplayName, 3)
 					.OnWChanged(this, &SNiagaraDistributionValueEditor::ValueChanged, 3, ValueIndex)
 					.OnWCommitted(this, &SNiagaraDistributionValueEditor::ValueCommitted, 3, ValueIndex)
 					.OnBeginSliderMovement(this, &SNiagaraDistributionValueEditor::BeginValueChange)
@@ -489,6 +498,11 @@ private:
 		}
 	}
 
+	FText GetChannelDisplayName(int32 ChannelIndex) const
+	{
+		return DistributionAdapter->GetChannelDisplayName(ChannelIndex);
+	}
+
 	FLinearColor GetColorValue(int32 ValueIndex) const
 	{
 		if (DistributionAdapter->GetNumChannels() == 3)
