Niagara - Fix an issue where delegate bindings in the system view model could be incorrect after rebuilding the emitter handle view models due tor a change in the emitter handle array.  This could result in weird behavior in dependency evaluation and could result in changes not propagating correctly in the UI.

diff --git a/Engine/Plugins/FX/Niagara/Source/NiagaraEditor/Private/ViewModels/NiagaraSystemViewModel.cpp b/Engine/Plugins/FX/Niagara/Source/NiagaraEditor/Private/ViewModels/NiagaraSystemViewModel.cpp
--- a/Engine/Plugins/FX/Niagara/Source/NiagaraEditor/Private/ViewModels/NiagaraSystemViewModel.cpp
+++ b/Engine/Plugins/FX/Niagara/Source/NiagaraEditor/Private/ViewModels/NiagaraSystemViewModel.cpp
@@ -2286,19 +2286,25 @@ void FNiagaraSystemViewModel::RefreshEmitterHandleViewModels()
 		if (bAdd)
 		{
 			ViewModel = MakeShared<FNiagaraEmitterHandleViewModel>(bIsForDataProcessingOnly);
-			ViewModel->OnPropertyChanged().AddRaw(this, &FNiagaraSystemViewModel::EmitterHandlePropertyChanged, EmitterHandle->GetId());
 			ViewModel->OnNameChanged().AddRaw(this, &FNiagaraSystemViewModel::EmitterHandleNameChanged);
 			ViewModel->GetEmitterViewModel()->OnPropertyChanged().AddRaw(this, &FNiagaraSystemViewModel::EmitterPropertyChanged);
 			ViewModel->GetEmitterViewModel()->OnScriptCompiled().AddRaw(this, &FNiagaraSystemViewModel::ScriptCompiled);
-			ViewModel->GetEmitterViewModel()->OnScriptGraphChanged().AddRaw(this, &FNiagaraSystemViewModel::EmitterScriptGraphChanged, EmitterHandle->GetId());
 			ViewModel->GetEmitterViewModel()->OnScriptParameterStoreChanged().AddRaw(this, &FNiagaraSystemViewModel::EmitterParameterStoreChanged);
 			ViewModel->GetEmitterStackViewModel()->OnStructureChanged().AddRaw(this, &FNiagaraSystemViewModel::StackViewModelStructureChanged);
 		}
 		else
 		{
 			ViewModel = OldViewModels[i];
+
+			// These delegate bindings include the emitter handle id which may have changed since we're reusing existing view models, so remove the binding here.
+			ViewModel->OnPropertyChanged().RemoveAll(this);
+			ViewModel->GetEmitterViewModel()->OnScriptGraphChanged().RemoveAll(this);
 		}
 
+		// Add delegate bindings which require the emitter handle id passed as an additional parameter.
+		ViewModel->OnPropertyChanged().AddRaw(this, &FNiagaraSystemViewModel::EmitterHandlePropertyChanged, EmitterHandle->GetId());
+		ViewModel->GetEmitterViewModel()->OnScriptGraphChanged().AddRaw(this, &FNiagaraSystemViewModel::EmitterScriptGraphChanged, EmitterHandle->GetId());
+
 		EmitterHandleViewModels.Add(ViewModel.ToSharedRef());
 		ViewModel->Initialize(this->AsShared(), i, Simulation);
 	}
