Niagara - Don't respond to rapid iteration parameter changes on finalized stack function input entries.
NOTE: This shouldn't normally be able to happen since the delegate is unboud in FinalizeInternal, but in cases involving undo/redo it's not always possible.

diff --git a/Engine/Plugins/FX/Niagara/Source/NiagaraEditor/Private/ViewModels/Stack/NiagaraStackFunctionInput.cpp b/Engine/Plugins/FX/Niagara/Source/NiagaraEditor/Private/ViewModels/Stack/NiagaraStackFunctionInput.cpp
--- a/Engine/Plugins/FX/Niagara/Source/NiagaraEditor/Private/ViewModels/Stack/NiagaraStackFunctionInput.cpp
+++ b/Engine/Plugins/FX/Niagara/Source/NiagaraEditor/Private/ViewModels/Stack/NiagaraStackFunctionInput.cpp
@@ -3376,6 +3376,12 @@ void UNiagaraStackFunctionInput::OnGraphChanged(const struct FEdGraphEditAction&
 
 void UNiagaraStackFunctionInput::OnRapidIterationParametersChanged()
 {
+	if (IsFinalized())
+	{
+		// This delegate isn't always unbound correctly, especially when dealing with undo, so guard against responding in a finalized entry. 
+		return;
+	}
+
 	bCanResetCache.Reset();
 	bCanResetToBaseCache.Reset();
 	if (ensureMsgf(OwningModuleNode.IsValid() && OwningFunctionCallNode.IsValid(), TEXT("Stack entry with invalid module or function call not cleaned up.")))
