Niagara - Stack function input style changes.
+ Move the static icon from the value to the name column so that the values are aligned.
+ Fix lots of issues with padding.
+ Refactor the value widget generation code in SNiagaraStackFunctionInputValue for clarity and consistency and to avoid creating widgets that will never be made visible.

diff --git a/Engine/Plugins/FX/Niagara/Source/NiagaraEditor/Private/NiagaraEditorStyle.cpp b/Engine/Plugins/FX/Niagara/Source/NiagaraEditor/Private/NiagaraEditorStyle.cpp
--- a/Engine/Plugins/FX/Niagara/Source/NiagaraEditor/Private/NiagaraEditorStyle.cpp
+++ b/Engine/Plugins/FX/Niagara/Source/NiagaraEditor/Private/NiagaraEditorStyle.cpp
@@ -307,6 +307,7 @@ void FNiagaraEditorStyle::InitParameters()
 		.SetEditableTextBoxStyle(InlineEditableTextBlockEditable);
 	Set("NiagaraEditor.Graph.Node.InlineEditablePinName", InlineEditableTextBlockStyle);
 	
+	Set("NiagaraEditor.StaticInputIcon", new IMAGE_BRUSH_SVG("Icons/StaticInput", Icon12x12));
 	Set("NiagaraEditor.StaticIcon", new IMAGE_BRUSH("Icons/staticpill_16x", Icon16x16));
 	Set("NiagaraEditor.Pins.StaticConnected", new IMAGE_BRUSH("Icons/StaticPin_Connected", Icon16x16));
 	Set("NiagaraEditor.Pins.StaticDisconnected", new IMAGE_BRUSH("Icons/StaticPin_Disconnected", Icon16x16));
diff --git a/Engine/Plugins/FX/Niagara/Source/NiagaraEditorWidgets/Private/Stack/SNiagaraStackFunctionInputName.cpp b/Engine/Plugins/FX/Niagara/Source/NiagaraEditorWidgets/Private/Stack/SNiagaraStackFunctionInputName.cpp
--- a/Engine/Plugins/FX/Niagara/Source/NiagaraEditorWidgets/Private/Stack/SNiagaraStackFunctionInputName.cpp
+++ b/Engine/Plugins/FX/Niagara/Source/NiagaraEditorWidgets/Private/Stack/SNiagaraStackFunctionInputName.cpp
@@ -89,6 +89,24 @@ void SNiagaraStackFunctionInputName::Construct(const FArguments& InArgs, UNiagar
 			.ToolTipText(this, &SNiagaraStackFunctionInputName::GetToolTipText)
 		];
 	}
+
+	if (FunctionInput->GetInputType().IsStatic())
+	{
+		NameBox->AddSlot()
+		.VAlign(VAlign_Center)
+		.AutoWidth()
+		[
+			SNew(SBox)
+			.ToolTipText(LOCTEXT("TypeModifierTooltip", "This variable is Static and cannot be changed at runtime, or linked to any\n dynamic attributes. Static variables can be used to reduce script instructions & memory."))
+			.Padding(3, 2, 1, 2)
+			.Content()
+			[
+				SNew(SImage)
+				.ColorAndOpacity(FSlateColor::UseSubduedForeground())
+				.Image(FNiagaraEditorStyle::Get().GetBrush("NiagaraEditor.StaticInputIcon"))
+			]
+		];
+	}
 }
 
 void SNiagaraStackFunctionInputName::Tick(const FGeometry& AllottedGeometry, const double InCurrentTime, const float InDeltaTime)
diff --git a/Engine/Plugins/FX/Niagara/Source/NiagaraEditorWidgets/Private/Stack/SNiagaraStackFunctionInputValue.cpp b/Engine/Plugins/FX/Niagara/Source/NiagaraEditorWidgets/Private/Stack/SNiagaraStackFunctionInputValue.cpp
--- a/Engine/Plugins/FX/Niagara/Source/NiagaraEditorWidgets/Private/Stack/SNiagaraStackFunctionInputValue.cpp
+++ b/Engine/Plugins/FX/Niagara/Source/NiagaraEditorWidgets/Private/Stack/SNiagaraStackFunctionInputValue.cpp
@@ -53,6 +53,7 @@
 #include "Widgets/Input/SEditableTextBox.h"
 #include "Widgets/Input/SNumericEntryBox.h"
 #include "Widgets/Input/SSpinBox.h"
+#include "Widgets/Layout/SBorder.h"
 #include "Widgets/Layout/SBox.h"
 #include "Widgets/NiagaraHLSLSyntaxHighlighter.h"
 #include "Widgets/SBoxPanel.h"
@@ -67,6 +68,76 @@ const float TextIconSize = 16;
 bool SNiagaraStackFunctionInputValue::bLibraryOnly = true;
 FName SNiagaraStackFunctionInputValue::FavoriteActionsProfile = TEXT("StackFunctionInputActionsProfile");
 
+class SNiagaraStackFunctionInputTypeIcon : public SBorder
+{
+public:
+	SLATE_BEGIN_ARGS(SNiagaraStackFunctionInputTypeIcon) { }
+	SLATE_END_ARGS()
+
+	void Construct(const FArguments& InArgs, UNiagaraStackFunctionInput* InFunctionInput)
+	{ 
+		TSharedRef<SImage> TypePill = SNew(SImage)
+			.ColorAndOpacity(UEdGraphSchema_Niagara::GetTypeColor(InFunctionInput->GetInputType()))
+			.Image(FNiagaraEditorStyle::Get().GetBrush("NiagaraEditor.Module.TypeIconPill"));
+
+		TSharedPtr<SWidget> TypeContent;
+		UNiagaraStackFunctionInput::EValueMode InputValueMode = InFunctionInput->GetValueMode();
+		if (InputValueMode == UNiagaraStackFunctionInput::EValueMode::Linked)
+		{
+			// Linked values don't need an icon since they have a custom parameter control.
+			TypeContent = TypePill;
+		}
+		else
+		{
+			TSharedPtr<SWidget> TypeIcon;
+			if (InputValueMode == UNiagaraStackFunctionInput::EValueMode::Dynamic)
+			{
+				TypeIcon = SNew(SImage)
+				.ColorAndOpacity(FLinearColor::White)
+				.Image(FNiagaraEditorStyle::Get().GetBrush("NiagaraEditor.Module.DynamicInput"));
+			}
+			else
+			{
+				TypeIcon = SNew(STextBlock)
+				.Font(FAppStyle::Get().GetFontStyle("FontAwesome.10"))
+				.Text(FNiagaraStackEditorWidgetsUtilities::GetIconTextForInputMode(InFunctionInput->GetValueMode()));
+			}
+
+			TypeContent = SNew(SHorizontalBox)
+			+ SHorizontalBox::Slot()
+			.AutoWidth()
+			.Padding(0)
+			[
+				TypePill
+			]
+			+ SHorizontalBox::Slot()
+			.VAlign(VAlign_Center)
+			.AutoWidth()
+			.Padding(3, 2)
+			[
+				SNew(SBox)
+				.HeightOverride(TextIconSize)
+				.VAlign(VAlign_Center)
+				[
+					TypeIcon.ToSharedRef()
+				]
+			];
+		}
+
+		SBorder::Construct(SBorder::FArguments()
+			.BorderImage(FNiagaraEditorStyle::Get().GetBrush("NiagaraEditor.Module.InputTypeBorder"))
+			.Visibility(EVisibility::Visible)
+			.Padding(2)
+			[
+				TypeContent.ToSharedRef()
+			]);
+
+		SetToolTipText(FText::Format(LOCTEXT("InputTypeIconToolTipFormat", "Mode: {0}\nType: {1}"),
+			FNiagaraStackEditorWidgetsUtilities::GetIconToolTipForInputMode(InFunctionInput->GetValueMode()),
+			InFunctionInput->GetInputType().GetNameText()));
+	}
+};
+
 void SNiagaraStackFunctionInputValue::Construct(const FArguments& InArgs, UNiagaraStackFunctionInput* InFunctionInput)
 {
 	FunctionInput = InFunctionInput;
@@ -118,91 +189,7 @@ void SNiagaraStackFunctionInputValue::Construct(const FArguments& InArgs, UNiaga
 	{
 			ChildrenBox->AddSlot()
 				.VAlign(VAlign_Center)
-				.AutoWidth()
-				.Padding(0, 0, 3, 0)
-				[
-					// Value Icon
-					SNew(SHorizontalBox)
-					.Visibility(this, &SNiagaraStackFunctionInputValue::GetInputIconVisibility)
-					+ SHorizontalBox::Slot()
-					.AutoWidth()
-					[
-						// icons without type info (e.g. data interfaces)
-						SNew(SBox)
-						.WidthOverride(TextIconSize)
-						.VAlign(VAlign_Center)
-						.Visibility_Lambda([this]()
-						{
-							return (FunctionInput->GetValueMode() == UNiagaraStackFunctionInput::EValueMode::Dynamic || FunctionInput->GetValueMode() == UNiagaraStackFunctionInput::EValueMode::Linked) ? EVisibility::Collapsed : EVisibility::Visible;
-						})
-						[
-							SNew(STextBlock)
-							.Font(FAppStyle::Get().GetFontStyle("FontAwesome.10"))
-							.Text(this, &SNiagaraStackFunctionInputValue::GetInputIconText)
-							.ToolTipText(this, &SNiagaraStackFunctionInputValue::GetInputIconToolTip)
-							.ColorAndOpacity(this, &SNiagaraStackFunctionInputValue::GetInputIconColor)
-						]
-					]
-				];
-
-				ChildrenBox->AddSlot()
-					.AutoWidth()
-					[
-						// icons type info pill
-						SNew(SBorder)
-						.BorderImage(FNiagaraEditorStyle::Get().GetBrush("NiagaraEditor.Module.InputTypeBorder"))
-						.BorderBackgroundColor(FLinearColor::White)
-						.Padding(FMargin(2, 2, 2, 2))
-						[
-							SNew(SHorizontalBox)
-							.Visibility_Lambda([this]()
-							{
-								return (FunctionInput->GetValueMode() == UNiagaraStackFunctionInput::EValueMode::Dynamic || FunctionInput->GetValueMode() == UNiagaraStackFunctionInput::EValueMode::Linked) ? EVisibility::Visible : EVisibility::Collapsed;
-							})
-							+ SHorizontalBox::Slot()
-							.VAlign(VAlign_Center)
-							.HAlign(HAlign_Center)
-							.AutoWidth()
-							[
-								SNew(SImage)
-								.Visibility_Lambda([this](){ return (FunctionInput->GetValueMode() == UNiagaraStackFunctionInput::EValueMode::Dynamic || FunctionInput->GetValueMode() == UNiagaraStackFunctionInput::EValueMode::Linked) ? EVisibility::Visible : EVisibility::Collapsed; })
-								.ToolTipText(this, &SNiagaraStackFunctionInputValue::GetInputIconToolTip)
-								.ColorAndOpacity_Lambda([this]() { return UEdGraphSchema_Niagara::GetTypeColor(FunctionInput->GetInputType()); })
-								.Image(FNiagaraEditorStyle::Get().GetBrush("NiagaraEditor.Module.TypeIconPill"))
-							]
-							+ SHorizontalBox::Slot()
-							.VAlign(VAlign_Center)
-							.HAlign(HAlign_Center)
-							.AutoWidth()
-							[
-								SNew(SImage)
-								.Visibility_Lambda([this](){ return FunctionInput->GetValueMode() == UNiagaraStackFunctionInput::EValueMode::Dynamic ? EVisibility::Visible : EVisibility::Collapsed; })
-								.ToolTipText(this, &SNiagaraStackFunctionInputValue::GetInputIconToolTip)
-								.ColorAndOpacity(FLinearColor::White)
-								.Image(FNiagaraEditorStyle::Get().GetBrush("NiagaraEditor.Module.DynamicInput"))
-							]
-						]
-					];
-
-			ChildrenBox->AddSlot()
-				.VAlign(VAlign_Center)
-				.AutoWidth()
-				.Padding(0, 0, 3, 0)
-				[
-					// Type Modifier Icon
-					SNew(SBox)
-					.WidthOverride(TextIconSize)
-					.VAlign(VAlign_Center)
-					.Visibility(this, &SNiagaraStackFunctionInputValue::GetTypeModifierIconVisibility)
-					[
-						SNew(SImage)
-						.Image(GetTypeModifierIcon())
-						.ToolTipText(GetTypeModifierIconToolTip())
-					]
-				];
-
-			ChildrenBox->AddSlot()
-				.VAlign(VAlign_Center)
+				.Padding(0)
 				[
 					// Value container and widgets.
 					SAssignNew(ValueContainer, SBox)
@@ -211,6 +198,8 @@ void SNiagaraStackFunctionInputValue::Construct(const FArguments& InArgs, UNiaga
 						ConstructValueWidgets()
 					]
 				];
+
+			ValueModeForGeneratedWidgets = FunctionInput->GetValueMode();
 	}
 
 	if (LayoutMode == ELayoutMode::FullRow)
@@ -292,8 +281,6 @@ void SNiagaraStackFunctionInputValue::Construct(const FArguments& InArgs, UNiaga
 			.VAlign(VAlign_Center)
 		];
 	}
-
-	ValueModeForGeneratedWidgets = FunctionInput->GetValueMode();
 }
 
 void SNiagaraStackFunctionInputValue::Tick(const FGeometry& AllottedGeometry, const double InCurrentTime, const float InDeltaTime)
@@ -311,50 +298,30 @@ TSharedRef<SWidget> SNiagaraStackFunctionInputValue::ConstructValueWidgets()
 	LocalValueStructParameterEditor.Reset();
 	LocalValueStructDetailsView.Reset();
 
-	switch (FunctionInput->GetValueMode())
+	TSharedRef<SHorizontalBox> ValueBox = SNew(SHorizontalBox);
+
+	bool bShowTypeIcon = true;
+	switch(FunctionInput->GetValueMode())
 	{
 	case UNiagaraStackFunctionInput::EValueMode::Local:
-	{
-		return ConstructLocalValueStructWidget();
-	}
-	case UNiagaraStackFunctionInput::EValueMode::Linked:
-	{
-		TSharedRef<SWidget> ParameterWidget = SNew(SNiagaraParameterName)
-			.ReadOnlyTextStyle(FNiagaraEditorStyle::Get(), "NiagaraEditor.ParameterText")
-			.ParameterName(this, &SNiagaraStackFunctionInputValue::GetLinkedValueHandleName)
-			.OnDoubleClicked(this, &SNiagaraStackFunctionInputValue::OnLinkedInputDoubleClicked);
+		ConstructLocalValueStructWidgets(ValueBox);
+		bShowTypeIcon = false;
+		break;
 
-		FNiagaraParameterHandle ParameterHandle = FNiagaraParameterHandle(FunctionInput->GetLinkedParameterValue().GetName());
-		if(ParameterHandle.IsUserHandle())
-		{
-			TArray<FNiagaraVariable> UserParameters;
-			FunctionInput->GetSystemViewModel()->GetSystem().GetExposedParameters().GetUserParameters(UserParameters);
-			FNiagaraVariable* MatchingVariable = UserParameters.FindByPredicate([ParameterHandle](const FNiagaraVariable& Variable)
-			{
-				return Variable.GetName().ToString() == ParameterHandle.GetName().ToString();
-			});
+	case UNiagaraStackFunctionInput::EValueMode::Linked:
+		ConstructLinkedValueWidgets(ValueBox);
+		break;
 
-			if(MatchingVariable)
-			{
-				FText Tooltip = FNiagaraEditorUtilities::UserParameters::GetScriptVariableForUserParameter(*MatchingVariable, FunctionInput->GetSystemViewModel())->Metadata.Description;
-				if(!Tooltip.IsEmpty())
-				{
-					ParameterWidget->SetToolTipText(Tooltip);
-				}
-			}
-		}
-			
-		return ParameterWidget; 
-	}
 	case UNiagaraStackFunctionInput::EValueMode::Data:
-	{
-		return SNew(STextBlock)
-			.TextStyle(FNiagaraEditorStyle::Get(), "NiagaraEditor.ParameterText")
-			.Text(this, &SNiagaraStackFunctionInputValue::GetDataValueText);
-	}
+		ConstructTextValueWidgets(ValueBox, TAttribute<FText>::CreateSP(this, &SNiagaraStackFunctionInputValue::GetDataValueText));
+		break;
+
 	case UNiagaraStackFunctionInput::EValueMode::ObjectAsset:
-	{
-		return SNew(SObjectPropertyEntryBox)
+		ValueBox->AddSlot()
+		.VAlign(VAlign_Center)
+		.Padding(0)
+		[
+			SNew(SObjectPropertyEntryBox)
 			.AllowedClass(FunctionInput->GetInputType().GetClass())
 			.ObjectPath_Lambda(
 				[this]() -> FString
@@ -372,82 +339,26 @@ TSharedRef<SWidget> SNiagaraStackFunctionInputValue::ConstructValueWidgets()
 				{
 					FunctionInput->SetObjectAssetValue(AssetData.GetAsset());
 				}
-			);
-	}
+			)
+		];
+		break;
+
 	case UNiagaraStackFunctionInput::EValueMode::Dynamic:
-	{
-		TSharedRef<SWidget> DynamicInputText = SNew(STextBlock)
-			.TextStyle(FNiagaraEditorStyle::Get(), "NiagaraEditor.ParameterText")
-			.Text(this, &SNiagaraStackFunctionInputValue::GetDynamicValueText)
-			.OnDoubleClicked(this, &SNiagaraStackFunctionInputValue::DynamicInputTextDoubleClicked);
-		if (FunctionInput->IsScratchDynamicInput())
-		{
-			return SNew(SHorizontalBox)
-				+ SHorizontalBox::Slot()
-				.VAlign(VAlign_Center)
-				[
-					DynamicInputText
-				]
-				+ SHorizontalBox::Slot()
-				.AutoWidth()
-				[
-					SNew(SButton)
-					.ButtonStyle(FAppStyle::Get(), "RoundButton")
-					.OnClicked(this, &SNiagaraStackFunctionInputValue::ScratchButtonPressed)
-					.ToolTipText(LOCTEXT("OpenInScratchToolTip", "Open this dynamic input in the scratch pad."))
-					.ContentPadding(FMargin(1.0f, 0.0f))
-					.Content()
-					[
-						SNew(SImage)
-						.Image(FNiagaraEditorStyle::Get().GetBrush("Tab.ScratchPad"))
-					]
-				];
-		}
-		// the function script could be wiped (deleted scratch pad script or missing asset)
-		else if (FunctionInput->GetDynamicInputNode()->FunctionScript && FunctionInput->GetDynamicInputNode()->FunctionScript->IsVersioningEnabled())
-		{
-			return SNew(SHorizontalBox)
-                + SHorizontalBox::Slot()
-                .VAlign(VAlign_Center)
-                [
-                    DynamicInputText
-                ]
-                + SHorizontalBox::Slot()
-                .AutoWidth()
-                [
-	                SNew(SComboButton)
-				    .HasDownArrow(false)
-				    .ButtonStyle(FAppStyle::Get(), "HoverHintOnly")
-				    .ForegroundColor(FSlateColor::UseForeground())
-				    .OnGetMenuContent(this, &SNiagaraStackFunctionInputValue::GetVersionSelectorDropdownMenu)
-				    .ContentPadding(FMargin(2))
-				    .ToolTipText(LOCTEXT("VersionTooltip", "Change the version of this module script"))
-				    .HAlign(HAlign_Center)
-				    .VAlign(VAlign_Center)
-				    .ButtonContent()
-				    [
-				        SNew(STextBlock)
-				        .Font(FAppStyle::Get().GetFontStyle("FontAwesome.10"))
-				        .ColorAndOpacity(this, &SNiagaraStackFunctionInputValue::GetVersionSelectorColor)
-				        .Text(FEditorFontGlyphs::Random)
-				    ]
-                ];
-		}
-		else
-		{
-			return DynamicInputText;
-		}
-	}
+		ConstructDynamicValueWidgets(ValueBox);
+		break;
+
 	case UNiagaraStackFunctionInput::EValueMode::DefaultFunction:
-	{
-		return SNew(STextBlock)
-			.TextStyle(FNiagaraEditorStyle::Get(), "NiagaraEditor.ParameterText")
-			.Text(this, &SNiagaraStackFunctionInputValue::GetDefaultFunctionText);
-	}
+		ConstructTextValueWidgets(ValueBox, TAttribute<FText>::CreateSP(this, &SNiagaraStackFunctionInputValue::GetDefaultFunctionText));
+		break;
+
 	case UNiagaraStackFunctionInput::EValueMode::Expression:
 	{
 		const FEditableTextBoxStyle& TextBoxStyle = FCoreStyle::Get().GetWidgetStyle<FEditableTextBoxStyle>("NormalEditableTextBox");
-		return SNew(SBorder)
+		ValueBox->AddSlot()
+		.VAlign(VAlign_Center)
+		.Padding(0)
+		[
+			SNew(SBorder)
 			.BorderBackgroundColor(TextBoxStyle.BackgroundColor)
 			.Padding(TextBoxStyle.Padding)
 			.BorderImage(&TextBoxStyle.BackgroundImageNormal)
@@ -458,23 +369,31 @@ TSharedRef<SWidget> SNiagaraStackFunctionInputValue::ConstructValueWidgets()
 				.AllowMultiLine(false)
 				.Text_UObject(FunctionInput, &UNiagaraStackFunctionInput::GetCustomExpressionText)
 				.OnTextCommitted(this, &SNiagaraStackFunctionInputValue::OnExpressionTextCommitted)
-			];
+			]
+		];
+		break;
 	}
+
 	case UNiagaraStackFunctionInput::EValueMode::InvalidOverride:
-	{
-		return SNew(STextBlock)
-			.TextStyle(FNiagaraEditorStyle::Get(), "NiagaraEditor.ParameterText")
-			.Text(LOCTEXT("InvalidOverrideText", "Invalid Script Value"));
-	}
+		ConstructTextValueWidgets(ValueBox, LOCTEXT("InvalidOverrideText", "Invalid Script Value"));
+		break;
+
 	case UNiagaraStackFunctionInput::EValueMode::UnsupportedDefault:
-	{
-		return SNew(STextBlock)
-			.TextStyle(FNiagaraEditorStyle::Get(), "NiagaraEditor.ParameterText")
-			.Text(LOCTEXT("UnsupportedDefault", "Custom Default"));
+		ConstructTextValueWidgets(ValueBox, LOCTEXT("UnsupportedDefault", "Custom Default"));
+		break;
 	}
-	default:
-		return SNullWidget::NullWidget;
+
+	if (bShowTypeIcon)
+	{
+		ValueBox->InsertSlot(0)
+			.Padding(0, 0, 3, 0)
+			.AutoWidth()
+			[
+				SNew(SNiagaraStackFunctionInputTypeIcon, FunctionInput)
+			];
 	}
+
+	return ValueBox;
 }
 
 TSharedRef<SWidget> SNiagaraStackFunctionInputValue::GetVersionSelectorDropdownMenu()
@@ -563,7 +482,19 @@ bool SNiagaraStackFunctionInputValue::GetEntryEnabled() const
 	return FunctionInput->IsFinalized() == false && FunctionInput->GetIsEnabledAndOwnerIsEnabled();
 }
 
-TSharedRef<SWidget> SNiagaraStackFunctionInputValue::ConstructLocalValueStructWidget()
+void SNiagaraStackFunctionInputValue::ConstructTextValueWidgets(TSharedRef<SHorizontalBox>& ValueBox, TAttribute<FText> GetText)
+{
+	ValueBox->AddSlot()
+		.VAlign(VAlign_Center)
+		.Padding(0)
+		[
+			SNew(STextBlock)
+			.TextStyle(FNiagaraEditorStyle::Get(), "NiagaraEditor.ParameterText")
+			.Text(GetText)
+		];
+}
+
+void SNiagaraStackFunctionInputValue::ConstructLocalValueStructWidgets(TSharedRef<SHorizontalBox>& ValueBox)
 {
 	LocalValueStructParameterEditor.Reset();
 	LocalValueStructDetailsView.Reset();
@@ -591,9 +522,10 @@ TSharedRef<SWidget> SNiagaraStackFunctionInputValue::ConstructLocalValueStructWi
 
 			LocalValueStructParameterEditor = ParameterEditor;
 
-			return SNew(SBox)
+			ValueBox->AddSlot()
 				.HAlign(ParameterEditor->GetHorizontalAlignment())
 				.VAlign(ParameterEditor->GetVerticalAlignment())
+				.Padding(0)
 				[
 					ParameterEditor.ToSharedRef()
 				];
@@ -619,12 +551,104 @@ TSharedRef<SWidget> SNiagaraStackFunctionInputValue::ConstructLocalValueStructWi
 			StructureDetailsView->GetOnFinishedChangingPropertiesDelegate().AddSP(this, &SNiagaraStackFunctionInputValue::ParameterPropertyValueChanged);
 
 			LocalValueStructDetailsView = StructureDetailsView;
-			return StructureDetailsView->GetWidget().ToSharedRef();
+
+			ValueBox->AddSlot()
+				.Padding(0)
+				[
+					StructureDetailsView->GetWidget().ToSharedRef()
+				];
 		}
 	}
-	else
+}
+
+void SNiagaraStackFunctionInputValue::ConstructLinkedValueWidgets(TSharedRef<SHorizontalBox>& ValueBox)
+{
+	TSharedRef<SWidget> ParameterWidget = SNew(SNiagaraParameterName)
+		.ReadOnlyTextStyle(FNiagaraEditorStyle::Get(), "NiagaraEditor.ParameterText")
+		.ParameterName(this, &SNiagaraStackFunctionInputValue::GetLinkedValueHandleName)
+		.OnDoubleClicked(this, &SNiagaraStackFunctionInputValue::OnLinkedInputDoubleClicked);
+
+	FNiagaraParameterHandle ParameterHandle = FNiagaraParameterHandle(FunctionInput->GetLinkedParameterValue().GetName());
+	if (ParameterHandle.IsUserHandle())
 	{
-		return SNullWidget::NullWidget;
+		TArray<FNiagaraVariable> UserParameters;
+		FunctionInput->GetSystemViewModel()->GetSystem().GetExposedParameters().GetUserParameters(UserParameters);
+		FNiagaraVariable* MatchingVariable = UserParameters.FindByPredicate([ParameterHandle](const FNiagaraVariable& Variable)
+			{
+				return Variable.GetName().ToString() == ParameterHandle.GetName().ToString();
+			});
+
+		if (MatchingVariable)
+		{
+			FText Tooltip = FNiagaraEditorUtilities::UserParameters::GetScriptVariableForUserParameter(*MatchingVariable, FunctionInput->GetSystemViewModel())->Metadata.Description;
+			if (!Tooltip.IsEmpty())
+			{
+				ParameterWidget->SetToolTipText(Tooltip);
+			}
+		}
+	}
+
+	ValueBox->AddSlot()
+		.Padding(0)
+		[
+			ParameterWidget
+		];
+}
+
+void SNiagaraStackFunctionInputValue::ConstructDynamicValueWidgets(TSharedRef<SHorizontalBox>& ValueBox)
+{
+	ValueBox->AddSlot()
+		.VAlign(VAlign_Center)
+		.Padding(0)
+		[
+			SNew(STextBlock)
+				.TextStyle(FNiagaraEditorStyle::Get(), "NiagaraEditor.ParameterText")
+				.Text(this, &SNiagaraStackFunctionInputValue::GetDynamicValueText)
+				.OnDoubleClicked(this, &SNiagaraStackFunctionInputValue::DynamicInputTextDoubleClicked)
+		];
+
+	if (FunctionInput->IsScratchDynamicInput())
+	{
+		ValueBox->AddSlot()
+			.VAlign(VAlign_Center)
+			.AutoWidth()
+			[
+				SNew(SButton)
+				.ButtonStyle(FAppStyle::Get(), "RoundButton")
+				.OnClicked(this, &SNiagaraStackFunctionInputValue::ScratchButtonPressed)
+				.ToolTipText(LOCTEXT("OpenInScratchToolTip", "Open this dynamic input in the scratch pad."))
+				.ContentPadding(FMargin(1.0f, 0.0f))
+				.Content()
+				[
+					SNew(SImage)
+					.Image(FNiagaraEditorStyle::Get().GetBrush("Tab.ScratchPad"))
+				]
+			];
+	}
+	else if (FunctionInput->GetDynamicInputNode()->FunctionScript && FunctionInput->GetDynamicInputNode()->FunctionScript->IsVersioningEnabled())
+	{
+		// the function script could be wiped (deleted scratch pad script or missing asset)
+		ValueBox->AddSlot()
+			.VAlign(VAlign_Center)
+			.AutoWidth()
+			[
+				SNew(SComboButton)
+				.HasDownArrow(false)
+				.ButtonStyle(FAppStyle::Get(), "HoverHintOnly")
+				.ForegroundColor(FSlateColor::UseForeground())
+				.OnGetMenuContent(this, &SNiagaraStackFunctionInputValue::GetVersionSelectorDropdownMenu)
+				.ContentPadding(FMargin(2))
+				.ToolTipText(LOCTEXT("VersionTooltip", "Change the version of this module script"))
+				.HAlign(HAlign_Center)
+				.VAlign(VAlign_Center)
+				.ButtonContent()
+				[
+					SNew(STextBlock)
+					.Font(FAppStyle::Get().GetFontStyle("FontAwesome.10"))
+					.ColorAndOpacity(this, &SNiagaraStackFunctionInputValue::GetVersionSelectorColor)
+					.Text(FEditorFontGlyphs::Random)
+				]
+			];
 	}
 }
 
@@ -635,33 +659,26 @@ void SNiagaraStackFunctionInputValue::OnInputValueChanged()
 		return;
 	}
 
-	if (ValueModeForGeneratedWidgets != FunctionInput->GetValueMode())
+	if (ValueModeForGeneratedWidgets == FunctionInput->GetValueMode() &&
+		ValueModeForGeneratedWidgets == UNiagaraStackFunctionInput::EValueMode::Local &&
+		DisplayedLocalValueStruct.IsValid() &&
+		DisplayedLocalValueStruct->GetStruct() == FunctionInput->GetLocalValueStruct()->GetStruct())
 	{
-		ValueContainer->SetContent(ConstructValueWidgets());
-		ValueModeForGeneratedWidgets = FunctionInput->GetValueMode();
+		FNiagaraEditorUtilities::CopyDataTo(*DisplayedLocalValueStruct.Get(), *FunctionInput->GetLocalValueStruct().Get());
+		if (LocalValueStructParameterEditor.IsValid())
+		{
+			LocalValueStructParameterEditor->UpdateInternalValueFromStruct(DisplayedLocalValueStruct.ToSharedRef());
+		}
+		if (LocalValueStructDetailsView.IsValid())
+		{
+			LocalValueStructDetailsView->SetStructureData(TSharedPtr<FStructOnScope>());
+			LocalValueStructDetailsView->SetStructureData(DisplayedLocalValueStruct);
+		}
 	}
 	else
 	{
-		if (ValueModeForGeneratedWidgets == UNiagaraStackFunctionInput::EValueMode::Local && DisplayedLocalValueStruct.IsValid())
-		{
-			if (DisplayedLocalValueStruct->GetStruct() == FunctionInput->GetLocalValueStruct()->GetStruct())
-			{
-				FNiagaraEditorUtilities::CopyDataTo(*DisplayedLocalValueStruct.Get(), *FunctionInput->GetLocalValueStruct().Get());
-				if (LocalValueStructParameterEditor.IsValid())
-				{
-					LocalValueStructParameterEditor->UpdateInternalValueFromStruct(DisplayedLocalValueStruct.ToSharedRef());
-				}
-				if (LocalValueStructDetailsView.IsValid())
-				{
-					LocalValueStructDetailsView->SetStructureData(TSharedPtr<FStructOnScope>());
-					LocalValueStructDetailsView->SetStructureData(DisplayedLocalValueStruct);
-				}
-			}
-			else
-			{
-				ValueContainer->SetContent(ConstructLocalValueStructWidget());
-			}
-		}
+		ValueContainer->SetContent(ConstructValueWidgets());
+		ValueModeForGeneratedWidgets = FunctionInput->GetValueMode();
 	}
 }
 
@@ -989,48 +1006,6 @@ FReply SNiagaraStackFunctionInputValue::ResetToBaseButtonPressed() const
 	return FReply::Handled();
 }
 
-EVisibility SNiagaraStackFunctionInputValue::GetInputIconVisibility() const
-{
-	return FunctionInput->GetValueMode() == UNiagaraStackFunctionInput::EValueMode::Local
-		? EVisibility::Collapsed
-		: EVisibility::Visible;
-}
-
-FText SNiagaraStackFunctionInputValue::GetInputIconText() const
-{
-	return FNiagaraStackEditorWidgetsUtilities::GetIconTextForInputMode(FunctionInput->GetValueMode());
-}
-
-FText SNiagaraStackFunctionInputValue::GetInputIconToolTip() const
-{
-	return FNiagaraStackEditorWidgetsUtilities::GetIconToolTipForInputMode(FunctionInput->GetValueMode());
-}
-
-FSlateColor SNiagaraStackFunctionInputValue::GetInputIconColor() const
-{
-	return FNiagaraEditorWidgetsStyle::Get().GetColor(FNiagaraStackEditorWidgetsUtilities::GetIconColorNameForInputMode(FunctionInput->GetValueMode()));
-}
-
-EVisibility SNiagaraStackFunctionInputValue::GetTypeModifierIconVisibility() const
-{
-	return FunctionInput->GetInputType().IsStatic() ? EVisibility::Visible : EVisibility::Collapsed;
-}
-
-const FSlateBrush* SNiagaraStackFunctionInputValue::GetTypeModifierIcon() const
-{
-	return FNiagaraEditorWidgetsStyle::Get().GetBrush("NiagaraEditor.Stack.StaticInputValue");
-}
-
-FText SNiagaraStackFunctionInputValue::GetTypeModifierIconToolTip() const
-{
-	return LOCTEXT("TypeModifierTooltip", "This variable is Static and cannot be changed at runtime, or linked to any dynamic attributes. Static variables can be used to reduce script instructions & memory.");
-}
-
-FSlateColor SNiagaraStackFunctionInputValue::GetTypeModifierIconColor() const
-{
-	return UEdGraphSchema_Niagara::GetTypeColor(FunctionInput->GetInputType());
-}
-
 FReply SNiagaraStackFunctionInputValue::OnFunctionInputDrop(const FGeometry& InGeometry, const FDragDropEvent& InDragDropEvent)
 {
 	TSharedPtr<FNiagaraParameterDragOperation> InputDragDropOperation = InDragDropEvent.GetOperationAs<FNiagaraParameterDragOperation>();
diff --git a/Engine/Plugins/FX/Niagara/Source/NiagaraEditorWidgets/Private/Stack/SNiagaraStackFunctionInputValue.h b/Engine/Plugins/FX/Niagara/Source/NiagaraEditorWidgets/Private/Stack/SNiagaraStackFunctionInputValue.h
--- a/Engine/Plugins/FX/Niagara/Source/NiagaraEditorWidgets/Private/Stack/SNiagaraStackFunctionInputValue.h
+++ b/Engine/Plugins/FX/Niagara/Source/NiagaraEditorWidgets/Private/Stack/SNiagaraStackFunctionInputValue.h
@@ -15,6 +15,7 @@ class UNiagaraStackFunctionInput;
 class UNiagaraScript;
 class SNiagaraParameterEditor;
 class SBox;
+class SHorizontalBox;
 class IStructureDetailsView;
 class SComboButton;
 struct FGraphActionListBuilderBase;
@@ -55,7 +56,13 @@ private:
 
 	bool GetEntryEnabled() const;
 
-	TSharedRef<SWidget> ConstructLocalValueStructWidget();
+	void ConstructTextValueWidgets(TSharedRef<SHorizontalBox>& ValueBox, TAttribute<FText> GetText);
+
+	void ConstructLocalValueStructWidgets(TSharedRef<SHorizontalBox>& ValueBox);
+
+	void ConstructLinkedValueWidgets(TSharedRef<SHorizontalBox>& ValueBox);
+
+	void ConstructDynamicValueWidgets(TSharedRef<SHorizontalBox>& ValueBox);
 
 	void OnInputValueChanged();
 
@@ -150,22 +157,6 @@ private:
 
 	FReply ResetToBaseButtonPressed() const;
 
-	EVisibility GetInputIconVisibility() const;
-
-	FText GetInputIconText() const;
-
-	FText GetInputIconToolTip() const;
-
-	FSlateColor GetInputIconColor() const;
-
-	EVisibility GetTypeModifierIconVisibility() const;
-
-	const FSlateBrush* GetTypeModifierIcon() const;
-
-	FText GetTypeModifierIconToolTip() const;
-
-	FSlateColor GetTypeModifierIconColor() const;
-
 	FReply OnFunctionInputDrop(const FGeometry& InGeometry, const FDragDropEvent& InDragDropEvent);
 
 	bool OnFunctionInputAllowDrop(TSharedPtr<FDragDropOperation> DragDropOperation);
@@ -186,7 +177,7 @@ private:
 	ELayoutMode LayoutMode;
 
 	TSharedPtr<SBox> ValueContainer;
-	UNiagaraStackFunctionInput::EValueMode ValueModeForGeneratedWidgets;
+	UNiagaraStackFunctionInput::EValueMode ValueModeForGeneratedWidgets = UNiagaraStackFunctionInput::EValueMode::None;
 
 	TSharedPtr<FStructOnScope> DisplayedLocalValueStruct;
 	TSharedPtr<SNiagaraParameterEditor> LocalValueStructParameterEditor;
