null check. Static Analysis fix.

diff --git a/Engine/Plugins/FX/Niagara/Source/NiagaraBlueprintNodes/Private/K2Node_DataChannel_WithContext.cpp b/Engine/Plugins/FX/Niagara/Source/NiagaraBlueprintNodes/Private/K2Node_DataChannel_WithContext.cpp
--- a/Engine/Plugins/FX/Niagara/Source/NiagaraBlueprintNodes/Private/K2Node_DataChannel_WithContext.cpp
+++ b/Engine/Plugins/FX/Niagara/Source/NiagaraBlueprintNodes/Private/K2Node_DataChannel_WithContext.cpp
@@ -64,7 +64,7 @@ struct FK2Node_MakeNDCAccessContextInst_OptionalPinManager : public FOptionalPin
 	void GetRecordDefaults(FProperty* TestProperty, FOptionalPinFromProperty& Record) const
 	{
 		Record.bCanToggleVisibility = true;
-		Record.bShowPin = TestProperty->HasMetaData(TEXT("NDCAccessContextInput"));
+		Record.bShowPin = TestProperty ? TestProperty->HasMetaData(TEXT("NDCAccessContextInput")) : false;
 		Record.bIsMarkedForAdvancedDisplay = TestProperty ? TestProperty->HasAnyPropertyFlags(CPF_AdvancedDisplay) : false;
 		bHasAdvancedPins |= Record.bIsMarkedForAdvancedDisplay;
 	}
