Niagara - Change the data interface base class so that it's not longer a UNiagaraMergeable.

diff --git a/Engine/Plugins/FX/Niagara/Source/NiagaraCore/Private/NiagaraMergeable.cpp b/Engine/Plugins/FX/Niagara/Source/NiagaraCore/Private/NiagaraMergeable.cpp
--- a/Engine/Plugins/FX/Niagara/Source/NiagaraCore/Private/NiagaraMergeable.cpp
+++ b/Engine/Plugins/FX/Niagara/Source/NiagaraCore/Private/NiagaraMergeable.cpp
@@ -13,11 +13,6 @@ UNiagaraMergeable::UNiagaraMergeable()
 }
 
 #if WITH_EDITOR
-void UNiagaraMergeable::PostEditChangeProperty(struct FPropertyChangedEvent& PropertyChangedEvent)
-{
-	Super::PostEditChangeProperty(PropertyChangedEvent);
-	OnChangedDelegate.Broadcast();
-}
 
 bool UNiagaraMergeable::Equals(const UNiagaraMergeable* Other)
 {
@@ -44,11 +39,6 @@ bool UNiagaraMergeable::Equals(const UNiagaraMergeable* Other)
 	return true;
 }
 
-UNiagaraMergeable::FOnChanged& UNiagaraMergeable::OnChanged()
-{
-	return OnChangedDelegate;
-}
-
 FGuid UNiagaraMergeable::GetMergeId() const
 {
 	return MergeId;
diff --git a/Engine/Plugins/FX/Niagara/Source/NiagaraCore/Private/NiagaraNotifyOnChanged.cpp b/Engine/Plugins/FX/Niagara/Source/NiagaraCore/Private/NiagaraNotifyOnChanged.cpp
new file mode 100644
--- /dev/null
+++ b/Engine/Plugins/FX/Niagara/Source/NiagaraCore/Private/NiagaraNotifyOnChanged.cpp
@@ -0,0 +1,20 @@
+// Copyright Epic Games, Inc. All Rights Reserved.
+
+#include "NiagaraNotifyOnChanged.h"
+
+#include UE_INLINE_GENERATED_CPP_BY_NAME(NiagaraNotifyOnChanged)
+
+#if WITH_EDITOR
+
+void UNiagaraNotifyOnChanged::PostEditChangeProperty(struct FPropertyChangedEvent& PropertyChangedEvent)
+{
+	Super::PostEditChangeProperty(PropertyChangedEvent);
+	OnChangedDelegate.Broadcast();
+}
+
+UNiagaraNotifyOnChanged::FOnChanged& UNiagaraNotifyOnChanged::OnChanged()
+{
+	return OnChangedDelegate;
+}
+
+#endif // WITH_EDITOR
diff --git a/Engine/Plugins/FX/Niagara/Source/NiagaraCore/Public/NiagaraDataInterfaceBase.h b/Engine/Plugins/FX/Niagara/Source/NiagaraCore/Public/NiagaraDataInterfaceBase.h
--- a/Engine/Plugins/FX/Niagara/Source/NiagaraCore/Public/NiagaraDataInterfaceBase.h
+++ b/Engine/Plugins/FX/Niagara/Source/NiagaraCore/Public/NiagaraDataInterfaceBase.h
@@ -41,7 +41,7 @@ struct FNiagaraDataInterfaceParametersCS
 
 /** Base class for all Niagara data interfaces. */
 UCLASS(abstract, EditInlineNew, MinimalAPI)
-class UNiagaraDataInterfaceBase : public UNiagaraMergeable
+class UNiagaraDataInterfaceBase : public UNiagaraNotifyOnChanged
 {
 	GENERATED_UCLASS_BODY()
 
diff --git a/Engine/Plugins/FX/Niagara/Source/NiagaraCore/Public/NiagaraMergeable.h b/Engine/Plugins/FX/Niagara/Source/NiagaraCore/Public/NiagaraMergeable.h
--- a/Engine/Plugins/FX/Niagara/Source/NiagaraCore/Public/NiagaraMergeable.h
+++ b/Engine/Plugins/FX/Niagara/Source/NiagaraCore/Public/NiagaraMergeable.h
@@ -2,28 +2,21 @@
 
 #pragma once
 
+#include "NiagaraNotifyOnChanged.h"
+
 #include "NiagaraMergeable.generated.h"
 
 UCLASS(MinimalAPI)
-class UNiagaraMergeable : public UObject
+class UNiagaraMergeable : public UNiagaraNotifyOnChanged
 {
 	GENERATED_BODY()
 
-public:
-#if WITH_EDITOR
-	DECLARE_MULTICAST_DELEGATE(FOnChanged);
-#endif
-
 public:
 	NIAGARACORE_API UNiagaraMergeable();
 
 #if WITH_EDITOR
-	NIAGARACORE_API virtual void PostEditChangeProperty(struct FPropertyChangedEvent& PropertyChangedEvent);
-
 	NIAGARACORE_API bool Equals(const UNiagaraMergeable* Other);
 
-	NIAGARACORE_API FOnChanged& OnChanged();
-
 	NIAGARACORE_API FGuid GetMergeId() const;
 
 protected:
@@ -31,10 +24,6 @@ protected:
 #endif
 
 private:
-#if WITH_EDITOR
-	FOnChanged OnChangedDelegate;
-#endif
-
 #if WITH_EDITORONLY_DATA
 	UPROPERTY()
 	FGuid MergeId;
diff --git a/Engine/Plugins/FX/Niagara/Source/NiagaraCore/Public/NiagaraNotifyOnChanged.h b/Engine/Plugins/FX/Niagara/Source/NiagaraCore/Public/NiagaraNotifyOnChanged.h
new file mode 100644
--- /dev/null
+++ b/Engine/Plugins/FX/Niagara/Source/NiagaraCore/Public/NiagaraNotifyOnChanged.h
@@ -0,0 +1,26 @@
+// Copyright Epic Games, Inc. All Rights Reserved.
+
+#pragma once
+
+#include "NiagaraNotifyOnChanged.generated.h"
+
+UCLASS()
+class UNiagaraNotifyOnChanged : public UObject
+{
+	GENERATED_BODY()
+
+#if WITH_EDITOR
+
+public:
+	DECLARE_MULTICAST_DELEGATE(FOnChanged);
+
+public:
+	NIAGARACORE_API virtual void PostEditChangeProperty(struct FPropertyChangedEvent& PropertyChangedEvent);
+	NIAGARACORE_API FOnChanged& OnChanged();
+
+private:
+	FOnChanged OnChangedDelegate;
+
+#endif // WITH_EDITOR
+
+};
diff --git a/Engine/Plugins/FX/Niagara/Source/NiagaraEditor/Private/NiagaraHlslTranslator.cpp b/Engine/Plugins/FX/Niagara/Source/NiagaraEditor/Private/NiagaraHlslTranslator.cpp
--- a/Engine/Plugins/FX/Niagara/Source/NiagaraEditor/Private/NiagaraHlslTranslator.cpp
+++ b/Engine/Plugins/FX/Niagara/Source/NiagaraEditor/Private/NiagaraHlslTranslator.cpp
@@ -7858,9 +7858,8 @@ int32 TNiagaraHlslTranslator<GraphBridge>::RegisterUObjectPath(const FNiagaraVar
 template<typename GraphBridge>
 int32 TNiagaraHlslTranslator<GraphBridge>::RegisterDataInterface(const FNiagaraVariable& Var, const UNiagaraDataInterface* DataInterface, bool bPlaceholder, bool bAddParameterMapRead)
 {
-	FString Id = DataInterface ? *DataInterface->GetMergeId().ToString() : TEXT("??");
-	FString PathName = DataInterface ? *DataInterface->GetPathName() : TEXT("XX");
-	//UE_LOG(LogNiagaraEditor, Log, TEXT("RegisterDataInterface %s %s %s '%s' %s"), *Var.GetName().ToString(), *Var.GetType().GetName(), *Id, *PathName,
+	//FString PathName = DataInterface ? *DataInterface->GetPathName() : TEXT("XX");
+	//UE_LOG(LogNiagaraEditor, Log, TEXT("RegisterDataInterface %s %s %s '%s' %s"), *Var.GetName().ToString(), *Var.GetType().GetName(), *PathName,
 	//	bPlaceholder ? TEXT("bPlaceholder=true") : TEXT("bPlaceholder=false"));
 
 	int32 FuncParam;
