undo CL: 44096584 due to build breakages

diff --git a/Engine/Plugins/FX/Niagara/Source/Niagara/Internal/DataInterface/NiagaraDataInterfaceDataChannelRead.h b/Engine/Plugins/FX/Niagara/Source/Niagara/Internal/DataInterface/NiagaraDataInterfaceDataChannelRead.h
--- a/Engine/Plugins/FX/Niagara/Source/Niagara/Internal/DataInterface/NiagaraDataInterfaceDataChannelRead.h
+++ b/Engine/Plugins/FX/Niagara/Source/Niagara/Internal/DataInterface/NiagaraDataInterfaceDataChannelRead.h
@@ -88,15 +88,7 @@ public:
 	/** The data channel to access and read from. */
 	UPROPERTY(EditAnywhere, Category="Data Channel")
 	TObjectPtr<UNiagaraDataChannelAsset> Channel;
-
-	/** Access Context used when accessing NDC data. Controls which of the NDCs internal data it will return for this Interface to read. Can be empty for Systems using bAutoLinkToSpawningNDC. */
-	UPROPERTY(EditAnywhere, Category = "Parameters", meta = (ShowOnlyInnerProperties))
-	FNDCAccessContextInst AccessContext;
-
-	/** If true this DI will attempt to automatically link itself to the NDC it's being spawned by. */
-	UPROPERTY(EditAnywhere, Category = "Data Channel", AdvancedDisplay)
-	bool bAutoLinkToSpawningNDC = true;
-
+	
 	/** 
 	* If this interface should read the current frame's data from the Data Channel. If false, the interface will read from the previous frame if it's available.
 	* 
@@ -147,10 +139,6 @@ public:
 	NIAGARA_API virtual void PostInitProperties() override;
 	NIAGARA_API virtual void BeginDestroy() override;
 	NIAGARA_API virtual void Serialize(FArchive& Ar) override;
-	NIAGARA_API virtual void PostLoad() override;
-	#if WITH_EDITOR
-	virtual void PostEditChangeProperty(struct FPropertyChangedEvent& PropertyChangedEvent)override;
-	#endif
 	//UObject Interface End
 
 	//UNiagaraDataInterface Interface
@@ -219,8 +207,6 @@ public:
 
 	FNDIDataChannelCompiledData& GetCompiledData() { return CompiledData; }
 
-	NIAGARA_API virtual void InitAccessContext();
-
 protected:
 #if WITH_EDITORONLY_DATA
 	NIAGARA_API virtual void GetFunctionsInternal(TArray<FNiagaraFunctionSignature>& OutFunctions) const override;
diff --git a/Engine/Plugins/FX/Niagara/Source/Niagara/Internal/DataInterface/NiagaraDataInterfaceDataChannelWrite.h b/Engine/Plugins/FX/Niagara/Source/Niagara/Internal/DataInterface/NiagaraDataInterfaceDataChannelWrite.h
--- a/Engine/Plugins/FX/Niagara/Source/Niagara/Internal/DataInterface/NiagaraDataInterfaceDataChannelWrite.h
+++ b/Engine/Plugins/FX/Niagara/Source/Niagara/Internal/DataInterface/NiagaraDataInterfaceDataChannelWrite.h
@@ -160,17 +160,8 @@ public:
 	UPROPERTY(EditAnywhere, Category = "Data Channel")
 	TObjectPtr<UNiagaraDataChannelAsset> Channel;
 
-	/** Access Context used when accessing NDC data. Controls which of the NDCs internal data it will return for this Interface to write into. Can be empty for Systems using bAutoLinkToSpawningNDC. */
-	UPROPERTY(EditAnywhere, Category = "Data Channel")
-	FNDCAccessContextInst AccessContext;
-
 	//UObject Interface
 	NIAGARA_API virtual void PostInitProperties() override;
-	NIAGARA_API virtual void PostLoad() override;
-
-#if WITH_EDITOR
-	NIAGARA_API virtual void PostEditChangeProperty(struct FPropertyChangedEvent& PropertyChangedEvent) override;
-#endif
 	//UObject Interface End
 
 	//UNiagaraDataInterface Interface
diff --git a/Engine/Plugins/FX/Niagara/Source/Niagara/Niagara.Build.cs b/Engine/Plugins/FX/Niagara/Source/Niagara/Niagara.Build.cs
--- a/Engine/Plugins/FX/Niagara/Source/Niagara/Niagara.Build.cs
+++ b/Engine/Plugins/FX/Niagara/Source/Niagara/Niagara.Build.cs
@@ -41,9 +41,8 @@ public class Niagara : ModuleRules
                 "PhysicsCore",
                 "RenderCore",
                 "RHI",
-		"SlateCore",
+				"SlateCore",
                 "VectorVM",
-		"GameplayTags"
             }
 		);
 
diff --git a/Engine/Plugins/FX/Niagara/Source/Niagara/Private/DataInterface/NiagaraDataInterfaceDataChannelRead.cpp b/Engine/Plugins/FX/Niagara/Source/Niagara/Private/DataInterface/NiagaraDataInterfaceDataChannelRead.cpp
--- a/Engine/Plugins/FX/Niagara/Source/Niagara/Private/DataInterface/NiagaraDataInterfaceDataChannelRead.cpp
+++ b/Engine/Plugins/FX/Niagara/Source/Niagara/Private/DataInterface/NiagaraDataInterfaceDataChannelRead.cpp
@@ -494,26 +494,9 @@ bool FNDIDataChannelReadInstanceData::Tick(UNiagaraDataInterfaceDataChannelRead*
 			bool bNDCDataIsValid = DataChannelData && DataChannelData->IsLayoutValid(DataChannelPtr);
 			if(!bIsInit && (bNDCDataIsValid == false || Interface->bUpdateSourceDataEveryTick))
 			{
-				FNDCAccessContextInst& AccessContext = DataChannelPtr->GetTransientAccessContext();
-				if(AccessContext.IsValid() && AccessContext.GetScriptStruct() != FNDCAccessContextLegacy::StaticStruct())
-				{
-					if(!Interface->bAutoLinkToSpawningNDC)
-					{
-						//TODO: Copy over the interface custom AccessContext. If this proves slow we can possibly cache in per instance data.
-						AccessContext = Interface->AccessContext;
-					}
-					
-					FNDCAccessContextBase& ContextTyped = AccessContext.GetChecked<FNDCAccessContextBase>();
-					ContextTyped.OwningComponent = Instance->GetAttachComponent();
-					ContextTyped.bIsAutoLinkingSystem = Interface->bAutoLinkToSpawningNDC;
-					SetDataChannelData(DataChannelPtr->FindData(AccessContext, ENiagaraResourceAccess::ReadOnly), Interface);//TODO: Maybe should have two paths, one for system instances and another for SceneComponents...
-				}
-				else
-				{
-					//Legacy NDC so just use old path.
-					FNiagaraDataChannelSearchParameters SearchParams(Instance->GetAttachComponent());
-					SetDataChannelData(DataChannelPtr->FindData(SearchParams, ENiagaraResourceAccess::ReadOnly), Interface);
-				}
+				//TODO: Automatically modify tick group if we have DIs that require current frame info?
+				FNiagaraDataChannelSearchParameters SearchParams(Instance->GetAttachComponent());
+				SetDataChannelData(DataChannelPtr->FindData(SearchParams, ENiagaraResourceAccess::ReadOnly), Interface);//TODO: Maybe should have two paths, one for system instances and another for SceneComponents...
 			}	
 
 			if(const UNiagaraDataChannel* ChannelPtr = DataChannelPtr->GetDataChannel())
@@ -648,53 +631,6 @@ void UNiagaraDataInterfaceDataChannelRead::Serialize(FArchive& Ar)
 	Super::Serialize(Ar);
 }
 
-void UNiagaraDataInterfaceDataChannelRead::PostLoad()
-{
-	Super::PostLoad();
-	
-	if(Channel)
-	{
-		Channel->ConditionalPostLoad();
-	}
-
-	InitAccessContext();
-}
-
-#if WITH_EDITOR
-void UNiagaraDataInterfaceDataChannelRead::PostEditChangeProperty(struct FPropertyChangedEvent& PropertyChangedEvent)
-{
-	Super::PostEditChangeProperty(PropertyChangedEvent);
-
-	if (PropertyChangedEvent.GetPropertyName() == GET_MEMBER_NAME_CHECKED(UNiagaraDataInterfaceDataChannelRead, Channel) ||
-		PropertyChangedEvent.GetPropertyName() == GET_MEMBER_NAME_CHECKED(UNiagaraDataInterfaceDataChannelRead, bAutoLinkToSpawningNDC))
-	{
-		InitAccessContext();
-	}
-}
-#endif
-
-void UNiagaraDataInterfaceDataChannelRead::InitAccessContext()
-{
-	UNiagaraDataChannel* NDC = Channel ? Channel->Get() : nullptr;
-	if (NDC && !bAutoLinkToSpawningNDC)
-	{
-		const UScriptStruct* CtxStruct = NDC->GetAccessContextType().Get();
-		//Special case legacy type. We avoid creating an access context for this type here and just use the old search params path.
-		if (CtxStruct == FNDCAccessContextLegacy::StaticStruct())
-		{
-			AccessContext.Reset();
-		}
-		else if (AccessContext.IsValid() == false || CtxStruct != AccessContext.GetScriptStruct())
-		{
-			AccessContext.Init(NDC->GetAccessContextType());
-		}
-	}
-	else
-	{
-		AccessContext.Reset();
-	}
-}
-
 bool UNiagaraDataInterfaceDataChannelRead::InitPerInstanceData(void* PerInstanceData, FNiagaraSystemInstance* SystemInstance)
 {
 	FNDIDataChannelReadInstanceData* InstanceData = new (PerInstanceData) FNDIDataChannelReadInstanceData;
@@ -1223,8 +1159,6 @@ bool UNiagaraDataInterfaceDataChannelRead::Equals(const UNiagaraDataInterface* O
 			//Scope == OtherTyped->Scope &&
 			//Source == OtherTyped->Source &&
 			Channel == OtherTyped->Channel &&
-			AccessContext == OtherTyped->AccessContext &&
-			bAutoLinkToSpawningNDC == OtherTyped->bAutoLinkToSpawningNDC &&
 			bReadCurrentFrame == OtherTyped->bReadCurrentFrame &&
 			bUpdateSourceDataEveryTick == OtherTyped->bUpdateSourceDataEveryTick &&
 			bOverrideSpawnGroupToDataChannelIndex == OtherTyped->bOverrideSpawnGroupToDataChannelIndex &&
@@ -1251,8 +1185,6 @@ bool UNiagaraDataInterfaceDataChannelRead::CopyToInternal(UNiagaraDataInterface*
 		//DestTyped->Source = Source;
 		DestTyped->Channel = Channel;
 		DestTyped->CompiledData = CompiledData;
-		DestTyped->AccessContext = AccessContext;
-		DestTyped->bAutoLinkToSpawningNDC = bAutoLinkToSpawningNDC;
 		DestTyped->bReadCurrentFrame = bReadCurrentFrame;
 		DestTyped->bUpdateSourceDataEveryTick = bUpdateSourceDataEveryTick;
 		DestTyped->bOverrideSpawnGroupToDataChannelIndex = bOverrideSpawnGroupToDataChannelIndex;
diff --git a/Engine/Plugins/FX/Niagara/Source/Niagara/Private/DataInterface/NiagaraDataInterfaceDataChannelWrite.cpp b/Engine/Plugins/FX/Niagara/Source/Niagara/Private/DataInterface/NiagaraDataInterfaceDataChannelWrite.cpp
--- a/Engine/Plugins/FX/Niagara/Source/Niagara/Private/DataInterface/NiagaraDataInterfaceDataChannelWrite.cpp
+++ b/Engine/Plugins/FX/Niagara/Source/Niagara/Private/DataInterface/NiagaraDataInterfaceDataChannelWrite.cpp
@@ -264,9 +264,6 @@ struct FNDIDataChannelWriteInstanceData
 {
 	/** Pointer to the world DataChannel Channel we'll push our DataChannel into. Can be null if DI is not set to publish it's DataChannel. */
 	TWeakObjectPtr<UNiagaraDataChannelHandler> DataChannel;
-	
-	/** Access context object we'll use for finding the correct source data from the Data Channel. */	
-	FNDCAccessContextInst AccessContext;
 
 	//Shared pointer to the actual data we'll be pushing into for this data channel.
 	FNiagaraDataChannelDataPtr DataChannelData;
@@ -310,14 +307,6 @@ struct FNDIDataChannelWriteInstanceData
 			FunctionToDatSetBindingInfo.Add(FNDIDataChannelLayoutManager::Get().GetLayoutInfo(FuncInfo, Interface->GetCompiledData().DataLayout, MissingParams));			
 		}
 
-		//Init our access context for finding the correct source data from the NDC.
-		if (Interface->AccessContext.IsValid())
-		{
-			AccessContext = Interface->AccessContext;
-			FNDCAccessContextBase& BaseContext = AccessContext.GetChecked<FNDCAccessContextBase>();
-			BaseContext.OwningComponent = Instance->GetAttachComponent();
-		}
-
 		return true;
 	}
 
@@ -382,17 +371,9 @@ struct FNDIDataChannelWriteInstanceData
 
 				bool bNDCDataIsValid = DataChannelData && DataChannelData->IsLayoutValid(DataChannelPtr);
 				if (bNDCDataIsValid == false || Interface->bUpdateDestinationDataEveryTick)
-				{	
-					if (AccessContext.IsValid() && AccessContext.GetScriptStruct() != FNDCAccessContextLegacy::StaticStruct())
-					{
-						DataChannelData = DataChannelPtr->FindData(AccessContext, ENiagaraResourceAccess::WriteOnly);
-					}
-					else
-					{
-						//If the access context is not valid then this must be a legacy NDC so just use legacy access code.
-						FNiagaraDataChannelSearchParameters SearchParams(Instance->GetAttachComponent());
-						DataChannelData = DataChannelPtr->FindData(SearchParams, ENiagaraResourceAccess::WriteOnly);
-					}
+				{
+					FNiagaraDataChannelSearchParameters SearchParams(Instance->GetAttachComponent());
+					DataChannelData = DataChannelPtr->FindData(SearchParams, ENiagaraResourceAccess::WriteOnly);
 				}
 
 				const FNDIDataChannelCompiledData& CompiledData = Interface->GetCompiledData();
@@ -555,66 +536,6 @@ void UNiagaraDataInterfaceDataChannelWrite::PostInitProperties()
 	}
 }
 
-void UNiagaraDataInterfaceDataChannelWrite::PostLoad()
-{
-	Super::PostLoad();
-
-	if(Channel)
-	{
-		Channel->ConditionalPostLoad();
-		UNiagaraDataChannel* NDC = Channel->Get();
-		if (NDC)
-		{
-			const UScriptStruct* CtxStruct = NDC->GetAccessContextType().Get();
-			//Special case legacy type. We avoid creating an access context for this type here and just use the old search params path.
-			if (CtxStruct == FNDCAccessContextLegacy::StaticStruct())
-			{
-				AccessContext.Reset();
-			}
-			else if (AccessContext.GetScriptStruct() != CtxStruct)
-			{
-				AccessContext.Init(NDC->GetAccessContextType());
-			}
-		}
-		else
-		{
-			AccessContext.Reset();
-		}
-	}
-	else
-	{
-		AccessContext.Reset();
-	}
-}
-
-#if WITH_EDITOR
-void UNiagaraDataInterfaceDataChannelWrite::PostEditChangeProperty(struct FPropertyChangedEvent& PropertyChangedEvent)
-{
-	Super::PostEditChangeProperty(PropertyChangedEvent);
-
-	if (PropertyChangedEvent.GetPropertyName() == GET_MEMBER_NAME_CHECKED(UNiagaraDataInterfaceDataChannelWrite, Channel))
-	{
-		UNiagaraDataChannel* NDC = Channel ? Channel->Get() : nullptr;
-		if (NDC)
-		{
-			const UScriptStruct* CtxStruct = NDC->GetAccessContextType().Get();
-			//Special case legacy type. We avoid creating an access context for this type here and just use the old search params path.
-			if (CtxStruct == FNDCAccessContextLegacy::StaticStruct())
-			{
-				AccessContext.Reset();
-			}
-			else if (AccessContext.GetScriptStruct() != CtxStruct)
-			{
-				AccessContext.Init(NDC->GetAccessContextType());
-			}
-		}
-		else
-		{
-			AccessContext.Reset();
-		}
-	}
-}
-#endif
 
 int32 UNiagaraDataInterfaceDataChannelWrite::PerInstanceDataSize() const
 {
@@ -845,7 +766,6 @@ bool UNiagaraDataInterfaceDataChannelWrite::Equals(const UNiagaraDataInterface*
 			bPublishToCPU == OtherTyped->bPublishToCPU &&
 			bPublishToGPU == OtherTyped->bPublishToGPU &&
 			Channel == OtherTyped->Channel &&
-			AccessContext == OtherTyped->AccessContext &&
 			bOnlyWriteOnceOnSubticks == OtherTyped->bOnlyWriteOnceOnSubticks &&
 			bUpdateDestinationDataEveryTick == OtherTyped->bUpdateDestinationDataEveryTick)
 		{
@@ -1019,26 +939,14 @@ bool UNiagaraDataInterfaceDataChannelWrite::SimCacheReadFrame(const FNiagaraSimC
 	}
 
 	FNiagaraSystemInstance* SystemInstance = ReadContext.GetSystemInstance();
-	UNiagaraDataChannel* NDC = Channel->Get();
 
 	FNiagaraDataChannelDataPtr DataChannelData;
 	if (FNiagaraWorldManager* WorldMan = FNiagaraWorldManager::Get(SystemInstance->GetWorld()))
 	{
-		if (UNiagaraDataChannelHandler* Handler = WorldMan->GetDataChannelManager().FindDataChannelHandler(NDC))
+		if (UNiagaraDataChannelHandler* Handler = WorldMan->GetDataChannelManager().FindDataChannelHandler(Channel->Get()))
 		{
-			if(AccessContext.IsValid() && AccessContext.GetScriptStruct() != FNDCAccessContextLegacy::StaticStruct())
-			{
-				FNDCAccessContextInst& TransientAccessContext = NDC->GetTransientAccessContext();
-				TransientAccessContext = AccessContext;
-				FNDCAccessContextBase& TypedContext = TransientAccessContext.GetChecked<FNDCAccessContextBase>();
-				TypedContext.OwningComponent = SystemInstance->GetAttachComponent();
-				DataChannelData = Handler->FindData(TransientAccessContext, ENiagaraResourceAccess::WriteOnly);
-			}
-			else
-			{
-				FNiagaraDataChannelSearchParameters SearchParams(SystemInstance->GetAttachComponent());
-				DataChannelData = Handler->FindData(SearchParams, ENiagaraResourceAccess::WriteOnly);
-			}
+			FNiagaraDataChannelSearchParameters SearchParams(SystemInstance->GetAttachComponent());
+			DataChannelData = Handler->FindData(SearchParams, ENiagaraResourceAccess::WriteOnly);
 		}
 	}
 	if (!DataChannelData.IsValid())
@@ -1060,7 +968,7 @@ bool UNiagaraDataInterfaceDataChannelWrite::SimCacheReadFrame(const FNiagaraSimC
 			PublishRequest.DebugSource = FString::Format(TEXT("{0} (Sim cache {1})"), {SystemInstance->GetSystem()->GetName(), GetPathNameSafe(Storage->GetOuter())});
 #endif
 
-			PublishRequest.GameData = MakeShared<FNiagaraDataChannelGameData>(NDC->GetLayoutInfo());
+			PublishRequest.GameData = MakeShared<FNiagaraDataChannelGameData>(Channel->Get()->GetLayoutInfo());
 			PublishRequest.GameData->SetNum(Frame.NumElements);
 			PublishRequest.LwcTile = ReadContext.GetLWCTileA();
 
@@ -1185,7 +1093,6 @@ bool UNiagaraDataInterfaceDataChannelWrite::CopyToInternal(UNiagaraDataInterface
 		DestTyped->bPublishToCPU = bPublishToCPU;
 		DestTyped->bPublishToGPU = bPublishToGPU;
 		DestTyped->Channel = Channel;
-		DestTyped->AccessContext = AccessContext;
 		DestTyped->CompiledData = CompiledData;
 		DestTyped->bUpdateDestinationDataEveryTick = bUpdateDestinationDataEveryTick;
 		DestTyped->bOnlyWriteOnceOnSubticks = bOnlyWriteOnceOnSubticks;
diff --git a/Engine/Plugins/FX/Niagara/Source/Niagara/Private/NiagaraDataChannel.cpp b/Engine/Plugins/FX/Niagara/Source/Niagara/Private/NiagaraDataChannel.cpp
--- a/Engine/Plugins/FX/Niagara/Source/Niagara/Private/NiagaraDataChannel.cpp
+++ b/Engine/Plugins/FX/Niagara/Source/Niagara/Private/NiagaraDataChannel.cpp
@@ -22,8 +22,6 @@
 
 #include "RenderGraphUtils.h"
 
-#include "Blueprint/BlueprintExceptionInfo.h"
-
 #include UE_INLINE_GENERATED_CPP_BY_NAME(NiagaraDataChannel)
 
 #define LOCTEXT_NAMESPACE "NiagaraDataChannels"
@@ -1410,8 +1408,6 @@ void UNiagaraDataChannel::PostLoad()
 
 	GetLayoutInfo();
 
-	TransientAccessContext.Init(GetAccessContextType());
-
 	INiagaraModule::RequestRefreshDataChannels();
 }
 
@@ -1488,8 +1484,6 @@ void UNiagaraDataChannel::PostEditChangeProperty(FPropertyChangedEvent& Property
 	LayoutInfo = nullptr;
 	GetLayoutInfo();
 
-	TransientAccessContext.Init(GetAccessContextType());
-
 	INiagaraModule::RequestRefreshDataChannels();
 }
 
@@ -1523,16 +1517,28 @@ UNiagaraDataChannelLibrary::UNiagaraDataChannelLibrary(const FObjectInitializer&
 
 UNiagaraDataChannelHandler* UNiagaraDataChannelLibrary::GetNiagaraDataChannel(const UObject* WorldContextObject, const UNiagaraDataChannelAsset* Channel)
 {
+	if (!Channel)
+	{
+		return nullptr;
+	}
 	return FindDataChannelHandler(WorldContextObject, Channel->Get());
 }
 
 UNiagaraDataChannelWriter* UNiagaraDataChannelLibrary::WriteToNiagaraDataChannel(const UObject* WorldContextObject, const UNiagaraDataChannelAsset* Channel, FNiagaraDataChannelSearchParameters SearchParams, int32 Count, bool bVisibleToGame, bool bVisibleToCPU, bool bVisibleToGPU, const FString& DebugSource)
 {
+	if (!Channel)
+	{
+		return nullptr;
+	}
 	return CreateDataChannelWriter(WorldContextObject, Channel->Get(), SearchParams, Count, bVisibleToGame, bVisibleToCPU, bVisibleToGPU, DebugSource);
 }
 
 UNiagaraDataChannelReader* UNiagaraDataChannelLibrary::ReadFromNiagaraDataChannel(const UObject* WorldContextObject, const UNiagaraDataChannelAsset* Channel, FNiagaraDataChannelSearchParameters SearchParams, bool bReadPreviousFrame)
 {
+	if (!Channel)
+	{
+		return nullptr;
+	}
 	return CreateDataChannelReader(WorldContextObject, Channel->Get(), SearchParams, bReadPreviousFrame);
 }
 
@@ -1558,185 +1564,6 @@ void UNiagaraDataChannelLibrary::WriteToNiagaraDataChannelSingle(const UObject*,
 	// this function is just a placeholder and calls into CreateDataChannelWriter and its individual write functions from the BP node
 }
 
-//////////////////////////////////////////////////////////////////////////
-// Access Context based functions
-
-UNiagaraDataChannelWriter* UNiagaraDataChannelLibrary::WriteToNiagaraDataChannel_WithContext(const UObject* WorldContextObject, const UNiagaraDataChannelAsset* Channel, FNDCAccessContextInst& AccessContext, int32 Count, bool bVisibleToBlueprint, bool bVisibleToNiagaraCPU, bool bVisibleToNiagaraGPU, const FString& DebugSource)
-{
-	return CreateDataChannelWriter_WithContext(WorldContextObject, Channel->Get(), AccessContext, Count, bVisibleToBlueprint, bVisibleToNiagaraCPU, bVisibleToNiagaraGPU, DebugSource);
-}
-
-UNiagaraDataChannelReader* UNiagaraDataChannelLibrary::ReadFromNiagaraDataChannel_WithContext(const UObject* WorldContextObject, const UNiagaraDataChannelAsset* Channel, FNDCAccessContextInst& AccessContext, bool bReadPreviousFrame)
-{
-	return CreateDataChannelReader_WithContext(WorldContextObject, Channel->Get(), AccessContext, bReadPreviousFrame);
-}
-
-int32 UNiagaraDataChannelLibrary::GetDataChannelElementCount_WithContext(const UObject* WorldContextObject, const UNiagaraDataChannelAsset* Channel, FNDCAccessContextInst& AccessContext, bool bReadPreviousFrame)
-{
-	if (Channel && Channel->Get())
-	{
-		if (UNiagaraDataChannelReader* Reader = CreateDataChannelReader_WithContext(WorldContextObject, Channel->Get(), AccessContext, bReadPreviousFrame))
-		{
-			return Reader->Num();
-		}
-	}
-	return 0;
-}
-
-void UNiagaraDataChannelLibrary::ReadFromNiagaraDataChannelSingle_WithContext(const UObject*, const UNiagaraDataChannelAsset*, int32, FNDCAccessContextInst&, bool, ENiagartaDataChannelReadResult&)
-{
-	// this function is just a placeholder and calls into CreateDataChannelReader_WithContext and its individual read functions from the BP node
-}
-
-void UNiagaraDataChannelLibrary::WriteToNiagaraDataChannelSingle_WithContext(const UObject*, const UNiagaraDataChannelAsset*, FNDCAccessContextInst& AccessContext, bool, bool, bool)
-{
-	// this function is just a placeholder and calls into CreateDataChannelWriter_WithContext and its individual write functions from the BP node
-}
-
-//////////////////////////////////////////////////////////////////////////
-//New Access context nodes
-
-FNDCAccessContextInst UNiagaraDataChannelLibrary::MakeNDCAccessContextInstance(UScriptStruct* ContextStruct)
-{
-	// We should never hit this! stubs to avoid NoExport on the class.
-	checkNoEntry();
-	return FNDCAccessContextInst();
-}
-
-void UNiagaraDataChannelLibrary::SetMembersInNDCAccessContextInstance(UPARAM(Ref) const FNDCAccessContextInst& AccessContext, UScriptStruct* ContextStruct)
-{
-	// We should never hit this! stubs to avoid NoExport on the class.
-	checkNoEntry();
-}
-
-void UNiagaraDataChannelLibrary::GetMembersInNDCAccessContextInstance(const FNDCAccessContextInst& AccessContext, UScriptStruct* ContextStruct)
-{
-	// We should never hit this! stubs to avoid NoExport on the class.
-	checkNoEntry();
-}
-
-void UNiagaraDataChannelLibrary::GetSinglePropertyInNDCAccessContextInstance(const FNDCAccessContextInst& AccessContext, UScriptStruct* ContextStruct, FName PropertyName, int32& Value)
-{
-	// We should never hit this! stubs to avoid NoExport on the class.
-	checkNoEntry();
-}
-
-void UNiagaraDataChannelLibrary::SetSinglePropertyInNDCAccessContextInstance(UPARAM(Ref) const FNDCAccessContextInst& AccessContext, UScriptStruct* ContextStruct, FName PropertyName, const int32& Value)
-{
-	// We should never hit this! stubs to avoid NoExport on the class.
-	checkNoEntry();
-}
-
-DEFINE_FUNCTION(UNiagaraDataChannelLibrary::execMakeNDCAccessContextInstance)
-{
-	P_GET_OBJECT(UScriptStruct, ContextStruct);
-	P_FINISH
-
-	if(ContextStruct && ContextStruct->IsChildOf(FNDCAccessContextBase::StaticStruct()))
-	{
-		P_NATIVE_BEGIN;
-		(*(FNDCAccessContextInst*)RESULT_PARAM).Init(TNDCAccessContextType(ContextStruct));		
-		P_NATIVE_END;
-	}
-	else
-	{
-		FBlueprintExceptionInfo ExceptionInfo(
-			EBlueprintExceptionType::AbortExecution,
-			LOCTEXT("NDC_MakeAccessContext_InvalidTypeWarning", "Invalid Value type for Make Niagara Data Channel Access Context. You must use a child of FNDCAccessContextBase"));
-		FBlueprintCoreDelegates::ThrowScriptException(P_THIS, Stack, ExceptionInfo);
-	}
-}
-
-DEFINE_FUNCTION(UNiagaraDataChannelLibrary::execGetSinglePropertyInNDCAccessContextInstance)
-{
-	P_GET_STRUCT_REF(FNDCAccessContextInst, AccessContext);
-	P_GET_OBJECT(UScriptStruct, ContextStruct);
-	P_GET_STRUCT(FName, PropertyName);
-
-	// Read wildcard Value input.
-	Stack.MostRecentPropertyAddress = nullptr;
-	Stack.MostRecentPropertyContainer = nullptr;
-	Stack.StepCompiledIn<FProperty>(nullptr);
-
-	P_FINISH
-
-	const FProperty* ValueProp = CastField<FProperty>(Stack.MostRecentProperty);
-	void* ValuePtr = Stack.MostRecentPropertyAddress;
-
-	const UScriptStruct* ActualContextStruct = AccessContext.GetScriptStruct();
-	const uint8* SrcPtr = AccessContext.AccessContext.GetMemory();
-	
-	if (ActualContextStruct && ActualContextStruct->IsChildOf(FNDCAccessContextBase::StaticStruct()) && ActualContextStruct->IsChildOf(ContextStruct))
-	{
-		P_NATIVE_BEGIN
-		if(FProperty* SrcProperty = ActualContextStruct->FindPropertyByName(PropertyName))
-		{
-			SrcProperty->CopyCompleteValue(ValuePtr, SrcProperty->ContainerPtrToValuePtr<uint8>(SrcPtr));
-		}
-		else
-		{
-			FBlueprintExceptionInfo ExceptionInfo(
-				EBlueprintExceptionType::AbortExecution,
-				FText::Format(LOCTEXT("NDC_AccessContext_MissingSourceProperty_Fmt", "Cannot read {0} from {1}. Property does not exist on this struct type."), FText::FromName(PropertyName), FText::FromString(ActualContextStruct->GetName())));
-			FBlueprintCoreDelegates::ThrowScriptException(P_THIS, Stack, ExceptionInfo);
-		}
-		P_NATIVE_END
-	}
-	else
-	{
-		FBlueprintExceptionInfo ExceptionInfo(
-			EBlueprintExceptionType::NonFatalError,
-			LOCTEXT("NDC_GetAccessContext_InvalidTypeWarning", "Invalid Value type for Get Members In Niagara Data Channel Access Context. You must use a child of FNDCAccessContextBase. Possible that calls with Incompatible types have been mixed."));
-		FBlueprintCoreDelegates::ThrowScriptException(P_THIS, Stack, ExceptionInfo);
-	}
-}
-
-DEFINE_FUNCTION(UNiagaraDataChannelLibrary::execSetSinglePropertyInNDCAccessContextInstance)
-{
-	P_GET_STRUCT_REF(FNDCAccessContextInst, AccessContext);
-	P_GET_OBJECT(UScriptStruct, ContextStruct);
-	P_GET_STRUCT(FName, PropertyName);
-
-	// Read wildcard Value input.
-	Stack.MostRecentPropertyAddress = nullptr;
-	Stack.MostRecentPropertyContainer = nullptr;
-	Stack.StepCompiledIn<FProperty>(nullptr);
-
-	P_FINISH
-
-	const FProperty* ValueProp = CastField<FProperty>(Stack.MostRecentProperty);
-	void* ValuePtr = Stack.MostRecentPropertyAddress;
-
-	const UScriptStruct* ActualContextStruct = AccessContext.GetScriptStruct();
-	uint8* DestPtr = AccessContext.AccessContext.GetMutableMemory();
-
-	if (ActualContextStruct && ActualContextStruct->IsChildOf(FNDCAccessContextBase::StaticStruct()) && ActualContextStruct->IsChildOf(ContextStruct))
-	{
-		P_NATIVE_BEGIN
-			if (FProperty* DestProperty = ActualContextStruct->FindPropertyByName(PropertyName))
-			{
-				DestProperty->CopyCompleteValue(DestProperty->ContainerPtrToValuePtr<uint8>(DestPtr), ValuePtr);
-			}
-			else
-			{
-				FBlueprintExceptionInfo ExceptionInfo(
-					EBlueprintExceptionType::AbortExecution,
-					FText::Format(LOCTEXT("NDC_AccessContext_MissingDestProperty_Fmt", "Cannot write {0} to {1}. Property does not exist on this struct type."), FText::FromName(PropertyName), FText::FromString(ActualContextStruct->GetName())));
-				FBlueprintCoreDelegates::ThrowScriptException(P_THIS, Stack, ExceptionInfo);
-			}
-		P_NATIVE_END
-	}
-	else
-	{
-		FBlueprintExceptionInfo ExceptionInfo(
-			EBlueprintExceptionType::AbortExecution,
-			LOCTEXT("NDC_SetAccessContext_InvalidTypeWarning", "Invalid Value type for Set Niagara Data Channel Access Context Members. You must use a child of FNDCAccessContextBase"));
-		FBlueprintCoreDelegates::ThrowScriptException(P_THIS, Stack, ExceptionInfo);
-	}
-}
-
-//////////////////////////////////////////////////////////////////////////
-// 
 void UNiagaraDataChannelLibrary::SubscribeToNiagaraDataChannel(const UObject* WorldContextObject, const UNiagaraDataChannelAsset* Channel, FNiagaraDataChannelSearchParameters SearchParams, const FOnNewNiagaraDataChannelPublish& UpdateDelegate, int32& UnsubscribeToken)
 {
 	if (UNiagaraDataChannelHandler* NiagaraDataChannelHandler = GetNiagaraDataChannel(WorldContextObject, Channel))
@@ -1745,14 +1572,6 @@ void UNiagaraDataChannelLibrary::SubscribeToNiagaraDataChannel(const UObject* Wo
 	}
 }
 
-void UNiagaraDataChannelLibrary::SubscribeToNiagaraDataChannel_WithContext(const UObject* WorldContextObject, const UNiagaraDataChannelAsset* Channel, FNDCAccessContextInst& AccessContext, const FOnNewNiagaraDataChannelPublish& UpdateDelegate, int32& UnsubscribeToken)
-{
-	if (UNiagaraDataChannelHandler* NiagaraDataChannelHandler = GetNiagaraDataChannel(WorldContextObject, Channel))
-	{
-		NiagaraDataChannelHandler->SubscribeToDataChannelUpdates_WithContext(UpdateDelegate, AccessContext, UnsubscribeToken);
-	}
-}
-
 void UNiagaraDataChannelLibrary::UnsubscribeFromNiagaraDataChannel(const UObject* WorldContextObject, const UNiagaraDataChannelAsset* Channel, const int32& UnsubscribeToken)
 {
 	if (UNiagaraDataChannelHandler* NiagaraDataChannelHandler = GetNiagaraDataChannel(WorldContextObject, Channel))
@@ -1764,7 +1583,7 @@ void UNiagaraDataChannelLibrary::UnsubscribeFromNiagaraDataChannel(const UObject
 UNiagaraDataChannelHandler* UNiagaraDataChannelLibrary::FindDataChannelHandler(const UObject* WorldContextObject, const UNiagaraDataChannel* Channel)
 {
 	UWorld* World = (WorldContextObject != nullptr) ? GEngine->GetWorldFromContextObject(WorldContextObject, EGetWorldErrorMode::LogAndReturnNull) : nullptr;
-	if (Channel && World)
+	if (World)
 	{
 		if (FNiagaraWorldManager* WorldMan = FNiagaraWorldManager::Get(World))
 		{
@@ -1778,7 +1597,7 @@ UNiagaraDataChannelWriter* UNiagaraDataChannelLibrary::CreateDataChannelWriter(c
 {
 	check(IsInGameThread());
 	UWorld* World = (WorldContextObject != nullptr) ? GEngine->GetWorldFromContextObject(WorldContextObject, EGetWorldErrorMode::LogAndReturnNull) : nullptr;
-	if (Channel && World && Count > 0)
+	if (World && Count > 0)
 	{
 		if (FNiagaraWorldManager* WorldMan = FNiagaraWorldManager::Get(World))
 		{
@@ -1799,9 +1618,8 @@ UNiagaraDataChannelWriter* UNiagaraDataChannelLibrary::CreateDataChannelWriter(c
 
 UNiagaraDataChannelReader* UNiagaraDataChannelLibrary::CreateDataChannelReader(const UObject* WorldContextObject, const UNiagaraDataChannel* Channel, FNiagaraDataChannelSearchParameters SearchParams, bool bReadPreviousFrame)
 {
-	check(IsInGameThread());
 	UWorld* World = (WorldContextObject != nullptr) ? GEngine->GetWorldFromContextObject(WorldContextObject, EGetWorldErrorMode::LogAndReturnNull) : nullptr;
-	if (Channel && World)
+	if (World)
 	{
 		if (FNiagaraWorldManager* WorldMan = FNiagaraWorldManager::Get(World))
 		{
@@ -1820,74 +1638,6 @@ UNiagaraDataChannelReader* UNiagaraDataChannelLibrary::CreateDataChannelReader(c
 	return nullptr;
 }
 
-UNiagaraDataChannelWriter* UNiagaraDataChannelLibrary::CreateDataChannelWriter_WithContext(const UObject* WorldContextObject, const UNiagaraDataChannel* Channel, FNDCAccessContextInst& AccessContext, int32 Count, bool bVisibleToGame, bool bVisibleToCPU, bool bVisibleToGPU, const FString& DebugSource)
-{
-	check(IsInGameThread());
-	UWorld* World = (WorldContextObject != nullptr) ? GEngine->GetWorldFromContextObject(WorldContextObject, EGetWorldErrorMode::LogAndReturnNull) : nullptr;
-	if (Channel && World && Count > 0)
-	{
-		if (FNiagaraWorldManager* WorldMan = FNiagaraWorldManager::Get(World))
-		{
-			if (UNiagaraDataChannelHandler* Handler = WorldMan->GetDataChannelManager().FindDataChannelHandler(Channel))
-			{
-				if (UNiagaraDataChannelWriter* Writer = Handler->GetDataChannelWriter())
-				{					
-					if (Writer->BeginWrite(AccessContext, Count, bVisibleToGame, bVisibleToCPU, bVisibleToGPU, DebugSource))
-					{
-						return Writer;
-					}
-				}
-			}
-		}
-	}
-	return nullptr;	
-}
-
-UNiagaraDataChannelReader* UNiagaraDataChannelLibrary::CreateDataChannelReader_WithContext(const UObject* WorldContextObject, const UNiagaraDataChannel* Channel, FNDCAccessContextInst& AccessContext, bool bReadPreviousFrame)
-{
-	check(IsInGameThread());
-	UWorld* World = (WorldContextObject != nullptr) ? GEngine->GetWorldFromContextObject(WorldContextObject, EGetWorldErrorMode::LogAndReturnNull) : nullptr;
-	if (Channel && World)
-	{
-		if (FNiagaraWorldManager* WorldMan = FNiagaraWorldManager::Get(World))
-		{
-			if (UNiagaraDataChannelHandler* Handler = WorldMan->GetDataChannelManager().FindDataChannelHandler(Channel))
-			{
-				if (UNiagaraDataChannelReader* Reader = Handler->GetDataChannelReader())
-				{
-					if (Reader->BeginRead(AccessContext, bReadPreviousFrame))
-					{
-						return Reader;
-					}
-				}
-			}
-		}
-	}
-	return nullptr;
-}
-
-
-//////////////////////////////////////////////////////////////////////////
-// Functions for modifying dynamic objects.
-// 
-
-//Helper to reduce boiler plate of finding an NDC and calling a function on it from BP callable UFunctions.
-template<typename TAction>
-void ExecNDCHandlerAction(const UObject* WorldContextObject, const UNiagaraDataChannelAsset* Channel, TAction&& Func)
-{
-	UWorld* World = (WorldContextObject != nullptr) ? GEngine->GetWorldFromContextObject(WorldContextObject, EGetWorldErrorMode::LogAndReturnNull) : nullptr;
-	if (Channel && World)
-	{
-		if (FNiagaraWorldManager* WorldMan = FNiagaraWorldManager::Get(World))
-		{
-			if (UNiagaraDataChannelHandler* Handler = WorldMan->GetDataChannelManager().FindDataChannelHandler(Channel))
-			{
-				Func(Handler);
-			}
-		}
-	}
-}
-
 //////////////////////////////////////////////////////////////////////////
 
 #if WITH_EDITOR
@@ -1916,86 +1666,6 @@ void UNiagaraDataChannelAsset::PostEditChangeProperty(FPropertyChangedEvent& Pro
 
 #endif
 
-//////////////////////////////////////////////////////////////////////////
-
-FNDCAccessContextInst& FNiagaraDataChannelReference::GetAccessContext()const
-{
-	InitAccessContext();
-	return AccessContext;
-}
-
-void FNiagaraDataChannelReference::InitAccessContext(bool bTransferExisting/* =true */)const
-{
-	//We need to verify that the access context has the correct type.
-	//We can get bad types if someone changes the NDC type for example.
-	if(DataChannel)
-	{
-		if (UNiagaraDataChannel* NDC = DataChannel->Get())
-		{
-			const UScriptStruct* ExpectedType = NDC->GetAccessContextType().Get();
-			if(ExpectedType != FNDCAccessContextLegacy::StaticStruct())//Special case legacy type indicating we should not create a context in most places.
-			{
-				if (AccessContext.IsValid() && bTransferExisting)
-				{
-					const UScriptStruct* ActualType = AccessContext.GetScriptStruct();
-					if (ActualType != ExpectedType)
-					{
-						//If we do have the wrong type. 
-						//Initialize to the correct type but try to pull any matching properties from the old type.
-						FInstancedStruct OldContext = AccessContext.AccessContext;
-						AccessContext.Init(NDC->GetAccessContextType());
-
-						const uint8* Src = OldContext.GetMemory();
-						uint8* Dst = AccessContext.AccessContext.GetMutableMemory();
-						for (TFieldIterator<FProperty> It(ExpectedType); It; ++It)
-						{
-							FProperty* DestProp = *It;
-							FProperty* SrcProp = ActualType->FindPropertyByName(DestProp->GetFName());
-							if (SrcProp && SrcProp->SameType(DestProp))
-							{
-								const uint8* PropSrcPtr = SrcProp->ContainerPtrToValuePtr<const uint8>(Src);
-								uint8* PropDstPtr = DestProp->ContainerPtrToValuePtr<uint8>(Dst);
-								DestProp->CopyCompleteValue(PropDstPtr, PropSrcPtr);
-							}
-						}
-					}
-				}
-				else
-				{
-					AccessContext.Init(NDC->GetAccessContextType());
-				}
-				return;
-			}
-		}
-	}
-
-	//If we fail to init properly then clear the access context to avoid memory bloat.
-	AccessContext.Reset();
-}
-
-void FNiagaraDataChannelReference::SetDataChannel(UNiagaraDataChannelAsset* InDataChannel)
-{
-	if(DataChannel != InDataChannel)
-	{
-		DataChannel = InDataChannel;
-		OnValueChanged();
-	}
-}
-
-void FNiagaraDataChannelReference::OnValueChanged()
-{
-	if(DataChannel)
-	{
-		InitAccessContext();
-	}
-	else
-	{
-		AccessContext.Reset();
-	}
-}
-
-//////////////////////////////////////////////////////////////////////////
-
 bool FNiagaraDataChannelVariable::Serialize(FArchive& Ar)
 {
 	if (FNiagaraVariableBase::Serialize(Ar))
diff --git a/Engine/Plugins/FX/Niagara/Source/Niagara/Private/NiagaraDataChannelAccessor.cpp b/Engine/Plugins/FX/Niagara/Source/Niagara/Private/NiagaraDataChannelAccessor.cpp
--- a/Engine/Plugins/FX/Niagara/Source/Niagara/Private/NiagaraDataChannelAccessor.cpp
+++ b/Engine/Plugins/FX/Niagara/Source/Niagara/Private/NiagaraDataChannelAccessor.cpp
@@ -13,7 +13,6 @@
 #include "NiagaraDataChannelManager.h"
 
 #include "NiagaraWorldManager.h"
-#include "NiagaraDataChannelAccessorImpl.h"
 
 #include UE_INLINE_GENERATED_CPP_BY_NAME(NiagaraDataChannelAccessor)
 
@@ -49,16 +48,6 @@ bool UNiagaraDataChannelReader::InitAccess(FNiagaraDataChannelSearchParameters S
 	return Data.IsValid();
 }
 
-bool UNiagaraDataChannelReader::BeginRead(FNDCAccessContextInst& AccessContext, bool bReadPreviousFrameData)
-{
-	Data = nullptr;
-	bReadingPreviousFrame = bReadPreviousFrameData;
-	check(Owner);
-
-	Data = Owner->FindData(AccessContext, ENiagaraResourceAccess::ReadOnly);
-	return Data.IsValid();
-}
-
 int32 UNiagaraDataChannelReader::Num()const
 {
 	if (Data.IsValid())
@@ -199,36 +188,6 @@ bool UNiagaraDataChannelWriter::InitWrite(FNiagaraDataChannelSearchParameters Se
 	return false;
 }
 
-bool UNiagaraDataChannelWriter::BeginWrite(FNDCAccessContextInst& AccessContext, int32 Count, bool bVisibleToGame, bool bVisibleToCPU, bool bVisibleToGPU, const FString& DebugSource)
-{
-	check(IsInGameThread());
-
-	if (Count == 0)
-	{
-		UE_LOG(LogNiagara, Warning, TEXT("Call to UNiagaraDataChannelWriter::InitWrite with Count == 0. Ignored."));
-		return false;
-	}
-
-	check(Owner);
-
-	if (FNiagaraDataChannelDataPtr DestData = Owner->FindData(AccessContext, ENiagaraResourceAccess::WriteOnly))
-	{
-		FString SourceString;
-#if !UE_BUILD_SHIPPING
-		SourceString = DebugSource.IsEmpty() ? GetPathName() : DebugSource;
-#endif
-
-		//Attempt to use an existing cached dest data.
-		Data = DestData->GetGameDataForWriteGT(Count, bVisibleToGame, bVisibleToCPU, bVisibleToGPU, SourceString);
-		StartIndex = Data->Num();
-		Data->SetNum(StartIndex + Count);
-
-		return true;
-	}
-
-	return false;
-}
-
 int32 UNiagaraDataChannelWriter::Num()const
 {
 	if (Data.IsValid())
@@ -392,34 +351,98 @@ void FNDCAccessorBase::Init(const UNiagaraDataChannel* DataChannel)
 //////////////////////////////////////////////////////////////////////////
 //FNDCWriterBase - Base class for user code NDCWriters. See full description in header.
 
-void FNDCWriterBase::EndWrite()
+bool FNDCWriterBase::BeginWrite(const UObject* WorldContextObject, const UNiagaraDataChannel* DataChannel, const FNiagaraDataChannelSearchParameters& SearchParams, int32 InCount, bool bVisibleToGame, bool bVisibleToCPU, bool bVisibleToGPU)
 {
-	Data = nullptr;
-	StartIndex = INDEX_NONE;
-	Count = INDEX_NONE;
-}
+	ensureMsgf(Data.IsValid() == false, TEXT("NDC Writer with non-null data on BeginWrite. Possible EndWrite was not called after previous write."));
 
-bool FNDCWriterBase::BeginWrite(const UObject* WorldContextObject, const UNiagaraDataChannel* DataChannel, FNDCAccessContextInst& AccessContext, int32 InCount, bool bVisibleToGame, bool bVisibleToCPU, bool bVisibleToGPU)
-{
-	return BeginWrite_Internal(WorldContextObject, DataChannel, AccessContext, InCount, bVisibleToGame, bVisibleToCPU, bVisibleToGPU);
+	const FNiagaraDataChannelLayoutInfoPtr NDCLayout = DataChannel->GetLayoutInfo();
+
+	if (CachedLayout != NDCLayout)
+	{
+		Init(DataChannel);
+		CachedLayout = NDCLayout;
+	}
+
+#if DEBUG_NDC_ACCESS
+	for (FNDCVarAccessorBase* Var : VariableAccessors)
+	{
+		Var->ValidateLayout();
+	}
+#endif
+
+	check(IsInGameThread());
+	UWorld* World = (WorldContextObject != nullptr) ? GEngine->GetWorldFromContextObject(WorldContextObject, EGetWorldErrorMode::LogAndReturnNull) : nullptr;
+	if (World && InCount > 0 && CachedLayout.IsValid())
+	{
+		if (FNiagaraWorldManager* WorldMan = FNiagaraWorldManager::Get(World))
+		{
+			if (UNiagaraDataChannelHandler* Handler = WorldMan->GetDataChannelManager().FindDataChannelHandler(DataChannel))
+			{
+				if (FNiagaraDataChannelDataPtr DestData = Handler->FindData(SearchParams, ENiagaraResourceAccess::WriteOnly))
+				{
+					//Attempt to use an existing cached dest data.
+					Data = DestData->GetGameDataForWriteGT(InCount, bVisibleToGame, bVisibleToCPU, bVisibleToGPU, DebugSource);
+
+					//We are potentially writing into the middle of an existing buffer here so track base index and count.
+					StartIndex = Data->Num();
+					Count = InCount;
+					Data->SetNum(StartIndex + InCount);
+
+					return true;
+				}
+			}
+		}
+	}
+	return false;
 }
 
-bool FNDCWriterBase::BeginWrite(const UObject* WorldContextObject, const UNiagaraDataChannel* DataChannel, FNiagaraDataChannelSearchParameters SearchParams, int32 InCount, bool bVisibleToGame, bool bVisibleToCPU, bool bVisibleToGPU)
+void FNDCWriterBase::EndWrite()
 {
-	return BeginWrite_Internal(WorldContextObject, DataChannel, SearchParams, InCount, bVisibleToGame, bVisibleToCPU, bVisibleToGPU);
+	Data = nullptr;
+	StartIndex = INDEX_NONE;
+	Count = INDEX_NONE;
 }
 
 //////////////////////////////////////////////////////////////////////////
 //NDCReaderBase - Base class for user code NDCReaders. See full description in header.
 
-bool FNDCReaderBase::BeginRead(const UObject* WorldContextObject, const UNiagaraDataChannel* DataChannel, FNDCAccessContextInst& AccessContext, bool bInPreviousFrame)
+bool FNDCReaderBase::BeginRead(const UObject* WorldContextObject, const UNiagaraDataChannel* DataChannel, const FNiagaraDataChannelSearchParameters& SearchParams, bool bInPreviousFrame)
 {
-	return BeginRead_Internal(WorldContextObject, DataChannel, AccessContext, bInPreviousFrame);
-}
+	ensureMsgf(Data.IsValid() == false, TEXT("NDC Reader with non-null data on BeginRead. Possible EndRead was not called after previous read."));
 
-bool FNDCReaderBase::BeginRead(const UObject* WorldContextObject, const UNiagaraDataChannel* DataChannel, FNiagaraDataChannelSearchParameters SearchParams, bool bInPreviousFrame)
-{
-	return BeginRead_Internal(WorldContextObject, DataChannel, SearchParams, bInPreviousFrame);
+	const FNiagaraDataChannelLayoutInfoPtr NDCLayout = DataChannel->GetLayoutInfo();
+
+	if (CachedLayout != NDCLayout)
+	{
+		Init(DataChannel);
+		CachedLayout = NDCLayout;
+	}
+
+#if DEBUG_NDC_ACCESS
+	for (FNDCVarAccessorBase* Var : VariableAccessors)
+	{
+		Var->ValidateLayout();
+	}
+#endif
+
+	check(IsInGameThread());
+	UWorld* World = (WorldContextObject != nullptr) ? GEngine->GetWorldFromContextObject(WorldContextObject, EGetWorldErrorMode::LogAndReturnNull) : nullptr;
+	if (World && CachedLayout.IsValid())
+	{
+		if (FNiagaraWorldManager* WorldMan = FNiagaraWorldManager::Get(World))
+		{
+			if (UNiagaraDataChannelHandler* Handler = WorldMan->GetDataChannelManager().FindDataChannelHandler(DataChannel))
+			{
+				if (FNiagaraDataChannelDataPtr NDCData = Handler->FindData(SearchParams, ENiagaraResourceAccess::ReadOnly))
+				{
+					Data = NDCData->GetGameData()->AsShared();
+					bPreviousFrame = bInPreviousFrame;
+					return true;
+				}
+			}
+		}
+	}
+	return false;
 }
 
 void FNDCReaderBase::EndRead()
@@ -435,4 +458,4 @@ int32 FNDCReaderBase::Num()const
 		return bPreviousFrame ? Data->PrevNum() : Data->Num();
 	}
 	return INDEX_NONE;
-}
\ No newline at end of file
+}
diff --git a/Engine/Plugins/FX/Niagara/Source/Niagara/Private/NiagaraDataChannelCommon.cpp b/Engine/Plugins/FX/Niagara/Source/Niagara/Private/NiagaraDataChannelCommon.cpp
--- a/Engine/Plugins/FX/Niagara/Source/Niagara/Private/NiagaraDataChannelCommon.cpp
+++ b/Engine/Plugins/FX/Niagara/Source/Niagara/Private/NiagaraDataChannelCommon.cpp
@@ -2,19 +2,8 @@
 
 #include "NiagaraDataChannelCommon.h"
 #include "Components/SceneComponent.h"
-#include "NiagaraComponent.h"
-
-FNDCSpawnedSystemRef::FNDCSpawnedSystemRef(UNiagaraComponent* Component)
-	: SpawnedSystem(Component)
-{
-}
-
-FVector FNDCAccessContext::GetLocation()const
-{
-	return OwningComponent && !bOverrideLocation ? OwningComponent->GetComponentLocation() : Location;
-}
 
 FVector FNiagaraDataChannelSearchParameters::GetLocation()const
 {
 	return OwningComponent && !bOverrideLocation ? OwningComponent->GetComponentLocation() : Location;
-}
+}
\ No newline at end of file
diff --git a/Engine/Plugins/FX/Niagara/Source/Niagara/Private/NiagaraDataChannelHandler.cpp b/Engine/Plugins/FX/Niagara/Source/Niagara/Private/NiagaraDataChannelHandler.cpp
--- a/Engine/Plugins/FX/Niagara/Source/Niagara/Private/NiagaraDataChannelHandler.cpp
+++ b/Engine/Plugins/FX/Niagara/Source/Niagara/Private/NiagaraDataChannelHandler.cpp
@@ -23,9 +23,7 @@ bool UNiagaraDataChannelHandler::IsReadyForFinishDestroy()
 
 void UNiagaraDataChannelHandler::Init(const UNiagaraDataChannel* InChannel)
 {
-	check(InChannel);
 	DataChannel = InChannel;
-	AccessContextType = InChannel->GetAccessContextType();
 }
 
 void UNiagaraDataChannelHandler::Cleanup()
@@ -49,7 +47,7 @@ void UNiagaraDataChannelHandler::Cleanup()
 	}
 }
 
-void UNiagaraDataChannelHandler::BeginFrame(float DeltaTime)
+void UNiagaraDataChannelHandler::BeginFrame(float DeltaTime, FNiagaraWorldManager* OwningWorld)
 {
 	CurrentTG = TG_PrePhysics;
 
@@ -62,12 +60,12 @@ void UNiagaraDataChannelHandler::BeginFrame(float DeltaTime)
 	}
 }
 
-void UNiagaraDataChannelHandler::EndFrame(float DeltaTime)
+void UNiagaraDataChannelHandler::EndFrame(float DeltaTime, FNiagaraWorldManager* OwningWorld)
 {
 
 }
 
-void UNiagaraDataChannelHandler::Tick(float DeltaTime, ETickingGroup TickGroup)
+void UNiagaraDataChannelHandler::Tick(float DeltaTime, ETickingGroup TickGroup, FNiagaraWorldManager* OwningWorld)
 {
 	CurrentTG = TickGroup;
 }
@@ -92,23 +90,6 @@ UNiagaraDataChannelReader* UNiagaraDataChannelHandler::GetDataChannelReader()
 	return Reader;
 }
 
-void UNiagaraDataChannelHandler::SubscribeToDataChannelUpdates_WithContext(FOnNewNiagaraDataChannelPublish UpdateDelegate, FNDCAccessContextInst& AccessContext, int32& UnsubscribeToken)
-{
-	if (!UpdateDelegate.IsBound())
-	{
-		UnsubscribeToken = -1;
-		return;
-	}
-
-	SubscriberTokens++;
-	UnsubscribeToken = SubscriberTokens;
-
-	FChannelSubscription& ChannelSubscription = ChannelSubscriptions.AddDefaulted_GetRef();
-	ChannelSubscription.SubscriptionToken = SubscriberTokens;
-	ChannelSubscription.OnPublishDelegate = UpdateDelegate;
-	ChannelSubscription.AccessContext = AccessContext;
-}
-
 void UNiagaraDataChannelHandler::SubscribeToDataChannelUpdates(FOnNewNiagaraDataChannelPublish UpdateDelegate, FNiagaraDataChannelSearchParameters SearchParams, int32& UnsubscribeToken)
 {
 	if (!UpdateDelegate.IsBound())
@@ -167,16 +148,7 @@ void UNiagaraDataChannelHandler::NotifySubscribers(FNiagaraDataChannelData* Sour
 			continue;
 		}
 
-		FNiagaraDataChannelDataPtr ChannelData = nullptr;
-		if(ChannelSubscription.AccessContext.IsValid())
-		{
-			ChannelData = FindData(ChannelSubscription.AccessContext, ENiagaraResourceAccess::ReadOnly);
-		}
-		else
-		{
-			ChannelData = FindData(ChannelSubscription.SearchParams, ENiagaraResourceAccess::ReadOnly);
-		}
-
+		FNiagaraDataChannelDataPtr ChannelData = FindData(ChannelSubscription.SearchParams, ENiagaraResourceAccess::ReadOnly);
 		if (ChannelData.Get() == Source)
 		{
 			UpdateContext.Reader->Data = ChannelData;
@@ -195,47 +167,3 @@ void UNiagaraDataChannelHandler::OnComputeDispatchInterfaceDestroyed(FNiagaraGpu
 		NDCData->DestroyRenderThreadProxy(InComputeDispatchInterface);
 	});
 }
-
-
-void UNiagaraDataChannelHandler::VerifyAccessContext(FNDCAccessContextInst& AccessContext)
-{
-#if !UE_BUILD_SHIPPING
-	if(const UNiagaraDataChannel* Channel = DataChannel.Get())
-	{
-		const UScriptStruct* Expected = Channel->GetAccessContextType().Get();
-		const UScriptStruct* Actual = AccessContext.GetScriptStruct();
-		if(Actual->IsChildOf(Expected) == false)
-		{
-			UE_LOGFMT(LogNiagara, Warning, "Attempting to find Niagara Data Channel with incorrect access context type.\nData Channel: {DataChannel}\nExpected Type:{Expected}\nActual:{ActualType}", 
-				GetNameSafe(DataChannel->GetAsset()), GetNameSafe(Expected), GetNameSafe(Actual));
-		}
-	}
-#endif
-}
-
-FNiagaraDataChannelDataPtr UNiagaraDataChannelHandler::FindData(FNiagaraDataChannelSearchParameters SearchParams, ENiagaraResourceAccess AccessType)
-{
-	return nullptr;
-}
-
-FNiagaraDataChannelDataPtr UNiagaraDataChannelHandler::FindData(FNDCAccessContextInst& AccessContext, ENiagaraResourceAccess AccessType)
-{
-	if(AccessContext.IsValid() == false)
-	{
-		return nullptr;
-	}
-
-	VerifyAccessContext(AccessContext);
-	//Base behavior is to just convert to legacy params and pass through.
-	//We'll only need this for NDCs that do not directly support the new API but people still call in with the new API.
-	//We only support accesses using or derived from FNDCAccessContext for legacy NDC types.
-	if(const FNDCAccessContextLegacy* AccessCtx = AccessContext.Get<FNDCAccessContextLegacy>())
-	{
-		FNiagaraDataChannelSearchParameters SearchParams;
-		SearchParams.OwningComponent = AccessCtx->GetOwner();
-		SearchParams.Location = AccessCtx->Location;
-		SearchParams.bOverrideLocation = AccessCtx->bOverrideLocation;
-		return FindData(SearchParams, AccessType);
-	}	
-	return nullptr;
-}
diff --git a/Engine/Plugins/FX/Niagara/Source/Niagara/Private/NiagaraDataChannelManager.cpp b/Engine/Plugins/FX/Niagara/Source/Niagara/Private/NiagaraDataChannelManager.cpp
--- a/Engine/Plugins/FX/Niagara/Source/Niagara/Private/NiagaraDataChannelManager.cpp
+++ b/Engine/Plugins/FX/Niagara/Source/Niagara/Private/NiagaraDataChannelManager.cpp
@@ -78,7 +78,7 @@ void FNiagaraDataChannelManager::BeginFrame(float DeltaSeconds)
 		{
 			if(const UNiagaraDataChannel* Channel = It.Key().Get())
 			{
-				It.Value()->BeginFrame(DeltaSeconds);
+				It.Value()->BeginFrame(DeltaSeconds, WorldMan);
 			}
 			else
 			{
@@ -102,7 +102,7 @@ void FNiagaraDataChannelManager::EndFrame(float DeltaSeconds)
 		{
 			if (const UNiagaraDataChannel* Channel = It.Key().Get())
 			{
-				It.Value()->EndFrame(DeltaSeconds);
+				It.Value()->EndFrame(DeltaSeconds, WorldMan);
 			}
 			else
 			{
@@ -124,7 +124,7 @@ void FNiagaraDataChannelManager::Tick(float DeltaSeconds, ETickingGroup TickGrou
 		//Tick all DataChannel channel handlers.
 		for (auto& ChannelPair : Channels)
 		{
-			ChannelPair.Value->Tick(DeltaSeconds, TickGroup);
+			ChannelPair.Value->Tick(DeltaSeconds, TickGroup, WorldMan);
 		}
 	}
 	else
@@ -176,7 +176,6 @@ UNiagaraDataChannelHandler* FNiagaraDataChannelManager::InitDataChannel(const UN
 				Handler->Cleanup();
 			}
 			Handler = InChannel->CreateHandler(WorldMan->GetWorld());
-			Handler->SetOwningWorld(WorldMan);
 		}
 		return Handler;
 	}
diff --git a/Engine/Plugins/FX/Niagara/Source/Niagara/Private/NiagaraDataChannelManager.h b/Engine/Plugins/FX/Niagara/Source/Niagara/Private/NiagaraDataChannelManager.h
--- a/Engine/Plugins/FX/Niagara/Source/Niagara/Private/NiagaraDataChannelManager.h
+++ b/Engine/Plugins/FX/Niagara/Source/Niagara/Private/NiagaraDataChannelManager.h
@@ -42,7 +42,6 @@ public:
 	UNiagaraDataChannelHandler* FindDataChannelHandler(const UNiagaraDataChannelAsset* Channel) { return Channel ? FindDataChannelHandler(Channel->Get()) : nullptr; }
 
 	UWorld* GetWorld()const;
-	FNiagaraWorldManager* GetWorldMan()const{ return WorldMan; }
 
 	void OnComputeDispatchInterfaceDestroyed(FNiagaraGpuComputeDispatchInterface* InComputeDispatchInterface);
 
diff --git a/Engine/Plugins/FX/Niagara/Source/Niagara/Private/NiagaraDataChannel_GameplayBurst.cpp b/Engine/Plugins/FX/Niagara/Source/Niagara/Private/NiagaraDataChannel_GameplayBurst.cpp
deleted file mode 100644
diff --git a/Engine/Plugins/FX/Niagara/Source/Niagara/Private/NiagaraDataChannel_Global.cpp b/Engine/Plugins/FX/Niagara/Source/Niagara/Private/NiagaraDataChannel_Global.cpp
--- a/Engine/Plugins/FX/Niagara/Source/Niagara/Private/NiagaraDataChannel_Global.cpp
+++ b/Engine/Plugins/FX/Niagara/Source/Niagara/Private/NiagaraDataChannel_Global.cpp
@@ -40,49 +40,44 @@ UNiagaraDataChannelHandler_Global::UNiagaraDataChannelHandler_Global(FObjectInit
 
 void UNiagaraDataChannelHandler_Global::Init(const UNiagaraDataChannel* InChannel)
 {
+	check(InChannel);
 	Super::Init(InChannel);
+	Data = CreateData();
 }
 
 void UNiagaraDataChannelHandler_Global::Cleanup()
 {
 	Data.Reset();
+
 	Super::Cleanup();
 }
 
-void UNiagaraDataChannelHandler_Global::BeginFrame(float DeltaTime)
+void UNiagaraDataChannelHandler_Global::BeginFrame(float DeltaTime, FNiagaraWorldManager* OwningWorld)
 {
-	Super::BeginFrame(DeltaTime);
-	if(Data)
-	{
-		Data->BeginFrame(this);
-	}
+	Super::BeginFrame(DeltaTime, OwningWorld);
+	Data->BeginFrame(this);
 }
 
-void UNiagaraDataChannelHandler_Global::EndFrame(float DeltaTime)
+void UNiagaraDataChannelHandler_Global::EndFrame(float DeltaTime, FNiagaraWorldManager* OwningWorld)
 {
-	Super::EndFrame(DeltaTime);
-	if (Data)
-	{
-		Data->EndFrame(this);
-	}
+	Super::EndFrame(DeltaTime, OwningWorld);
+	Data->EndFrame(this);
 }
 
-void UNiagaraDataChannelHandler_Global::Tick(float DeltaSeconds, ETickingGroup TickGroup)
+void UNiagaraDataChannelHandler_Global::Tick(float DeltaSeconds, ETickingGroup TickGroup, FNiagaraWorldManager* OwningWorld)
 {
 	SCOPE_CYCLE_COUNTER(STAT_DataChannelHandler_Global_Tick);
-	Super::Tick(DeltaSeconds, TickGroup);
-	if (Data)
-	{
-		Data->ConsumePublishRequests(this, TickGroup);
-	}
+	Super::Tick(DeltaSeconds, TickGroup, OwningWorld);
+
+	Data->ConsumePublishRequests(this, TickGroup);
 }
 
 FNiagaraDataChannelDataPtr UNiagaraDataChannelHandler_Global::FindData(FNiagaraDataChannelSearchParameters SearchParams, ENiagaraResourceAccess AccessType)
-{	
-	//For more complicated channels we could check the location + bounds of the system instance etc to return some spatially localized data.
-	if(Data.IsValid() == false)
+{
+	if(!Data)
 	{
 		Data = CreateData();
 	}
 	return Data;
+	//For more complicated channels we could check the location + bounds of the system instance etc to return some spatially localized data.
 }
diff --git a/Engine/Plugins/FX/Niagara/Source/Niagara/Private/NiagaraDataChannel_Islands.cpp b/Engine/Plugins/FX/Niagara/Source/Niagara/Private/NiagaraDataChannel_Islands.cpp
--- a/Engine/Plugins/FX/Niagara/Source/Niagara/Private/NiagaraDataChannel_Islands.cpp
+++ b/Engine/Plugins/FX/Niagara/Source/Niagara/Private/NiagaraDataChannel_Islands.cpp
@@ -141,9 +141,9 @@ void UNiagaraDataChannelHandler_Islands::Cleanup()
 	Super::Cleanup();
 }
 
-void UNiagaraDataChannelHandler_Islands::BeginFrame(float DeltaTime)
+void UNiagaraDataChannelHandler_Islands::BeginFrame(float DeltaTime, FNiagaraWorldManager* OwningWorld)
 {
-	Super::BeginFrame(DeltaTime);
+	Super::BeginFrame(DeltaTime, OwningWorld);
 	for (auto It = ActiveIslands.CreateIterator(); It; ++It)
 	{
 		FNDCIsland& Island = IslandPool[*It];
@@ -151,9 +151,9 @@ void UNiagaraDataChannelHandler_Islands::BeginFrame(float DeltaTime)
 	}
 }
 
-void UNiagaraDataChannelHandler_Islands::EndFrame(float DeltaTime)
+void UNiagaraDataChannelHandler_Islands::EndFrame(float DeltaTime, FNiagaraWorldManager* OwningWorld)
 {
-	Super::EndFrame(DeltaTime);
+	Super::EndFrame(DeltaTime, OwningWorld);
 	const UNiagaraDataChannel_Islands* IslandChannel = GetChannelTyped<const UNiagaraDataChannel_Islands>();
 	for (auto It = ActiveIslands.CreateIterator(); It; ++It)
 	{
@@ -175,9 +175,9 @@ void UNiagaraDataChannelHandler_Islands::EndFrame(float DeltaTime)
 	}
 }
 
-void UNiagaraDataChannelHandler_Islands::Tick(float DeltaTime, ETickingGroup TickGroup)
+void UNiagaraDataChannelHandler_Islands::Tick(float DeltaTime, ETickingGroup TickGroup, FNiagaraWorldManager* OwningWorld)
 {
-	Super::Tick(DeltaTime, TickGroup);
+	Super::Tick(DeltaTime, TickGroup, OwningWorld);
 
 	//static UEnum* TGEnum = StaticEnum<ETickingGroup>();
 	//UE_LOG(LogNiagara,Warning, TEXT("UNiagaraDataChannelHandler_Islands::Tick - %s"), *TGEnum->GetDisplayNameTextByValue((int32)TickGroup).ToString());
diff --git a/Engine/Plugins/FX/Niagara/Source/Niagara/Private/NiagaraDataChannel_Map.cpp b/Engine/Plugins/FX/Niagara/Source/Niagara/Private/NiagaraDataChannel_Map.cpp
deleted file mode 100644
diff --git a/Engine/Plugins/FX/Niagara/Source/Niagara/Private/NiagaraSystemCollection.cpp b/Engine/Plugins/FX/Niagara/Source/Niagara/Private/NiagaraSystemCollection.cpp
deleted file mode 100644
diff --git a/Engine/Plugins/FX/Niagara/Source/Niagara/Public/NiagaraDataChannel.h b/Engine/Plugins/FX/Niagara/Source/Niagara/Public/NiagaraDataChannel.h
--- a/Engine/Plugins/FX/Niagara/Source/Niagara/Public/NiagaraDataChannel.h
+++ b/Engine/Plugins/FX/Niagara/Source/Niagara/Public/NiagaraDataChannel.h
@@ -157,6 +157,8 @@ private:
 	bool bKeepPreviousFrameData = false;
 };
 
+DECLARE_MULTICAST_DELEGATE_OneParam(FOnDataChannelCreated, const UNiagaraDataChannel*);
+
 UCLASS(abstract, EditInlineNew, MinimalAPI, prioritizeCategories=("Data Channel"))
 class UNiagaraDataChannel : public UObject
 {
@@ -179,12 +181,6 @@ public:
 
 	/** If true, we keep our previous frame's data. Some users will prefer a frame of latency to tick dependency. */
 	bool KeepPreviousFrameData() const { return bKeepPreviousFrameData; }
-	
-	/** Returns the type for the Access Context used by this Data Channel. */
-	[[nodiscard]] NIAGARA_API virtual TNDCAccessContextType GetAccessContextType()const { return TNDCAccessContextType(FNDCAccessContextLegacy::StaticStruct()); }
-	
-	/** Returns a pre allocated access context that you can use for accessing this data channel. Only safe to use from the game thread and it is not safe to hold onto this context after your access. */
-	[[nodiscard]] inline FNDCAccessContextInst& GetTransientAccessContext()const;
 
 	/** Create the appropriate handler object for this data channel. */
 	virtual UNiagaraDataChannelHandler* CreateHandler(UWorld* OwningWorld) const PURE_VIRTUAL(UNiagaraDataChannel::CreateHandler, {return nullptr;} );
@@ -252,18 +248,8 @@ private:
 	#endif
 
 	FRenderCommandFence RTFence;
-
-	/** A transient access context instance that calling code can use for accessing this Data Channel. GT only*/
-	mutable FNDCAccessContextInst TransientAccessContext;
 };
 
-FNDCAccessContextInst& UNiagaraDataChannel::GetTransientAccessContext()const
-{
-	check(IsInGameThread());
-	TransientAccessContext.Init(GetAccessContextType());//We need to init (or reset) ready for use.
-	return TransientAccessContext;
-}
-
 template<typename TAction>
 void UNiagaraDataChannel::ForEachDataChannel(TAction Func)
 {
@@ -297,14 +283,8 @@ class UNiagaraDataChannelLibrary : public UBlueprintFunctionLibrary
 	UFUNCTION(BlueprintInternalUseOnly, Category = NiagaraDataChannel, meta = (Keywords = "niagara DataChannel", WorldContext = "WorldContextObject", UnsafeDuringActorConstruction = "true"))
 	static UE_API UNiagaraDataChannelHandler* GetNiagaraDataChannel(const UObject* WorldContextObject, const UNiagaraDataChannelAsset* Channel);
 
-	
-	//////////////////////////////////////////////////////////////////////////
-	//These functions are soon to be deprecated in favor of new functions that take a more general AccessContext instead of Search Parameters.
-
 	/**
-	 * LEGACY FUNCTION: Please use non-legacy version.
 	 * Initializes and returns the Niagara Data Channel writer to write N elements to the given data channel.
-	 * This function is now legacy and and will soon be deprecated. Please use the new non-legacy function.
 	 *
 	 * @param WorldContextObject	World to execute in
 	 * @param Channel				The channel to write to
@@ -315,13 +295,11 @@ class UNiagaraDataChannelLibrary : public UBlueprintFunctionLibrary
 	 * @param bVisibleToGPU	If true, the data written to this data channel is visible to Niagara GPU emitters
 	 * @param DebugSource	Instigator for this write, used in the debug hud to track writes to the data channel from different sources
 	 */
-	UFUNCTION(BlueprintCallable, Category = NiagaraDataChannel, DisplayName="Write To Niagara Data Channel (Batch) (Legacy)", meta = (AdvancedDisplay = "SearchParams, DebugSource", Keywords = "niagara DataChannel", WorldContext = "WorldContextObject", UnsafeDuringActorConstruction = "true", AutoCreateRefTerm="DebugSource"))
+	UFUNCTION(BlueprintCallable, Category = NiagaraDataChannel, DisplayName="Write To Niagara Data Channel (Batch)", meta = (AdvancedDisplay = "SearchParams, DebugSource", Keywords = "niagara DataChannel", WorldContext = "WorldContextObject", UnsafeDuringActorConstruction = "true", AutoCreateRefTerm="DebugSource"))
 	static UE_API UNiagaraDataChannelWriter* WriteToNiagaraDataChannel(const UObject* WorldContextObject, const UNiagaraDataChannelAsset* Channel, FNiagaraDataChannelSearchParameters SearchParams, int32 Count, UPARAM(DisplayName = "Visible to Blueprint") bool bVisibleToGame, UPARAM(DisplayName = "Visible to Niagara CPU") bool bVisibleToCPU, UPARAM(DisplayName = "Visible to Niagara GPU") bool bVisibleToGPU, const FString& DebugSource);
 
 	/**
-	 * LEGACY FUNCTION: Please use non-legacy version.
 	 * Initializes and returns the Niagara Data Channel reader for the given data channel.
-	 * This function is now legacy and and will soon be deprecated. Please use the new non-legacy function.
 	 *
 	 * @param WorldContextObject	World to execute in
 	 * @param Channel				The channel to read from
@@ -330,13 +308,11 @@ class UNiagaraDataChannelLibrary : public UBlueprintFunctionLibrary
 	 *								Reading the current frame allows for zero latency reads, but any data elements that are generated after this reader is used are missed.
 	 *								Reading the previous frame's data introduces a frame of latency but ensures we never miss any data as we have access to the whole frame.
 	 */
-	UFUNCTION(BlueprintCallable, Category = NiagaraDataChannel, DisplayName="Read From Niagara Data Channel (Batch) (Legacy)", meta = (AdvancedDisplay = "SearchParams", Keywords = "niagara DataChannel", WorldContext = "WorldContextObject", UnsafeDuringActorConstruction = "true"))
+	UFUNCTION(BlueprintCallable, Category = NiagaraDataChannel, DisplayName="Read From Niagara Data Channel (Batch)", meta = (AdvancedDisplay = "SearchParams", Keywords = "niagara DataChannel", WorldContext = "WorldContextObject", UnsafeDuringActorConstruction = "true"))
 	static UE_API UNiagaraDataChannelReader* ReadFromNiagaraDataChannel(const UObject* WorldContextObject, const UNiagaraDataChannelAsset* Channel, FNiagaraDataChannelSearchParameters SearchParams, bool bReadPreviousFrame);
 
 	/**
-	 * LEGACY FUNCTION: Please use non-legacy version.
 	 * Returns the number of readable elements in the given data channel
-	 * This function is now legacy and and will soon be deprecated. Please use the new non-legacy function.
 	 *
 	 * @param WorldContextObject	World to execute in
 	 * @param Channel				The channel to read from
@@ -345,11 +321,10 @@ class UNiagaraDataChannelLibrary : public UBlueprintFunctionLibrary
 	 *								Reading the current frame allows for zero latency reads, but any data elements that are generated after this reader is used are missed.
 	 *								Reading the previous frame's data introduces a frame of latency but ensures we never miss any data as we have access to the whole frame.
 	 */
-	UFUNCTION(BlueprintCallable, Category = NiagaraDataChannel, meta = (bReadPreviousFrame="true", DisplayName="Get Data Channel Element Count (Legacy)", AdvancedDisplay = "SearchParams, bReadPreviousFrame", Keywords = "niagara DataChannel num size", WorldContext = "WorldContextObject", UnsafeDuringActorConstruction = "true"))
+	UFUNCTION(BlueprintCallable, Category = NiagaraDataChannel, meta = (bReadPreviousFrame="true", AdvancedDisplay = "SearchParams, bReadPreviousFrame", Keywords = "niagara DataChannel num size", WorldContext = "WorldContextObject", UnsafeDuringActorConstruction = "true"))
 	static UE_API int32 GetDataChannelElementCount(const UObject* WorldContextObject, const UNiagaraDataChannelAsset* Channel, FNiagaraDataChannelSearchParameters SearchParams, bool bReadPreviousFrame);
 
 	/**
-	 * LEGACY FUNCTION: Please use non-legacy version.
 	 * Subscribes to a single data channel and calls a delegate every times new data is written to the data channel.
 	 *
 	 * @param WorldContextObject	World to execute in
@@ -358,7 +333,7 @@ class UNiagaraDataChannelLibrary : public UBlueprintFunctionLibrary
 	 * @param UpdateDelegate		The delegate to be called when new data is available in the data channel. Can be called multiple times per tick.
 	 * @param UnsubscribeToken		This token can be used to unsubscribe from the data channel.
 	 */
-	UFUNCTION(BlueprintCallable, Category = NiagaraDataChannel, meta = (DisplayName="Subscribe To Niagara Data Channel (Legacy)", Keywords = "niagara DataChannel event reader subscription delegate listener updates", WorldContext = "WorldContextObject", UnsafeDuringActorConstruction = "true"))
+	UFUNCTION(BlueprintCallable, Category = NiagaraDataChannel, meta = (Keywords = "niagara DataChannel event reader subscription delegate listener updates", WorldContext = "WorldContextObject", UnsafeDuringActorConstruction = "true"))
 	static UE_API void SubscribeToNiagaraDataChannel(const UObject* WorldContextObject, const UNiagaraDataChannelAsset* Channel, FNiagaraDataChannelSearchParameters SearchParams, const FOnNewNiagaraDataChannelPublish& UpdateDelegate, int32& UnsubscribeToken);
 
 	/**
@@ -372,7 +347,6 @@ class UNiagaraDataChannelLibrary : public UBlueprintFunctionLibrary
 	static UE_API void UnsubscribeFromNiagaraDataChannel(const UObject* WorldContextObject, const UNiagaraDataChannelAsset* Channel, const int32& UnsubscribeToken);
 	
 	/**
-	 * LEGACY FUNCTION: Please use non-legacy version.
 	 * Reads a single entry from the given data channel, if possible.
 	 *
 	 * @param WorldContextObject	World to execute in
@@ -384,11 +358,10 @@ class UNiagaraDataChannelLibrary : public UBlueprintFunctionLibrary
 	 *								Reading the previous frame's data introduces a frame of latency but ensures we never miss any data as we have access to the whole frame.
 	 * @param ReadResult			Used by Blueprint for the return value
 	 */
-	UFUNCTION(BlueprintInternalUseOnly, Category = NiagaraDataChannel, DisplayName="Read From Niagara Data Channel (Legacy)", meta = (bReadPreviousFrame="true", AdvancedDisplay = "SearchParams, bReadPreviousFrame", ExpandEnumAsExecs="ReadResult", Keywords = "niagara DataChannel", WorldContext = "WorldContextObject", UnsafeDuringActorConstruction = "true"))
+	UFUNCTION(BlueprintInternalUseOnly, Category = NiagaraDataChannel, DisplayName="Read From Niagara Data Channel", meta = (bReadPreviousFrame="true", AdvancedDisplay = "SearchParams, bReadPreviousFrame", ExpandEnumAsExecs="ReadResult", Keywords = "niagara DataChannel", WorldContext = "WorldContextObject", UnsafeDuringActorConstruction = "true"))
 	static UE_API void ReadFromNiagaraDataChannelSingle(const UObject* WorldContextObject, const UNiagaraDataChannelAsset* Channel, int32 Index, FNiagaraDataChannelSearchParameters SearchParams, bool bReadPreviousFrame, ENiagartaDataChannelReadResult& ReadResult);
 
 	/**
-	 * LEGACY FUNCTION: Please use non-legacy version.
 	 * Writes a single element to a Niagara Data Channel. The element won't be immediately visible to readers, as it needs to be processed first. The earliest point it can be read is in the next tick group.
 	 *
 	 * @param WorldContextObject	World to execute in
@@ -398,132 +371,12 @@ class UNiagaraDataChannelLibrary : public UBlueprintFunctionLibrary
 	 * @param bVisibleToNiagaraCPU	If true, the data written to this data channel is visible to Niagara CPU emitters
 	 * @param bVisibleToNiagaraGPU	If true, the data written to this data channel is visible to Niagara GPU emitters
 	 */
-	UFUNCTION(BlueprintInternalUseOnly, Category = NiagaraDataChannel, DisplayName="Write To Niagara Data Channel (Legacy)", meta = (bVisibleToBlueprint="true", bVisibleToNiagaraCPU="true", bVisibleToNiagaraGPU="true", AdvancedDisplay = "bVisibleToBlueprint, bVisibleToNiagaraCPU, bVisibleToNiagaraGPU, SearchParams", Keywords = "niagara DataChannel event writer", WorldContext = "WorldContextObject", UnsafeDuringActorConstruction = "true"))
+	UFUNCTION(BlueprintInternalUseOnly, Category = NiagaraDataChannel, DisplayName="Write To Niagara Data Channel", meta = (bVisibleToBlueprint="true", bVisibleToNiagaraCPU="true", bVisibleToNiagaraGPU="true", AdvancedDisplay = "bVisibleToBlueprint, bVisibleToNiagaraCPU, bVisibleToNiagaraGPU, SearchParams", Keywords = "niagara DataChannel event writer", WorldContext = "WorldContextObject", UnsafeDuringActorConstruction = "true"))
 	static UE_API void WriteToNiagaraDataChannelSingle(const UObject* WorldContextObject, const UNiagaraDataChannelAsset* Channel, FNiagaraDataChannelSearchParameters SearchParams, bool bVisibleToBlueprint, bool bVisibleToNiagaraCPU, bool bVisibleToNiagaraGPU);
 
-	//////////////////////////////////////////////////////////////////////////
-
-	/**
-	 * Initializes and returns the Niagara Data Channel reader for the given data channel.
-	 *
-	 * @param WorldContextObject	World to execute in
-	 * @param Channel				The channel to read from
-	 * @param AccessContext			Context allowing passing of NDC type specific information between the caller and internal Data Channel that can control internal routing and help feedback information to the caller.
-	 * @param bReadPreviousFrame	True if this reader will read the previous frame's data. If false, we read the current frame.
-	 *								Reading the current frame allows for zero latency reads, but any data elements that are generated after this reader is used are missed.
-	 *								Reading the previous frame's data introduces a frame of latency but ensures we never miss any data as we have access to the whole frame.
-	 */
-	UFUNCTION(BlueprintCallable, BlueprintInternalUseOnly, Category = NiagaraDataChannel, DisplayName="Read From Niagara Data Channel (Batch)", meta = (Keywords = "niagara DataChannel", WorldContext = "WorldContextObject", UnsafeDuringActorConstruction = "true"))
-	static UE_API UNiagaraDataChannelReader* ReadFromNiagaraDataChannel_WithContext(const UObject* WorldContextObject, const UNiagaraDataChannelAsset* Channel, UPARAM(ref) FNDCAccessContextInst& AccessContext, bool bReadPreviousFrame);
-	
-	/**
-	 * Initializes and returns the Niagara Data Channel writer to write N elements to the given data channel.
-	 *
-	 * @param WorldContextObject	World to execute in
-	 * @param Channel				The channel to write to
-	 * @param AccessContext			Context allowing passing of NDC type specific information between the caller and internal Data Channel that can control internal routing and help feedback information to the caller.
-	 * @param Count					The number of elements to write 
-	 * @param bVisibleToGame	If true, the data written to this data channel is visible to Blueprint and C++ logic reading from it
-	 * @param bVisibleToCPU	If true, the data written to this data channel is visible to Niagara CPU emitters
-	 * @param bVisibleToGPU	If true, the data written to this data channel is visible to Niagara GPU emitters
-	 * @param DebugSource	Instigator for this write, used in the debug hud to track writes to the data channel from different sources
-	 */
-	UFUNCTION(BlueprintCallable, BlueprintInternalUseOnly, Category = NiagaraDataChannel, DisplayName = "Write To Niagara Data Channel (Batch)", meta = (AdvancedDisplay = "DebugSource", Keywords = "niagara DataChannel", WorldContext = "WorldContextObject", UnsafeDuringActorConstruction = "true", AutoCreateRefTerm = "DebugSource", bVisibleToBlueprint="true", bVisibleToNiagaraCPU="true", bVisibleToNiagaraGPU="true"))
-	static UE_API UNiagaraDataChannelWriter* WriteToNiagaraDataChannel_WithContext(const UObject* WorldContextObject, const UNiagaraDataChannelAsset* Channel, UPARAM(ref) FNDCAccessContextInst& AccessContext, int32 Count, UPARAM(DisplayName = "Visible to Blueprint") bool bVisibleToBlueprint, UPARAM(DisplayName = "Visible to Niagara CPU") bool bVisibleToNiagaraCPU, UPARAM(DisplayName = "Visible to Niagara GPU") bool bVisibleToNiagaraGPU, const FString& DebugSource);
-	
-	/**
-	 * Returns the number of readable elements in the given data channel
-	 *
-	 * @param WorldContextObject	World to execute in
-	 * @param Channel				The channel to read from
-	 * @param AccessContext			Context allowing passing of NDC type specific information between the caller and internal Data Channel that can control internal routing and help feedback information to the caller.
-	 * @param bReadPreviousFrame	True if this reader will read the previous frame's data. If false, we read the current frame.
-	 *								Reading the current frame allows for zero latency reads, but any data elements that are generated after this reader is used are missed.
-	 *								Reading the previous frame's data introduces a frame of latency but ensures we never miss any data as we have access to the whole frame.
-	 */
-	UFUNCTION(BlueprintCallable, BlueprintInternalUseOnly,Category = NiagaraDataChannel, meta = (bReadPreviousFrame="true", Keywords = "niagara DataChannel num size", WorldContext = "WorldContextObject", UnsafeDuringActorConstruction = "true"))
-	static UE_API int32 GetDataChannelElementCount_WithContext(const UObject* WorldContextObject, const UNiagaraDataChannelAsset* Channel, UPARAM(ref) FNDCAccessContextInst& AccessContext, bool bReadPreviousFrame);
-	
-	/**
-	 * Subscribes to a single data channel and calls a delegate every times new data is written to the data channel.
-	 *
-	 * @param WorldContextObject	World to execute in
-	 * @param Channel				The channel to subscribe to for updates
-	 * @param AccessContext			Context allowing passing of NDC type specific information between the caller and internal Data Channel that can control internal routing and help feedback information to the caller.
-	 * @param UpdateDelegate		The delegate to be called when new data is available in the data channel. Can be called multiple times per tick.
-	 * @param UnsubscribeToken		This token can be used to unsubscribe from the data channel.
-	 */
-	UFUNCTION(BlueprintCallable, Category = NiagaraDataChannel, meta = (Keywords = "niagara DataChannel event reader subscription delegate listener updates", WorldContext = "WorldContextObject", UnsafeDuringActorConstruction = "true"))
-	static UE_API void SubscribeToNiagaraDataChannel_WithContext(const UObject* WorldContextObject, const UNiagaraDataChannelAsset* Channel, UPARAM(ref) FNDCAccessContextInst& AccessContext, const FOnNewNiagaraDataChannelPublish& UpdateDelegate, int32& UnsubscribeToken);
-		
-	/**
-	 * Reads a single entry from the given data channel, if possible.
-	 *
-	 * @param WorldContextObject	World to execute in
-	 * @param Channel				The channel to read from
-	 * @param Index					The data index to read from
-	 * @param AccessContext			Context allowing passing of NDC type specific information between the caller and internal Data Channel that can control internal routing and help feedback information to the caller.
-	 * @param bReadPreviousFrame	True if this reader will read the previous frame's data. If false, we read the current frame.
-	 *								Reading the current frame allows for zero latency reads, but any data elements that are generated after this reader is used are missed.
-	 *								Reading the previous frame's data introduces a frame of latency but ensures we never miss any data as we have access to the whole frame.
-	 * @param ReadResult			Used by Blueprint for the return value
-	 */
-	UFUNCTION(BlueprintInternalUseOnly, Category = NiagaraDataChannel, DisplayName="Read From Niagara Data Channel", meta = (bReadPreviousFrame="true", ExpandEnumAsExecs="ReadResult", Keywords = "niagara DataChannel", WorldContext = "WorldContextObject", UnsafeDuringActorConstruction = "true"))
-	static UE_API void ReadFromNiagaraDataChannelSingle_WithContext(const UObject* WorldContextObject, const UNiagaraDataChannelAsset* Channel, int32 Index, UPARAM(ref) FNDCAccessContextInst& AccessContext, bool bReadPreviousFrame, ENiagartaDataChannelReadResult& ReadResult);
-
-	/**
-	 * Writes a single element to a Niagara Data Channel. The element won't be immediately visible to readers, as it needs to be processed first. The earliest point it can be read is in the next tick group.
-	 *
-	 * @param WorldContextObject	World to execute in
-	 * @param Channel				The channel to write to
-	 * @param AccessContext			Context allowing passing of NDC type specific information between the caller and internal Data Channel that can control internal routing and help feedback information to the caller.
-	 * @param bVisibleToBlueprint	If true, the data written to this data channel is visible to Blueprint and C++ logic reading from it
-	 * @param bVisibleToNiagaraCPU	If true, the data written to this data channel is visible to Niagara CPU emitters
-	 * @param bVisibleToNiagaraGPU	If true, the data written to this data channel is visible to Niagara GPU emitters
-	 */
-	UFUNCTION(BlueprintInternalUseOnly, Category = NiagaraDataChannel, DisplayName="Write To Niagara Data Channel", meta = (bVisibleToBlueprint="true", bVisibleToNiagaraCPU="true", bVisibleToNiagaraGPU="true", Keywords = "niagara DataChannel event writer", WorldContext = "WorldContextObject", UnsafeDuringActorConstruction = "true"))
-	static UE_API void WriteToNiagaraDataChannelSingle_WithContext(const UObject* WorldContextObject, const UNiagaraDataChannelAsset* Channel, UPARAM(ref) FNDCAccessContextInst& AccessContext, bool bVisibleToBlueprint, bool bVisibleToNiagaraCPU, bool bVisibleToNiagaraGPU);
-
 	static UE_API UNiagaraDataChannelHandler* FindDataChannelHandler(const UObject* WorldContextObject, const UNiagaraDataChannel* Channel);
 	static UE_API UNiagaraDataChannelWriter* CreateDataChannelWriter(const UObject* WorldContextObject, const UNiagaraDataChannel* Channel, FNiagaraDataChannelSearchParameters SearchParams, int32 Count, bool bVisibleToGame, bool bVisibleToCPU, bool bVisibleToGPU, const FString& DebugSource);
 	static UE_API UNiagaraDataChannelReader* CreateDataChannelReader(const UObject* WorldContextObject, const UNiagaraDataChannel* Channel, FNiagaraDataChannelSearchParameters SearchParams, bool bReadPreviousFrame);
-
-	static UE_API UNiagaraDataChannelWriter* CreateDataChannelWriter_WithContext(const UObject* WorldContextObject, const UNiagaraDataChannel* Channel, FNDCAccessContextInst& AccessContext, int32 Count, bool bVisibleToGame, bool bVisibleToCPU, bool bVisibleToGPU, const FString& DebugSource);
-	static UE_API UNiagaraDataChannelReader* CreateDataChannelReader_WithContext(const UObject* WorldContextObject, const UNiagaraDataChannel* Channel, FNDCAccessContextInst& AccessContext, bool bReadPreviousFrame);
-
-
-	/** Creates a new NDCAccessContext with a specified internal struct type. Struct must be a child of FNDCAccessContextBase. */
-	UFUNCTION(BlueprintInternalUseOnly, BlueprintCallable, CustomThunk, Category = "Niagara Data Channel", meta = (HidePin = "ContextStruct"))
-	static UE_API FNDCAccessContextInst MakeNDCAccessContextInstance(UScriptStruct* ContextStruct);
-
-	/** 
-	Gets the members from the internal Struct of an FNDCAccessContextIns. Expanded on compile into individual per property calls to GetSinglePropertyInNDCAccessContextInstance.
-	*/
-	UFUNCTION(BlueprintInternalUseOnly, BlueprintCallable, Category = "Niagara Data Channel", meta = (HidePin = "ContextStruct"))
-	static UE_API void GetMembersInNDCAccessContextInstance(const FNDCAccessContextInst& AccessContext, UScriptStruct* ContextStruct);
-	
-	/** 
-	Sets the members in the internal Struct of an FNDCAccessContextIns. Expanded on compile into individual per property calls to SetSinglePropertyInNDCAccessContextInstance.
-	*/
-	UFUNCTION(BlueprintInternalUseOnly, BlueprintCallable, Category = "Niagara Data Channel", meta = (HidePin = "ContextStruct"))
-	static UE_API void SetMembersInNDCAccessContextInstance(UPARAM(Ref) const FNDCAccessContextInst& AccessContext, UScriptStruct* ContextStruct);
-	
-	/** 
-	Gets a specific property within the internal data of an FNDCAccessContextInst.
-	*/
-	UFUNCTION(BlueprintInternalUseOnly, BlueprintCallable, CustomThunk, Category = "Niagara Data Channel", meta = (CustomStructureParam = "Value", AutoCreateRefTerm="Value"))
-	static UE_API void GetSinglePropertyInNDCAccessContextInstance(const FNDCAccessContextInst& AccessContext, UScriptStruct* ContextStruct, FName PropertyName, int32& Value);
-	
-	/** 
-	Sets a specific property within the internal data of an FNDCAccessContextInst.
-	*/
-	UFUNCTION(BlueprintInternalUseOnly, BlueprintCallable, CustomThunk, Category = "Niagara Data Channel", meta = (CustomStructureParam = "Value", AutoCreateRefTerm="Value"))
-	static UE_API void SetSinglePropertyInNDCAccessContextInstance(const FNDCAccessContextInst& AccessContext, UScriptStruct* ContextStruct, FName PropertyName, const int32& Value);
-
-private:
-	DECLARE_FUNCTION(execMakeNDCAccessContextInstance);
-	DECLARE_FUNCTION(execGetSinglePropertyInNDCAccessContextInstance);
-	DECLARE_FUNCTION(execSetSinglePropertyInNDCAccessContextInstance);
 };
 
 #undef UE_API
diff --git a/Engine/Plugins/FX/Niagara/Source/Niagara/Public/NiagaraDataChannelAccessor.h b/Engine/Plugins/FX/Niagara/Source/Niagara/Public/NiagaraDataChannelAccessor.h
--- a/Engine/Plugins/FX/Niagara/Source/Niagara/Public/NiagaraDataChannelAccessor.h
+++ b/Engine/Plugins/FX/Niagara/Source/Niagara/Public/NiagaraDataChannelAccessor.h
@@ -38,15 +38,9 @@ public:
 	UPROPERTY()
 	TObjectPtr<UNiagaraDataChannelHandler> Owner;
 	
-	/** 
-	Call before each access to the data channel to grab the correct data to read.
-	NOTE: Init access is only to support legacy flows. Please use Begin Read.
-	*/
-	UFUNCTION(meta = (DisplayName="Init Access (Legacy)"))
-	NIAGARA_API bool InitAccess(FNiagaraDataChannelSearchParameters SearchParams, bool bReadPrevFrameData);
-
 	/** Call before each access to the data channel to grab the correct data to read. */
-	NIAGARA_API bool BeginRead(FNDCAccessContextInst& AccessContext, bool bReadPrevFrameData);
+	UFUNCTION(BlueprintCallable, Category = NiagaraDataChannel, meta = (Keywords = "Niagara DataChannel"))
+	NIAGARA_API bool InitAccess(FNiagaraDataChannelSearchParameters SearchParams, bool bReadPrevFrameData);
 
 	UFUNCTION(BlueprintCallable, Category = NiagaraDataChannel, meta = (Keywords = "Niagara DataChannel"))
 	NIAGARA_API int32 Num()const;
@@ -118,15 +112,9 @@ public:
 	UPROPERTY()
 	TObjectPtr<UNiagaraDataChannelHandler> Owner;
 	
-	/** 
-	Call before each batch of writes to allocate the data we'll be writing to.
-	NOTE: Init Write is only to support legacy flows. Please use Begin Write.
-	*/
-	UFUNCTION(meta = (DisplayName="Init Write (Legacy)"))
-	NIAGARA_API bool InitWrite(FNiagaraDataChannelSearchParameters SearchParams, int32 Count, bool bVisibleToGame, bool bVisibleToCPU, bool bVisibleToGPU, const FString& DebugSource);
-
 	/** Call before each batch of writes to allocate the data we'll be writing to. */
-	NIAGARA_API bool BeginWrite(FNDCAccessContextInst& AccessContext, int32 Count, bool bVisibleToGame, bool bVisibleToCPU, bool bVisibleToGPU, const FString& DebugSource);
+	UFUNCTION(BlueprintCallable, Category = NiagaraDataChannel, meta = (bVisibleToGame="true", bVisibleToCPU="true", bVisibleToGPU="true", Keywords = "niagara DataChannel", AdvancedDisplay = "DebugSource", AutoCreateRefTerm="DebugSource"))
+	NIAGARA_API bool InitWrite(FNiagaraDataChannelSearchParameters SearchParams, int32 Count, bool bVisibleToGame, bool bVisibleToCPU, bool bVisibleToGPU, const FString& DebugSource);
 
 	UFUNCTION(BlueprintCallable, Category = NiagaraDataChannel, meta = (Keywords = "niagara DataChannel"))
 	NIAGARA_API int32 Num()const;
@@ -213,7 +201,7 @@ public:
 * 
 * 3.
 * Use this writer object to write data into the Data Channel.
-* First you must call BeginWrite() to initialize the writer for the current access context.
+* First you must call BeginWrite() to initialize the writer for the current search parameters.
 * 
 * void FYouUserClass::YourFunction()
 * {
@@ -344,18 +332,12 @@ protected:
 	//Number of items we're expecting/allowed to write.
 	int32 Count = INDEX_NONE;
 
-	template<typename TContext>
-	inline bool BeginWrite_Internal(const UObject* WorldContextObject, const UNiagaraDataChannel* DataChannel, TContext& AccessContext, int32 InCount, bool bVisibleToGame, bool bVisibleToCPU, bool bVisibleToGPU);
-
 public:
 
-	//Finds the correct target NDC data using the current AccessContext.
+	//Finds the correct target NDC data using the current SearchParameters.
 	//Allocates enough space in the target NDC data.
 	//Refreshes internal layout info if needed.
-	UE_API bool BeginWrite(const UObject* WorldContextObject, const UNiagaraDataChannel* DataChannel, FNDCAccessContextInst& AccessContext, int32 InCount, bool bVisibleToGame, bool bVisibleToCPU, bool bVisibleToGPU);
-
-	//Legacy version taking an old FNiagaraDataChannelSearchParameters struct. Works for legacy code but will not support newer NDC types or features.
-	UE_API bool BeginWrite(const UObject* WorldContextObject, const UNiagaraDataChannel* DataChannel, FNiagaraDataChannelSearchParameters SearchParams, int32 InCount, bool bVisibleToGame, bool bVisibleToCPU, bool bVisibleToGPU);
+	UE_API bool BeginWrite(const UObject* WorldContextObject, const UNiagaraDataChannel* DataChannel, const FNiagaraDataChannelSearchParameters& SearchParams, int32 Count, bool bVisibleToGame, bool bVisibleToCPU, bool bVisibleToGPU);
 
 	//Finalizes the write and free/reset any data.
 	UE_API void EndWrite();
@@ -383,16 +365,10 @@ protected:
 	FNiagaraDataChannelGameDataPtr Data;
 	bool bPreviousFrame = false;
 
-	template<typename TContext>
-	inline bool BeginRead_Internal(const UObject* WorldContextObject, const UNiagaraDataChannel* DataChannel, TContext& SearchParams, bool bInPreviousFrame);
-
 public:
 
-	//Finds the correct source NDC data using the current AccessContext.
-	UE_API bool BeginRead(const UObject* WorldContextObject, const UNiagaraDataChannel* DataChannel, FNDCAccessContextInst& AccessContext, bool bInPreviousFrame);
-	
-	//Legacy version taking an old FNiagaraDataChannelSearchParameters struct. Works for legacy code but will not support newer NDC types or features.	
-	UE_API bool BeginRead(const UObject* WorldContextObject, const UNiagaraDataChannel* DataChannel, FNiagaraDataChannelSearchParameters SearchParams, bool bInPreviousFrame);
+	//Finds the correct source NDC data using the current SearchParameters.
+	UE_API bool BeginRead(const UObject* WorldContextObject, const UNiagaraDataChannel* DataChannel, const FNiagaraDataChannelSearchParameters& SearchParams, bool bInPreviousFrame);
 
 	//Finalizes the read and free/reset any data.
 	UE_API void EndRead();
diff --git a/Engine/Plugins/FX/Niagara/Source/Niagara/Public/NiagaraDataChannelAccessorImpl.h b/Engine/Plugins/FX/Niagara/Source/Niagara/Public/NiagaraDataChannelAccessorImpl.h
deleted file mode 100644
diff --git a/Engine/Plugins/FX/Niagara/Source/Niagara/Public/NiagaraDataChannelHandler.h b/Engine/Plugins/FX/Niagara/Source/Niagara/Public/NiagaraDataChannelHandler.h
--- a/Engine/Plugins/FX/Niagara/Source/Niagara/Public/NiagaraDataChannelHandler.h
+++ b/Engine/Plugins/FX/Niagara/Source/Niagara/Public/NiagaraDataChannelHandler.h
@@ -17,7 +17,6 @@ Some more complex handlers may want to divide up the scene in various different
 
 class FNiagaraGpuComputeDispatchInterface;
 struct FNiagaraDataChannelData;
-class FNiagaraDataChannelManager;
 
 UCLASS(abstract, BlueprintType, MinimalAPI)
 class UNiagaraDataChannelHandler : public UObject
@@ -35,22 +34,17 @@ public:
 
 	NIAGARA_API virtual void Cleanup();
 
-	NIAGARA_API virtual void BeginFrame(float DeltaTime);
+	NIAGARA_API virtual void BeginFrame(float DeltaTime, FNiagaraWorldManager* OwningWorld);
 
-	NIAGARA_API virtual void EndFrame(float DeltaTime);
+	NIAGARA_API virtual void EndFrame(float DeltaTime, FNiagaraWorldManager* OwningWorld);
 
-	NIAGARA_API virtual void Tick(float DeltaTime, ETickingGroup TickGroup);
+	NIAGARA_API virtual void Tick(float DeltaTime, ETickingGroup TickGroup, FNiagaraWorldManager* OwningWorld);
 
 	/** 
-	Legacy FindData that supports old FNiagaraDataChannelSearchParameters path.
-	*/
-	NIAGARA_API virtual FNiagaraDataChannelDataPtr FindData(FNiagaraDataChannelSearchParameters SearchParams, ENiagaraResourceAccess AccessType);
-
-	/** 
-	Finds the correct internal data for this data channel and the given access context. For example in some cases this may require a search of several elements of data that correspond to different spacial regions.
+	Finds the correct internal data for this data channel and the given search parameters. For example in some cases this may require a search of several elements of data that correspond to different spacial regions.
 	This shared ptr provides access to Game level data, CPU simulation data and a render thread proxy that can be given to the RT and provides access to GPUSimulaiton data.
 	*/
-	NIAGARA_API virtual FNiagaraDataChannelDataPtr FindData(FNDCAccessContextInst& AccessContext, ENiagaraResourceAccess AccessType);
+	virtual FNiagaraDataChannelDataPtr FindData(FNiagaraDataChannelSearchParameters SearchParams, ENiagaraResourceAccess AccessType)  PURE_VIRTUAL(UNiagaraDataChannelHandler::FindData, return nullptr;);
 
 	const UNiagaraDataChannel* GetDataChannel()const
 	{
@@ -65,15 +59,10 @@ public:
 	UFUNCTION(BlueprintCallable, Category = "Data Channel")
 	NIAGARA_API UNiagaraDataChannelReader* GetDataChannelReader();
 
-	/** The provided delegate will be called whenever new entries are added to the relevant data channel. This means the delegate can be called multiple times per tick.
-	 * This is only relevant for data published to the game thread, so no gpu data or data that's only visible to niagara systems. */
-	UFUNCTION(BlueprintCallable, Category = "Data Channel", meta=(DisplayName="Subscribe To Data Channel Updates (Legacy)"))
-	NIAGARA_API void SubscribeToDataChannelUpdates(FOnNewNiagaraDataChannelPublish UpdateDelegate, FNiagaraDataChannelSearchParameters SearchParams, int32& UnsubscribeToken);
-	
 	/** The provided delegate will be called whenever new entries are added to the relevant data channel. This means the delegate can be called multiple times per tick.
 	 * This is only relevant for data published to the game thread, so no gpu data or data that's only visible to niagara systems. */
 	UFUNCTION(BlueprintCallable, Category = "Data Channel")
-	NIAGARA_API void SubscribeToDataChannelUpdates_WithContext(FOnNewNiagaraDataChannelPublish UpdateDelegate, FNDCAccessContextInst& AccessContext, int32& UnsubscribeToken);
+	NIAGARA_API void SubscribeToDataChannelUpdates(FOnNewNiagaraDataChannelPublish UpdateDelegate, FNiagaraDataChannelSearchParameters SearchParams, int32& UnsubscribeToken);
 
 	UFUNCTION(BlueprintCallable, Category = "Data Channel")
 	NIAGARA_API void UnsubscribeFromDataChannelUpdates(const int32& UnsubscribeToken);
@@ -92,16 +81,8 @@ public:
 	template<typename TAction>
 	void ForEachNDCData(TAction Func);
 
-	/** Returns a pre allocated access context that you can use for accessing this data channel. Only safe to use from the game thread and it is not safe to hold onto this context after your access. */
-	[[nodiscard]] inline FNDCAccessContextInst& GetTransientAccessContext()const;
-
-	void SetOwningWorld(FNiagaraWorldManager* InOwner){ OwningWorld = InOwner; }
 protected:
 
-	FNiagaraWorldManager* OwningWorld = nullptr;
-
-	NIAGARA_API void VerifyAccessContext(FNDCAccessContextInst& AccessContext);
-
 	UPROPERTY()
 	TWeakObjectPtr<const UNiagaraDataChannel> DataChannel;
 
@@ -127,26 +108,14 @@ private:
 	{
 		FOnNewNiagaraDataChannelPublish OnPublishDelegate;
 		FNiagaraDataChannelSearchParameters SearchParams;
-		FNDCAccessContextInst AccessContext;
 		int32 SubscriptionToken = -1;
 	};
 	TArray<FChannelSubscription> ChannelSubscriptions;
 
 	// used to generate the unsubscribe tokens that are stored in ChannelSubscriptions
 	int32 SubscriberTokens = 0;
-
-	/** A transient access context instance that calling code can use for accessing this Data Channel. GT only*/
-	mutable FNDCAccessContextInst TransientAccessContext;
-	TNDCAccessContextType AccessContextType;
 };
 
-FNDCAccessContextInst& UNiagaraDataChannelHandler::GetTransientAccessContext()const
-{
-	check(IsInGameThread());
-	TransientAccessContext.Init(AccessContextType);//We need to init (or reset) ready for use.
-	return TransientAccessContext;
-}
-
 template<typename TAction>
 void UNiagaraDataChannelHandler::ForEachNDCData(TAction Func)
 {
@@ -158,4 +127,4 @@ void UNiagaraDataChannelHandler::ForEachNDCData(TAction Func)
 			Func(DataPtr);
 		}
 	}
-}
\ No newline at end of file
+}
diff --git a/Engine/Plugins/FX/Niagara/Source/Niagara/Public/NiagaraDataChannelPublic.h b/Engine/Plugins/FX/Niagara/Source/Niagara/Public/NiagaraDataChannelPublic.h
--- a/Engine/Plugins/FX/Niagara/Source/Niagara/Public/NiagaraDataChannelPublic.h
+++ b/Engine/Plugins/FX/Niagara/Source/Niagara/Public/NiagaraDataChannelPublic.h
@@ -4,7 +4,6 @@
 
 #include "Engine/EngineBaseTypes.h"
 #include "NiagaraCommon.h"
-#include "StructUtils/InstancedStruct.h"
 #include "NiagaraDataChannelPublic.generated.h"
 
 struct FNiagaraDataChannelPublishRequest;
@@ -39,6 +38,36 @@ struct FNiagaraDataChannelUpdateContext
 
 DECLARE_DYNAMIC_DELEGATE_OneParam(FOnNewNiagaraDataChannelPublish, const FNiagaraDataChannelUpdateContext&, NewDataContext);
 
+/** Niagara Data Channels are a system for communication between Niagara Systems and with game code/Blueprint.
+
+Data channel assets define the payload as well as some transfer settings.
+Niagara Systems can read from and write to data channels via data interfaces.
+Blueprint and C++ code can also read from and write to data channels using its API functions.
+
+ */
+UCLASS(BlueprintType, DisplayName = "Niagara Data Channel", MinimalAPI)
+class UNiagaraDataChannelAsset : public UObject
+{
+	GENERATED_BODY()
+
+	UPROPERTY(EditAnywhere, Category = DataChannel, Instanced)
+	TObjectPtr<UNiagaraDataChannel> DataChannel;
+	
+	#if WITH_EDITORONLY_DATA
+	/** When changing data channel types we cache the old channel and attempt to copy over any common properties from one to the other. */
+	UPROPERTY(Transient)
+	TObjectPtr<UNiagaraDataChannel> CachedPreChangeDataChannel;
+	#endif
+
+public:
+
+#if WITH_EDITOR
+	NIAGARA_API virtual void PreEditChange(class FProperty* PropertyAboutToChange)override;
+	NIAGARA_API virtual void PostEditChangeProperty(FPropertyChangedEvent& PropertyChangedEvent)override;
+#endif
+
+	UNiagaraDataChannel* Get() const { return DataChannel; }
+};
 
 USTRUCT(BlueprintType)
 struct FNiagaraDataChannelVariable : public FNiagaraVariableBase
@@ -74,26 +103,31 @@ Minimal set of types and declares required for external users of Niagara Data Ch
 */
 
 /**
-Legacy search parameters to support legacy code.
+Parameters used when retrieving a specific set of Data Channel Data to read or write.
+Many Data Channel types will have multiple internal sets of data and these parameters control which the Channel should return to users for access.
+An example of this would be the Islands Data Channel type which will subdivide the world and have a different set of data for each sub division.
+It will return to users the correct data for their location based on these parameters.
 */
 USTRUCT(BlueprintType)
 struct FNiagaraDataChannelSearchParameters
 {
 	GENERATED_BODY()
-
-	FNiagaraDataChannelSearchParameters() = default;
+	
+	FNiagaraDataChannelSearchParameters()
+	: bOverrideLocation(false)
+	{
+	}
 
 	FNiagaraDataChannelSearchParameters(USceneComponent* Owner)
-		: OwningComponent(Owner)
-		, Location(FVector::ZeroVector)
-		, bOverrideLocation(false)
+	: OwningComponent(Owner)
+	, bOverrideLocation(false)
 	{
 	}
 
 	FNiagaraDataChannelSearchParameters(USceneComponent* Owner, FVector LocationOverride)
-		: OwningComponent(Owner)
-		, Location(LocationOverride)
-		, bOverrideLocation(true)
+	: OwningComponent(Owner)
+	, Location(LocationOverride)
+	, bOverrideLocation(true)
 	{
 	}
 
@@ -103,9 +137,9 @@ struct FNiagaraDataChannelSearchParameters
 		, bOverrideLocation(true)
 	{
 	}
-
-	[[nodiscard]] NIAGARA_API FVector GetLocation()const;
-	[[nodiscard]] inline USceneComponent* GetOwner()const { return OwningComponent; }
+	
+	NIAGARA_API FVector GetLocation()const;
+	USceneComponent* GetOwner()const { return OwningComponent; }
 
 	/** In cases where there is an owning component such as an object spawning from itself etc, then we pass that component in. Some handlers may only use it's location but others may make use of more data. */
 	UPROPERTY(EditAnywhere, BlueprintReadWrite, Category = "Parameters")
@@ -117,212 +151,9 @@ struct FNiagaraDataChannelSearchParameters
 
 	/** If true, even if an owning component is set, the data channel should use the Location value rather than the component location. If this is false, the NDC will get any location needed from the owning component. */
 	UPROPERTY(EditAnywhere, BlueprintReadWrite, Category = "Parameters")
-	uint32 bOverrideLocation : 1 = false;
-};
-
-
-/**
-Base class for context used when retrieving a specific set of Data Channel Data to read or write.
-Many Data Channel types will have multiple internal sets of data and these parameters control which the Channel should return to users for access.
-An example of this would be the Islands Data Channel type which will subdivide the world and have a different set of data for each sub division.
-It will return to users the correct data for their location based on these parameters.
-NDC Types must specify which of these access context types they require.
-The base struct contains only data needed for internal systems to function. 
-Beyond this you are free to add anything else you need.
-Some data can be intended as input to the search/internal routing or data creation. This should be marked with "NDCAccessContextInput".
-Other data can be intended as output to store derived or transient state and pass back to the user. This should be marked with "NDCAccessContextOutput".
-*/
-USTRUCT(BlueprintType, meta=(Hidden))
-struct FNDCAccessContextBase
-{
-	GENERATED_BODY()
-
-public:
-
-	FNDCAccessContextBase() = default;
-	[[nodiscard]] explicit FNDCAccessContextBase(USceneComponent* Owner)
-		: OwningComponent(Owner)
-		, bIsAutoLinkingSystem(false)
-	{
-	}
-
-	[[nodiscard]] inline USceneComponent* GetOwner()const { return OwningComponent; }
-
-	UPROPERTY(BlueprintReadWrite, Transient, Category = "Input", meta=(NDCAccessContextInput))
-	TObjectPtr<USceneComponent> OwningComponent = nullptr;
-	
-	/** 
-	If true, this access is for a Niagara System that wishes to "Auto link" i.e. not do a proper lookup, just use the currently initializing map entry or the spawned component lookup.
-	This is used by NDC Data Interfaces in systems spawned by NDCs to make linking back to the originating NDC entry easier.
-	*/
-	UPROPERTY(transient)
-	uint32 bIsAutoLinkingSystem : 1 = false;
-};
-
-/** Utility class to wrap a weak pointer to a Niagara Component allowing FNDCAccessContext to hold an array of weak spawned systems. */
-USTRUCT(BlueprintType)
-struct FNDCSpawnedSystemRef
-{
-	GENERATED_BODY()
-	
-	FNDCSpawnedSystemRef() = default;
-	[[nodiscard]] FNDCSpawnedSystemRef(UNiagaraComponent* Component);
-
-	UPROPERTY(BlueprintReadWrite, Transient, Category = "Output")
-	TWeakObjectPtr<UNiagaraComponent> SpawnedSystem;
-
-	[[nodiscard]] inline UNiagaraComponent* Get()const{ return SpawnedSystem.Get(); }
-};
-
-/**
-Basic AccessContext type providing a little more utility than the base.
-If you require more data than this you can subclass this or replace entirely and subclass FNDCAccessContextBase.
-*/
-USTRUCT(BlueprintType)
-struct FNDCAccessContext : public FNDCAccessContextBase
-{
-	GENERATED_BODY()
-
-	FNDCAccessContext() = default;
-	[[nodiscard]] explicit FNDCAccessContext(USceneComponent* Owner)
-		: FNDCAccessContextBase(Owner)
-	{
-	}
-
-	[[nodiscard]] explicit FNDCAccessContext(USceneComponent* Owner, FVector LocationOverride)
-		: FNDCAccessContextBase(Owner)
-		, Location(LocationOverride)
-		, bOverrideLocation(true)
-	{
-	}
-
-	[[nodiscard]] explicit FNDCAccessContext(FVector InLocation)
-		: FNDCAccessContextBase(nullptr)
-		, Location(InLocation)
-		, bOverrideLocation(true)
-	{
-	}
-
-	[[nodiscard]] NIAGARA_API FVector GetLocation()const;
-
-	/** In cases where there is no owning component for data being read or written to a data channel, we simply pass in a location. We can also use this when bOverrideLocaiton is set. */
-	UPROPERTY(BlueprintReadWrite, Transient, Category = "Input", meta=(NDCAccessContextInput))
-	FVector Location = FVector::ZeroVector;
-
-	/** If true, even if an owning component is set, the data channel should use the Location value rather than the component location. If this is false, the NDC will get any location needed from the owning component. */
-	UPROPERTY(BlueprintReadWrite, Transient, Category = "Input", meta=(NDCAccessContextInput))
-	uint32 bOverrideLocation : 1 = false;
-
-	UPROPERTY(EditAnywhere, BlueprintReadWrite, Category = "Input", meta = (NDCAccessContextInput, AllowedClasses = "/Script/Niagara.NiagaraSystem,/Script/Niagara.NiagaraSystemCollection"))
-	TObjectPtr<UObject> SystemToSpawn;
-
-	/** The Niagara System spawned in response to this access, if any. Allows callers to set parameters etc on this system. It is not safe to store and use this component reference later. It can be reclaimed by the system at any time. */
-	UPROPERTY(BlueprintReadWrite, Transient, Category = "Output", meta = (NDCAccessContextOutput, ShowOnlyInnerProperties))
-	TArray<FNDCSpawnedSystemRef> SpawnedSystems;
-};
-
-/** Special legacy access context type. Used only to facilitate legacy NDCs to be usable by the new API. */
-USTRUCT(BlueprintType)
-struct FNDCAccessContextLegacy : public FNDCAccessContextBase
-{
-	GENERATED_BODY()
-
-	FNDCAccessContextLegacy() = default;
-	[[nodiscard]] explicit FNDCAccessContextLegacy(USceneComponent* Owner)
-		: FNDCAccessContextBase(Owner)
-	{
-	}
-
-	[[nodiscard]] explicit FNDCAccessContextLegacy(USceneComponent* Owner, FVector LocationOverride)
-		: FNDCAccessContextBase(Owner)
-		, Location(LocationOverride)
-		, bOverrideLocation(true)
-	{
-	}
-
-	[[nodiscard]] explicit FNDCAccessContextLegacy(FVector InLocation)
-		: FNDCAccessContextBase(nullptr)
-		, Location(InLocation)
-		, bOverrideLocation(true)
-	{
-	}
-
-	[[nodiscard]] NIAGARA_API FVector GetLocation()const;
-
-	/** In cases where there is no owning component for data being read or written to a data channel, we simply pass in a location. We can also use this when bOverrideLocaiton is set. */
-	UPROPERTY(BlueprintReadWrite, Transient, Category = "Input", meta = (NDCAccessContextInput))
-	FVector Location = FVector::ZeroVector;
-
-	/** If true, even if an owning component is set, the data channel should use the Location value rather than the component location. If this is false, the NDC will get any location needed from the owning component. */
-	UPROPERTY(BlueprintReadWrite, Transient, Category = "Input", meta = (NDCAccessContextInput))
-	uint32 bOverrideLocation : 1 = false;
-};
-
-//Describes a NDC Access Context Type.
-//Wrapper around a UScriptStruct* that checks for a valid type.
-struct TNDCAccessContextType
-{
-	TNDCAccessContextType() = default;
-	[[nodiscard]] explicit TNDCAccessContextType(UScriptStruct* InStruct)
-		: AccessContextStruct(InStruct->IsChildOf(FNDCAccessContextBase::StaticStruct()) ? InStruct : nullptr)
-	{
-		check(AccessContextStruct);
-	}
-
-	[[nodiscard]] const UScriptStruct* Get()const { return AccessContextStruct; }
-
-	//Returns the underlying struct as non-const. Only use to inter operate with code that stores structs without constness.
-	[[nodiscard]] UScriptStruct* Get_MutableInteropOnly()const { return const_cast<UScriptStruct*>(AccessContextStruct); }
-
-private:
-	const UScriptStruct* AccessContextStruct = nullptr;
-};
-
-/**
-Context object for Niagara and Game code to access a Niagara Data Channel.
-Can be customized per Data Channel to provide specific control for each Data Channel type.
-Can include input data from accessing code such as a location, game play tags etc that can influence internal data routing/partitions.
-Can also include output data allowing the NDC to pass information back to the accessing code.
-*/
-USTRUCT(BlueprintType)
-struct FNDCAccessContextInst
-{
-	GENERATED_BODY()
-
-	FNDCAccessContextInst() = default;
-	[[nodiscard]] FNDCAccessContextInst(TNDCAccessContextType AccessContextType)
-	{
-		AccessContext.InitializeAs(AccessContextType.Get());
-	}
-
-	inline bool Init(TNDCAccessContextType ParamsType)
-	{
-		AccessContext.InitializeAs(ParamsType.Get());
-		return AccessContext.IsValid();
-	}
-	
-	[[nodiscard]] bool operator==(const FNDCAccessContextInst&) const = default;
-
-	inline void Reset()
-	{
-		AccessContext.Reset();
-	}
-
-	[[nodiscard]] inline bool IsValid()const { return AccessContext.IsValid(); }
-
-	template<typename T = FNDCAccessContextBase>
-	[[nodiscard]] inline T& GetChecked() { return AccessContext.GetMutable<T>(); }
-
-	template<typename T = FNDCAccessContextBase>
-	[[nodiscard]] inline T* Get() { return AccessContext.GetMutablePtr<T>();	}
-
-	[[nodiscard]] inline const UScriptStruct* GetScriptStruct()const { return AccessContext.GetScriptStruct(); }
-
-	UPROPERTY(EditAnywhere, BlueprintReadWrite, Category = "Access Context", meta = (ShowOnlyInnerProperties, BaseStruct = "/Script/Niagara.NDCAccessContextBase", ExcludeBaseStruct))
-	FInstancedStruct AccessContext;
+	uint32 bOverrideLocation : 1;
 };
 
-
 USTRUCT()
 struct FNiagaraDataChannelGameDataLayout
 {
@@ -548,56 +379,3 @@ public:
 		return false;
 	}
 };
-
-/** Niagara Data Channels are a system for communication between Niagara Systems and with game code/Blueprint.
-
-Data channel assets define the payload as well as some transfer settings.
-Niagara Systems can read from and write to data channels via data interfaces.
-Blueprint and C++ code can also read from and write to data channels using its API functions.
-
- */
-UCLASS(BlueprintType, DisplayName = "Niagara Data Channel", MinimalAPI)
-class UNiagaraDataChannelAsset : public UObject
-{
-	GENERATED_BODY()
-
-	UPROPERTY(EditAnywhere, Category = DataChannel, Instanced)
-	TObjectPtr<UNiagaraDataChannel> DataChannel;
-	
-	#if WITH_EDITORONLY_DATA
-	/** When changing data channel types we cache the old channel and attempt to copy over any common properties from one to the other. */
-	UPROPERTY(Transient)
-	TObjectPtr<UNiagaraDataChannel> CachedPreChangeDataChannel;
-	#endif
-
-public:
-
-#if WITH_EDITOR
-	NIAGARA_API virtual void PreEditChange(class FProperty* PropertyAboutToChange)override;
-	NIAGARA_API virtual void PostEditChangeProperty(FPropertyChangedEvent& PropertyChangedEvent)override;
-#endif
-
-	[[nodiscard]] UNiagaraDataChannel* Get() const { return DataChannel; }
-};
-
-/** A reference to a Niagara Data Channel with a pre-typed access context for convenience. */
-USTRUCT(BlueprintType)
-struct NIAGARA_API FNiagaraDataChannelReference
-{
-	GENERATED_BODY()
-
-	[[nodiscard]] inline UNiagaraDataChannelAsset* GetDataChannel()const { return DataChannel; }
-	[[nodiscard]] FNDCAccessContextInst& GetAccessContext()const;
-	void SetDataChannel(UNiagaraDataChannelAsset* InDataChannel);	
-	void OnValueChanged();
-
-	UPROPERTY(EditAnywhere, BlueprintReadWrite, Category = DataChannel)
-	TObjectPtr<UNiagaraDataChannelAsset> DataChannel;
-
-	UPROPERTY(EditAnywhere, BlueprintReadWrite, Category = DataChannel, meta = (ShowOnlyInnerProperties, ShowNDCAccessContextInput))
-	mutable FNDCAccessContextInst AccessContext;
-
-private:
-
-	void InitAccessContext(bool bTransferExisting = true)const;
-};
\ No newline at end of file
diff --git a/Engine/Plugins/FX/Niagara/Source/Niagara/Public/NiagaraDataChannel_GameplayBurst.h b/Engine/Plugins/FX/Niagara/Source/Niagara/Public/NiagaraDataChannel_GameplayBurst.h
deleted file mode 100644
diff --git a/Engine/Plugins/FX/Niagara/Source/Niagara/Public/NiagaraDataChannel_Global.h b/Engine/Plugins/FX/Niagara/Source/Niagara/Public/NiagaraDataChannel_Global.h
--- a/Engine/Plugins/FX/Niagara/Source/Niagara/Public/NiagaraDataChannel_Global.h
+++ b/Engine/Plugins/FX/Niagara/Source/Niagara/Public/NiagaraDataChannel_Global.h
@@ -29,9 +29,9 @@ class UNiagaraDataChannelHandler_Global : public UNiagaraDataChannelHandler
 
 	NIAGARA_API virtual void Init(const UNiagaraDataChannel* InChannel) override;
 	NIAGARA_API virtual void Cleanup() override;
-	NIAGARA_API virtual void BeginFrame(float DeltaTime)override;
-	NIAGARA_API virtual void EndFrame(float DeltaTime)override;
-	NIAGARA_API virtual void Tick(float DeltaTime, ETickingGroup TickGroup) override;
+	NIAGARA_API virtual void BeginFrame(float DeltaTime, FNiagaraWorldManager* OwningWorld)override;
+	NIAGARA_API virtual void EndFrame(float DeltaTime, FNiagaraWorldManager* OwningWorld)override;
+	NIAGARA_API virtual void Tick(float DeltaTime, ETickingGroup TickGroup, FNiagaraWorldManager* OwningWorld) override;
 	NIAGARA_API virtual FNiagaraDataChannelDataPtr FindData(FNiagaraDataChannelSearchParameters SearchParams, ENiagaraResourceAccess AccessType) override;
 };
 
diff --git a/Engine/Plugins/FX/Niagara/Source/Niagara/Public/NiagaraDataChannel_Islands.h b/Engine/Plugins/FX/Niagara/Source/Niagara/Public/NiagaraDataChannel_Islands.h
--- a/Engine/Plugins/FX/Niagara/Source/Niagara/Public/NiagaraDataChannel_Islands.h
+++ b/Engine/Plugins/FX/Niagara/Source/Niagara/Public/NiagaraDataChannel_Islands.h
@@ -193,9 +193,9 @@ class UNiagaraDataChannelHandler_Islands : public UNiagaraDataChannelHandler
 
 	NIAGARA_API virtual void Init(const UNiagaraDataChannel* InChannel) override;
 	NIAGARA_API virtual void Cleanup() override;
-	NIAGARA_API virtual void BeginFrame(float DeltaTime)override;
-	NIAGARA_API virtual void EndFrame(float DeltaTime)override;
-	NIAGARA_API virtual void Tick(float DeltaTime, ETickingGroup TickGroup) override;
+	NIAGARA_API virtual void BeginFrame(float DeltaTime, FNiagaraWorldManager* OwningWorld)override;
+	NIAGARA_API virtual void EndFrame(float DeltaTime, FNiagaraWorldManager* OwningWorld)override;
+	NIAGARA_API virtual void Tick(float DeltaTime, ETickingGroup TickGroup, FNiagaraWorldManager* OwningWorld) override;
 	NIAGARA_API virtual FNiagaraDataChannelDataPtr FindData(FNiagaraDataChannelSearchParameters SearchParams, ENiagaraResourceAccess AccessType) override;
 
 protected:
diff --git a/Engine/Plugins/FX/Niagara/Source/Niagara/Public/NiagaraDataChannel_Map.h b/Engine/Plugins/FX/Niagara/Source/Niagara/Public/NiagaraDataChannel_Map.h
deleted file mode 100644
diff --git a/Engine/Plugins/FX/Niagara/Source/Niagara/Public/NiagaraSystemCollection.h b/Engine/Plugins/FX/Niagara/Source/Niagara/Public/NiagaraSystemCollection.h
deleted file mode 100644
diff --git a/Engine/Plugins/FX/Niagara/Source/Niagara/Public/NiagaraWorldManager.h b/Engine/Plugins/FX/Niagara/Source/Niagara/Public/NiagaraWorldManager.h
--- a/Engine/Plugins/FX/Niagara/Source/Niagara/Public/NiagaraWorldManager.h
+++ b/Engine/Plugins/FX/Niagara/Source/Niagara/Public/NiagaraWorldManager.h
@@ -197,7 +197,7 @@ public:
 	// Dump details about what's inside the world manager
 	void DumpDetails(FOutputDevice& Ar, bool bDetailed);
 	
-	NIAGARA_API UWorld* GetWorld();
+	UWorld* GetWorld();
 	inline UWorld* GetWorld()const { return World; }
 
 	//Various helper functions for scalability culling.
diff --git a/Engine/Plugins/FX/Niagara/Source/NiagaraBlueprintNodes/Internal/K2Node_DataChannelBase.h b/Engine/Plugins/FX/Niagara/Source/NiagaraBlueprintNodes/Internal/K2Node_DataChannelBase.h
--- a/Engine/Plugins/FX/Niagara/Source/NiagaraBlueprintNodes/Internal/K2Node_DataChannelBase.h
+++ b/Engine/Plugins/FX/Niagara/Source/NiagaraBlueprintNodes/Internal/K2Node_DataChannelBase.h
@@ -26,13 +26,11 @@ public:
 	// //~ End UEdGraphNode Interface.
 
 	//~ Begin K2Node Interface
-	virtual void GetMenuActions(FBlueprintActionDatabaseRegistrar& InActionRegistrar) const override;
+	//virtual void GetMenuActions(FBlueprintActionDatabaseRegistrar& InActionRegistrar) const override;
 	virtual FText GetMenuCategory() const override;
 	virtual ERedirectType DoPinsMatchForReconstruction(const UEdGraphPin* NewPin, int32 NewPinIndex, const UEdGraphPin* OldPin, int32 OldPinIndex) const override;
 	virtual void PreloadRequiredAssets() override;
 	virtual bool ShouldShowNodeProperties() const override;
-	virtual void ExpandNode(class FKismetCompilerContext& CompilerContext, UEdGraph* SourceGraph) override;
-	virtual bool IsConnectionDisallowed(const UEdGraphPin* MyPin, const UEdGraphPin* OtherPin, FString& OutReason)const override;
 	//~ End K2Node Interface
 
 	UPROPERTY()
@@ -41,20 +39,12 @@ public:
 protected:
 
 	UEdGraphPin* GetChannelSelectorPin() const;
-	UEdGraphPin* GetAccessContextPin(EEdGraphPinDirection Direction) const;
 
 #if WITH_EDITORONLY_DATA
 	UPROPERTY()
 	FGuid DataChannelVersion;
 #endif
 
-	/**	Some nodes require and explicit NDC reference to generate derived pins. We ban pin connections for these nodes but allow them by default. */
-	[[nodiscard]] virtual bool SupportsDynamicDataChannel()const { return true; }
-
-	virtual void AddNDCDerivedPins();
-	virtual void RemoveNDCDerivedPins();
-	TArray<UEdGraphPin*> NDCDerivedGeneratedPins;
-
 private:
 
 	UPROPERTY()
diff --git a/Engine/Plugins/FX/Niagara/Source/NiagaraBlueprintNodes/Internal/K2Node_DataChannel_WithContext.h b/Engine/Plugins/FX/Niagara/Source/NiagaraBlueprintNodes/Internal/K2Node_DataChannel_WithContext.h
deleted file mode 100644
diff --git a/Engine/Plugins/FX/Niagara/Source/NiagaraBlueprintNodes/Internal/K2Node_ReadDataChannel.h b/Engine/Plugins/FX/Niagara/Source/NiagaraBlueprintNodes/Internal/K2Node_ReadDataChannel.h
--- a/Engine/Plugins/FX/Niagara/Source/NiagaraBlueprintNodes/Internal/K2Node_ReadDataChannel.h
+++ b/Engine/Plugins/FX/Niagara/Source/NiagaraBlueprintNodes/Internal/K2Node_ReadDataChannel.h
@@ -15,16 +15,16 @@ class UK2Node_ReadDataChannel : public UK2Node_DataChannelBase
 public:
 	UK2Node_ReadDataChannel();
 
+	// //~ Begin UEdGraphNode Interface.
+	virtual void AllocateDefaultPins() override;
+	// //~ End UEdGraphNode Interface.
+
 	//~ Begin K2Node Interface
+	virtual void GetMenuActions(FBlueprintActionDatabaseRegistrar& InActionRegistrar) const override;
 	virtual void ExpandNode(class FKismetCompilerContext& CompilerContext, UEdGraph* SourceGraph) override;
 	//~ End K2Node Interface
 
-protected:
-
-	virtual void AddNDCDerivedPins()override;
-
 private:
 	UFunction* GetReadFunctionForType(const FNiagaraTypeDefinition& TypeDef);
 	UEdGraphPin* GetVarPin(const FName& Name) const;
-	virtual FName GetReaderFunctionName()const { return GET_FUNCTION_NAME_CHECKED(UNiagaraDataChannelLibrary, ReadFromNiagaraDataChannel); }
 };
diff --git a/Engine/Plugins/FX/Niagara/Source/NiagaraBlueprintNodes/Internal/K2Node_WriteDataChannel.h b/Engine/Plugins/FX/Niagara/Source/NiagaraBlueprintNodes/Internal/K2Node_WriteDataChannel.h
--- a/Engine/Plugins/FX/Niagara/Source/NiagaraBlueprintNodes/Internal/K2Node_WriteDataChannel.h
+++ b/Engine/Plugins/FX/Niagara/Source/NiagaraBlueprintNodes/Internal/K2Node_WriteDataChannel.h
@@ -15,15 +15,15 @@ class UK2Node_WriteDataChannel : public UK2Node_DataChannelBase
 public:
 	UK2Node_WriteDataChannel();
 
+	// //~ Begin UEdGraphNode Interface.
+	virtual void AllocateDefaultPins() override;
+	// //~ End UEdGraphNode Interface.
+
 	//~ Begin K2Node Interface
+	virtual void GetMenuActions(FBlueprintActionDatabaseRegistrar& InActionRegistrar) const override;
 	virtual void ExpandNode(class FKismetCompilerContext& CompilerContext, UEdGraph* SourceGraph) override;
 	//~ End K2Node Interface
 
-protected:
-
-	virtual void AddNDCDerivedPins()override;
-
+private:
 	UFunction* GetWriteFunctionForType(const FNiagaraTypeDefinition& TypeDef);
-
-	[[nodiscard]] virtual FName GetWriterFunctionName()const{ return GET_FUNCTION_NAME_CHECKED(UNiagaraDataChannelLibrary, WriteToNiagaraDataChannel); }
 };
diff --git a/Engine/Plugins/FX/Niagara/Source/NiagaraBlueprintNodes/NiagaraBlueprintNodes.build.cs b/Engine/Plugins/FX/Niagara/Source/NiagaraBlueprintNodes/NiagaraBlueprintNodes.build.cs
--- a/Engine/Plugins/FX/Niagara/Source/NiagaraBlueprintNodes/NiagaraBlueprintNodes.build.cs
+++ b/Engine/Plugins/FX/Niagara/Source/NiagaraBlueprintNodes/NiagaraBlueprintNodes.build.cs
@@ -28,11 +28,6 @@ public class NiagaraBlueprintNodes : ModuleRules
 				"Niagara",
 				"Kismet",
 				"KismetCompiler",
-				"BlueprintEditorLibrary",
-				"StructUtilsEditor",
-				"SlateCore",
-				"Slate",
-				"PropertyEditor",
 			}
 		);
 	}
diff --git a/Engine/Plugins/FX/Niagara/Source/NiagaraBlueprintNodes/Private/K2Node_DataChannelBase.cpp b/Engine/Plugins/FX/Niagara/Source/NiagaraBlueprintNodes/Private/K2Node_DataChannelBase.cpp
--- a/Engine/Plugins/FX/Niagara/Source/NiagaraBlueprintNodes/Private/K2Node_DataChannelBase.cpp
+++ b/Engine/Plugins/FX/Niagara/Source/NiagaraBlueprintNodes/Private/K2Node_DataChannelBase.cpp
@@ -1,26 +1,11 @@
 ﻿// Copyright Epic Games, Inc. All Rights Reserved.
 
+
 #include "K2Node_DataChannelBase.h"
 #include "EdGraphSchema_K2.h"
 
 #include "NiagaraDataChannelPublic.h"
 
-#include "BlueprintActionDatabaseRegistrar.h"
-#include "BlueprintNodeSpawner.h"
-#include "K2Node_ExecutionSequence.h"
-#include "K2Node_Self.h"
-#include "K2Node_MakeStruct.h"
-#include "K2Node_BreakStruct.h"
-#include "KismetCompiler.h"
-#include "NiagaraBlueprintUtil.h"
-#include "NiagaraDataChannel.h"
-#include "NiagaraDataChannelAccessor.h"
-#include "Kismet/KismetSystemLibrary.h"
-#include "Kismet/BlueprintInstancedStructLibrary.h"
-#include "Kismet2/BlueprintEditorUtils.h"
-
-#include "K2Node_DataChannel_WithContext.h"
-
 #include UE_INLINE_GENERATED_CPP_BY_NAME(K2Node_DataChannelBase)
 
 #define LOCTEXT_NAMESPACE "K2Node_DataChannelBase"
@@ -50,92 +35,15 @@ void UK2Node_DataChannelBase::PostLoad()
 void UK2Node_DataChannelBase::AllocateDefaultPins()
 {
 	Super::AllocateDefaultPins();
-	if(!SupportsDynamicDataChannel())
-	{
-		//Prevent connections to our channel pin if it does not support a dynamic channel.
-		UEdGraphPin* ChannelPin = FindPinChecked(TEXT("Channel"));
-		ChannelPin->bNotConnectable = true;
-		ChannelPin->bDefaultValueIsReadOnly = false;
-	}
-	
-	PreloadObject(DataChannel);
 
 #if WITH_EDITORONLY_DATA
 	if (HasValidDataChannel())
 	{
 		DataChannelVersion = GetDataChannel()->GetVersion();
-		AddNDCDerivedPins();
 	}
 #endif
 }
 
-void UK2Node_DataChannelBase::GetMenuActions(FBlueprintActionDatabaseRegistrar& InActionRegistrar) const
-{
-	const UClass* ActionKey = GetClass();
-	if (InActionRegistrar.IsOpenForRegistration(ActionKey))
-	{
-		UBlueprintNodeSpawner* NodeSpawner = UBlueprintNodeSpawner::Create(GetClass());
-		InActionRegistrar.AddBlueprintAction(ActionKey, NodeSpawner);
-	}
-}
-
-
-void UK2Node_DataChannelBase::ExpandNode(class FKismetCompilerContext& CompilerContext, UEdGraph* SourceGraph)
-{
-	if (!HasValidDataChannel() && !SupportsDynamicDataChannel())
-	{
-		CompilerContext.MessageLog.Error(*LOCTEXT("NoValidDataChannel", "Node does not have a valid data channel - @@").ToString(), this);
-		return;
-	}
-
-	if(UNiagaraDataChannel* NDC = GetDataChannel())
-	{
-		if (DataChannelVersion != NDC->GetVersion())
-		{
-			CompilerContext.MessageLog.Error(*LOCTEXT("StaleNode", "Node is out of sync with the data channel asset, please refresh node to fix up the pins - @@").ToString(), this);
-			return;
-		}
-	}
-
-	ExpandSplitPins(CompilerContext, SourceGraph);
-}
-
-bool UK2Node_DataChannelBase::IsConnectionDisallowed(const UEdGraphPin* MyPin, const UEdGraphPin* OtherPin, FString& OutReason)const
-{
-	//TODO: 
-	// I would like to do this kind of type checking and give some useful feedback to users.
-	// However currently BP context menus add a bunch of unrelated cruft by default.
-	// This means that my OutputReason is never shown to the users.
-	// Instead they're just prompted to add a node converting from FNDCAccessContextInst to FNDCAccessContextInst.
-	// With all the options being unrelated cruft. So this would actually add to confusion etc rather than improving things.
-	// 
-	if(MyPin->PinType.PinCategory == UEdGraphSchema_K2::PC_Struct && MyPin->PinType.PinSubCategoryObject == FNDCAccessContextInst::StaticStruct())
-	{
-		if(OtherPin->PinType.PinCategory == UEdGraphSchema_K2::PC_Struct && OtherPin->PinType.PinSubCategoryObject == FNDCAccessContextInst::StaticStruct())
-		{
-			if(UK2Node_DataChannelAccessContextOperation* ContextOp = Cast<UK2Node_DataChannelAccessContextOperation>(OtherPin->GetOwningNode()))
-			{
-				if(UNiagaraDataChannel* NDC = GetDataChannel())
-				{
-					//If we have a specific data channel and we're connecting to a context op, ensure the NDC type and the context op type are compatible.
-					const UScriptStruct* NDCType = NDC->GetAccessContextType().Get();
-					if(NDCType != ContextOp->ContextStruct)
-					{
-						OutReason = TEXT("Invalid Type For Access Context Pin.");
-						return true;						
-					}
-				}
-			}
-		}
-		else
-		{
-			OutReason = TEXT("Invalid Type For Access Context Pin.");
-			return true;
-		}
-	}
-	return false;
-}
-
 void UK2Node_DataChannelBase::PinDefaultValueChanged(UEdGraphPin* Pin)
 {
 	Super::PinDefaultValueChanged(Pin);
@@ -143,24 +51,11 @@ void UK2Node_DataChannelBase::PinDefaultValueChanged(UEdGraphPin* Pin)
 #if WITH_EDITORONLY_DATA
 	if (Pin == GetChannelSelectorPin())
 	{
-		DataChannel = nullptr;
-
 		if (UNiagaraDataChannelAsset* ChannelAsset = Cast<UNiagaraDataChannelAsset>(Pin->DefaultObject))
 		{
 			DataChannel = ChannelAsset;
+			ReconstructNode();
 		}
-		
-		if (UNiagaraDataChannel* NDC = GetDataChannel())
-		{
-			DataChannelVersion = NDC->GetVersion();
-		}
-		else
-		{		
-			DataChannelVersion = FGuid();
-		}
-
-		RemoveNDCDerivedPins();
-		AddNDCDerivedPins();
 	}
 #endif
 }
@@ -172,23 +67,11 @@ void UK2Node_DataChannelBase::PinConnectionListChanged(UEdGraphPin* Pin)
 #if WITH_EDITORONLY_DATA
 	if (Pin == GetChannelSelectorPin())
 	{
-		DataChannel = nullptr;
 		if (UNiagaraDataChannelAsset* ChannelAsset = Cast<UNiagaraDataChannelAsset>(Pin->DefaultObject))
 		{
 			DataChannel = Pin->LinkedTo.Num() == 0 ? ChannelAsset : nullptr;
+			ReconstructNode();
 		}
-		
-		if (UNiagaraDataChannel* NDC = GetDataChannel())
-		{
-			DataChannelVersion = NDC->GetVersion();
-		}
-		else
-		{
-			DataChannelVersion = FGuid();
-		}
-
-		RemoveNDCDerivedPins();
-		AddNDCDerivedPins();
 	}
 #endif
 }
@@ -212,31 +95,6 @@ UK2Node::ERedirectType UK2Node_DataChannelBase::DoPinsMatchForReconstruction(con
 	return Result;
 }
 
-void UK2Node_DataChannelBase::AddNDCDerivedPins()
-{
- 	GetGraph()->NotifyNodeChanged(this);
- 
- 	UBlueprint* NodeBP = GetBlueprint();
- 	check(NodeBP);
- 	FBlueprintEditorUtils::MarkBlueprintAsStructurallyModified(NodeBP);
-}
-
-void UK2Node_DataChannelBase::RemoveNDCDerivedPins()
-{
-	for(UEdGraphPin* Pin : NDCDerivedGeneratedPins)
-	{
-		Pins.Remove(Pin);
-		DestroyPin(Pin);
-	}
-	NDCDerivedGeneratedPins.Empty();
-
- 	GetGraph()->NotifyNodeChanged(this);
- 
- 	UBlueprint* NodeBP = GetBlueprint();
- 	check(NodeBP);
- 	FBlueprintEditorUtils::MarkBlueprintAsStructurallyModified(NodeBP);
-}
-
 void UK2Node_DataChannelBase::PreloadRequiredAssets()
 {
 	Super::PreloadRequiredAssets();
@@ -258,9 +116,4 @@ UEdGraphPin* UK2Node_DataChannelBase::GetChannelSelectorPin() const
 	return FindPinChecked(FName("Channel"), EGPD_Input);
 }
 
-UEdGraphPin* UK2Node_DataChannelBase::GetAccessContextPin(EEdGraphPinDirection Direction)const
-{
-	return FindPin(FName("AccessContext"), Direction);
-}
-
 #undef LOCTEXT_NAMESPACE
diff --git a/Engine/Plugins/FX/Niagara/Source/NiagaraBlueprintNodes/Private/K2Node_DataChannel_WithContext.cpp b/Engine/Plugins/FX/Niagara/Source/NiagaraBlueprintNodes/Private/K2Node_DataChannel_WithContext.cpp
deleted file mode 100644
diff --git a/Engine/Plugins/FX/Niagara/Source/NiagaraBlueprintNodes/Private/K2Node_ReadDataChannel.cpp b/Engine/Plugins/FX/Niagara/Source/NiagaraBlueprintNodes/Private/K2Node_ReadDataChannel.cpp
--- a/Engine/Plugins/FX/Niagara/Source/NiagaraBlueprintNodes/Private/K2Node_ReadDataChannel.cpp
+++ b/Engine/Plugins/FX/Niagara/Source/NiagaraBlueprintNodes/Private/K2Node_ReadDataChannel.cpp
@@ -24,13 +24,13 @@ UK2Node_ReadDataChannel::UK2Node_ReadDataChannel()
 	FunctionReference.SetExternalMember(GET_FUNCTION_NAME_CHECKED(UNiagaraDataChannelLibrary, ReadFromNiagaraDataChannelSingle), UNiagaraDataChannelLibrary::StaticClass());
 }
 
-void UK2Node_ReadDataChannel::AddNDCDerivedPins()
+void UK2Node_ReadDataChannel::AllocateDefaultPins()
 {
-	Super::AddNDCDerivedPins();
+	Super::AllocateDefaultPins();
 
-	if (UNiagaraDataChannel* NDC = GetDataChannel())
+	if (HasValidDataChannel())
 	{
-		for (const FNiagaraDataChannelVariable& InVar : NDC->GetVariables())
+		for (const FNiagaraDataChannelVariable& InVar : GetDataChannel()->GetVariables())
 		{
 			if (IgnoredVariables.Contains(InVar.Version))
 			{
@@ -38,8 +38,6 @@ void UK2Node_ReadDataChannel::AddNDCDerivedPins()
 			}
 			UEdGraphPin* NewPin = CreatePin(EGPD_Output, FNiagaraBlueprintUtil::TypeDefinitionToBlueprintType(InVar.GetType()), InVar.GetName());
 
-			NDCDerivedGeneratedPins.Add(NewPin);
-
 #if WITH_EDITORONLY_DATA
 			NewPin->PersistentGuid = InVar.Version;
 #endif
@@ -47,30 +45,52 @@ void UK2Node_ReadDataChannel::AddNDCDerivedPins()
 	}
 }
 
-void UK2Node_ReadDataChannel::ExpandNode(FKismetCompilerContext& CompilerContext, UEdGraph* SourceGraph)
+void UK2Node_ReadDataChannel::GetMenuActions(FBlueprintActionDatabaseRegistrar& InActionRegistrar) const
 {
-	Super::ExpandNode(CompilerContext, SourceGraph);
+	const UClass* ActionKey = GetClass();
+	if (InActionRegistrar.IsOpenForRegistration(ActionKey))
+	{
+		UBlueprintNodeSpawner* NodeSpawner = UBlueprintNodeSpawner::Create(GetClass());
+		InActionRegistrar.AddBlueprintAction(ActionKey, NodeSpawner);
+	}
+}
 
+void UK2Node_ReadDataChannel::ExpandNode(FKismetCompilerContext& CompilerContext, UEdGraph* SourceGraph)
+{
 	const UEdGraphSchema_K2* Schema = CompilerContext.GetSchema();
 	UEdGraphPin* SuccessPin = FindPinChecked(FName("Success"), EGPD_Output);
 	UEdGraphPin* FailurePin = FindPinChecked(FName("Failure"), EGPD_Output);
+	if (!HasValidDataChannel())
+	{
+		// replace function with a noop if we don't have a data channel
+		UK2Node_ExecutionSequence* NoopNode = CompilerContext.SpawnIntermediateNode<UK2Node_ExecutionSequence>(this, SourceGraph);
+		NoopNode->AllocateDefaultPins();
+		
+		CompilerContext.MovePinLinksToIntermediate(*GetExecPin(), *NoopNode->GetExecPin());
+		CompilerContext.MovePinLinksToIntermediate(*FailurePin, *NoopNode->GetThenPinGivenIndex(0));
+		Schema->BreakPinLinks(*SuccessPin, false);
+		return;
+	}
+	if (DataChannelVersion != GetDataChannel()->GetVersion())
+	{
+		CompilerContext.MessageLog.Error(*LOCTEXT("StaleNode", "Node is out of sync with the data channel asset, please refresh node to fix up the pins - @@").ToString(), this);
+		return;
+	}
+	ExpandSplitPins(CompilerContext, SourceGraph);
 	
 	// create function call node to init the writer object 
 	UK2Node_CallFunction* CreateReaderNode = CompilerContext.SpawnIntermediateNode<UK2Node_CallFunction>(this, SourceGraph);
-	CreateReaderNode->SetFromFunction(UNiagaraDataChannelLibrary::StaticClass()->FindFunctionByName(GetReaderFunctionName()));
+	CreateReaderNode->SetFromFunction(UNiagaraDataChannelLibrary::StaticClass()->FindFunctionByName(GET_FUNCTION_NAME_CHECKED(UNiagaraDataChannelLibrary, ReadFromNiagaraDataChannel)));
 	CreateReaderNode->AllocateDefaultPins();
 
 	// transfer the input pins over
-	static TArray<TPair<FName, FName>> PinsToTransfer = { {"Channel", "Channel"}, {"SearchParams", "SearchParams"}, {"AccessContext", "AccessContext"}, {"bReadPreviousFrame", "bReadPreviousFrame"}};
+	static TArray<TPair<FName, FName>> PinsToTransfer = { {"Channel", "Channel"}, {"SearchParams", "SearchParams"}, {"bReadPreviousFrame", "bReadPreviousFrame"}};
 	
 	for (TPair<FName, FName> Pair : PinsToTransfer)
 	{
-		UEdGraphPin* OrgInputPin = FindPin(Pair.Key, EGPD_Input);
-		UEdGraphPin* NewInputPin = CreateReaderNode->FindPin(Pair.Value, EGPD_Input);
-		if(OrgInputPin && NewInputPin)
-		{
-			CompilerContext.MovePinLinksToIntermediate(*OrgInputPin, *NewInputPin);
-		}
+		UEdGraphPin* OrgInputPin = FindPinChecked(Pair.Key, EGPD_Input);
+		UEdGraphPin* NewInputPin = CreateReaderNode->FindPinChecked(Pair.Value, EGPD_Input);
+		CompilerContext.MovePinLinksToIntermediate(*OrgInputPin, *NewInputPin);
 	}
 	CompilerContext.MovePinLinksToIntermediate(*GetExecPin(), *CreateReaderNode->GetExecPin());
 	if (FindPin(FName("WorldContextObject"), EGPD_Input) && CreateReaderNode->FindPin(FName("WorldContextObject")))
diff --git a/Engine/Plugins/FX/Niagara/Source/NiagaraBlueprintNodes/Private/K2Node_WriteDataChannel.cpp b/Engine/Plugins/FX/Niagara/Source/NiagaraBlueprintNodes/Private/K2Node_WriteDataChannel.cpp
--- a/Engine/Plugins/FX/Niagara/Source/NiagaraBlueprintNodes/Private/K2Node_WriteDataChannel.cpp
+++ b/Engine/Plugins/FX/Niagara/Source/NiagaraBlueprintNodes/Private/K2Node_WriteDataChannel.cpp
@@ -4,6 +4,7 @@
 
 #include "BlueprintActionDatabaseRegistrar.h"
 #include "BlueprintNodeSpawner.h"
+#include "K2Node_ExecutionSequence.h"
 #include "K2Node_Self.h"
 #include "KismetCompiler.h"
 #include "NiagaraBlueprintUtil.h"
@@ -20,13 +21,13 @@ UK2Node_WriteDataChannel::UK2Node_WriteDataChannel()
 	FunctionReference.SetExternalMember(GET_FUNCTION_NAME_CHECKED(UNiagaraDataChannelLibrary, WriteToNiagaraDataChannelSingle), UNiagaraDataChannelLibrary::StaticClass());
 }
 
-void UK2Node_WriteDataChannel::AddNDCDerivedPins()
+void UK2Node_WriteDataChannel::AllocateDefaultPins()
 {
-	Super::AddNDCDerivedPins();
+	Super::AllocateDefaultPins();
 
-	if (UNiagaraDataChannel* NDC = GetDataChannel())
+	if (HasValidDataChannel())
 	{
-		for (const FNiagaraDataChannelVariable& InVar : NDC->GetVariables())
+		for (const FNiagaraDataChannelVariable& InVar : GetDataChannel()->GetVariables())
 		{
 			if (IgnoredVariables.Contains(InVar.Version))
 			{
@@ -34,8 +35,6 @@ void UK2Node_WriteDataChannel::AddNDCDerivedPins()
 			}
 			UEdGraphPin* NewPin = CreatePin(EGPD_Input, FNiagaraBlueprintUtil::TypeDefinitionToBlueprintType(InVar.GetType()), InVar.GetName());
 
-			NDCDerivedGeneratedPins.Add(NewPin);
-
 #if WITH_EDITORONLY_DATA
 			NewPin->PersistentGuid = InVar.Version;
 #endif
@@ -43,19 +42,37 @@ void UK2Node_WriteDataChannel::AddNDCDerivedPins()
 	}
 }
 
-void UK2Node_WriteDataChannel::ExpandNode(FKismetCompilerContext& CompilerContext, UEdGraph* SourceGraph)
+void UK2Node_WriteDataChannel::GetMenuActions(FBlueprintActionDatabaseRegistrar& InActionRegistrar) const
 {
-	Super::ExpandNode(CompilerContext, SourceGraph);
+	const UClass* ActionKey = GetClass();
+	if (InActionRegistrar.IsOpenForRegistration(ActionKey))
+	{
+		UBlueprintNodeSpawner* NodeSpawner = UBlueprintNodeSpawner::Create(GetClass());
+		InActionRegistrar.AddBlueprintAction(ActionKey, NodeSpawner);
+	}
+}
 
-	// create function call node to init the writer object 
-	UK2Node_CallFunction* CreateWriterNode = CompilerContext.SpawnIntermediateNode<UK2Node_CallFunction>(this, SourceGraph);
-	if(CreateWriterNode == nullptr)
+void UK2Node_WriteDataChannel::ExpandNode(FKismetCompilerContext& CompilerContext, UEdGraph* SourceGraph)
+{
+	if (!HasValidDataChannel())
 	{
-		CompilerContext.MessageLog.Error(*LOCTEXT("FailedToCreateWriterNode", "Failed to create NDC write node. - @@").ToString(), this);
+		// replace function with a noop if we don't have a data channel
+		UK2Node_ExecutionSequence* NoopNode = CompilerContext.SpawnIntermediateNode<UK2Node_ExecutionSequence>(this, SourceGraph);
+		NoopNode->AllocateDefaultPins();
+		CompilerContext.MovePinLinksToIntermediate(*GetExecPin(), *NoopNode->GetExecPin());
+		CompilerContext.MovePinLinksToIntermediate(*GetThenPin(), *NoopNode->GetThenPinGivenIndex(0));
 		return;
 	}
-
-	CreateWriterNode->SetFromFunction(UNiagaraDataChannelLibrary::StaticClass()->FindFunctionByName(GetWriterFunctionName()));
+	if (DataChannelVersion != GetDataChannel()->GetVersion())
+	{
+		CompilerContext.MessageLog.Error(*LOCTEXT("StaleNode", "Node is out of sync with the data channel asset, please refresh node to fix up the pins - @@").ToString(), this);
+		return;
+	}
+	ExpandSplitPins(CompilerContext, SourceGraph);
+	
+	// create function call node to init the writer object 
+	UK2Node_CallFunction* CreateWriterNode = CompilerContext.SpawnIntermediateNode<UK2Node_CallFunction>(this, SourceGraph);
+	CreateWriterNode->SetFromFunction(UNiagaraDataChannelLibrary::StaticClass()->FindFunctionByName(GET_FUNCTION_NAME_CHECKED(UNiagaraDataChannelLibrary, WriteToNiagaraDataChannel)));
 	CreateWriterNode->AllocateDefaultPins();
 
 #if WITH_NIAGARA_DEBUGGER
@@ -71,17 +88,14 @@ void UK2Node_WriteDataChannel::ExpandNode(FKismetCompilerContext& CompilerContex
 #endif
 
 	// transfer the input pins over
-	static TArray<TPair<FName, FName>> PinsToTransfer = { {"Channel", "Channel"}, {"SearchParams", "SearchParams"}, {"AccessContext", "AccessContext"}, {"bVisibleToBlueprint", "bVisibleToGame"},
+	static TArray<TPair<FName, FName>> PinsToTransfer = { {"Channel", "Channel"}, {"SearchParams", "SearchParams"}, {"bVisibleToBlueprint", "bVisibleToGame"},
 		{"bVisibleToNiagaraCPU", "bVisibleToCPU"}, {"bVisibleToNiagaraGPU", "bVisibleToGPU"}};
 	
 	for (TPair<FName, FName> Pair : PinsToTransfer)
 	{
-		UEdGraphPin* OrgInputPin = FindPin(Pair.Key, EGPD_Input);
-		UEdGraphPin* NewInputPin = CreateWriterNode->FindPin(Pair.Value, EGPD_Input);
-		if(OrgInputPin && NewInputPin)
-		{	
-			CompilerContext.MovePinLinksToIntermediate(*OrgInputPin, *NewInputPin);
-		}
+		UEdGraphPin* OrgInputPin = FindPinChecked(Pair.Key, EGPD_Input);
+		UEdGraphPin* NewInputPin = CreateWriterNode->FindPinChecked(Pair.Value, EGPD_Input);
+		CompilerContext.MovePinLinksToIntermediate(*OrgInputPin, *NewInputPin);
 	}
 	CreateWriterNode->FindPinChecked(FName("Count"))->DefaultValue = FString::FromInt(1);
 	if (FindPin(FName("WorldContextObject"), EGPD_Input) && CreateWriterNode->FindPin(FName("WorldContextObject")))
diff --git a/Engine/Plugins/FX/Niagara/Source/NiagaraBlueprintNodes/Private/NiagaraBlueprintNodesDetails.cpp b/Engine/Plugins/FX/Niagara/Source/NiagaraBlueprintNodes/Private/NiagaraBlueprintNodesDetails.cpp
deleted file mode 100644
diff --git a/Engine/Plugins/FX/Niagara/Source/NiagaraBlueprintNodes/Private/NiagaraBlueprintNodesDetails.h b/Engine/Plugins/FX/Niagara/Source/NiagaraBlueprintNodes/Private/NiagaraBlueprintNodesDetails.h
deleted file mode 100644
diff --git a/Engine/Plugins/FX/Niagara/Source/NiagaraBlueprintNodes/Private/NiagaraBlueprintNodesModule.cpp b/Engine/Plugins/FX/Niagara/Source/NiagaraBlueprintNodes/Private/NiagaraBlueprintNodesModule.cpp
--- a/Engine/Plugins/FX/Niagara/Source/NiagaraBlueprintNodes/Private/NiagaraBlueprintNodesModule.cpp
+++ b/Engine/Plugins/FX/Niagara/Source/NiagaraBlueprintNodes/Private/NiagaraBlueprintNodesModule.cpp
@@ -3,103 +3,15 @@
 #include "Modules/ModuleInterface.h"
 #include "Modules/ModuleManager.h"
 
-#include "BlueprintActionDatabase.h"
-#include "BlueprintEditorLibrary.h"
-
-#include "K2Node_DataChannel_WithContext.h"
-
-#include "NiagaraBlueprintNodesDetails.h"
-
-#include "PropertyEditorModule.h"
-
 class FNiagaraBlueprintNodesModule : public IModuleInterface
 {
 public:
 
-	virtual void StartupModule() override;
-	virtual void ShutdownModule() override;
-
-private:
-
-	/** Callback for changes to other modules. Allows us to refresh certain things to reflect new types or changes etc. */
-	void OnModulesChanged(FName Module, EModuleChangeReason Reason);
+	virtual void StartupModule() override
+	{}
 
-	/** In response to certain events like modules being loaded, we must ensure our BP actions are properly setup. */
-	void UpdateBPHooks();
-
-	FDelegateHandle OnModulesChangedHandle;
-
-	/** Set of struts with custom BP nodes/logic that require a refresh of the BP action DB. */
-	TSet<TWeakObjectPtr<UScriptStruct>> InitializedCustomBPStructs;
+	virtual void ShutdownModule() override
+	{}
 };
 
-
-void FNiagaraBlueprintNodesModule::StartupModule()
-{
-	OnModulesChangedHandle = FModuleManager::Get().OnModulesChanged().AddRaw(this, &FNiagaraBlueprintNodesModule::OnModulesChanged);
-
-	FPropertyEditorModule& PropertyModule = FModuleManager::LoadModuleChecked<FPropertyEditorModule>("PropertyEditor");
-
- 	PropertyModule.RegisterCustomClassLayout(
- 		UK2Node_DataChannelAccessContextOperation::StaticClass()->GetFName(),
- 		FOnGetDetailCustomizationInstance::CreateStatic(&FNDCAccessContextOperationNodeDetailsDetails::MakeInstance));
-		
-		
-	UpdateBPHooks();
-}
-
-void FNiagaraBlueprintNodesModule::ShutdownModule()
-{
-	FModuleManager::Get().OnModulesChanged().Remove(OnModulesChangedHandle);
-}
-
-void FNiagaraBlueprintNodesModule::OnModulesChanged(FName Module, EModuleChangeReason Reason)
-{
-	//Refresh any BP actions to accommodate newly loaded classes.
-	if (Reason == EModuleChangeReason::ModuleLoaded)
-	{
-		UpdateBPHooks();
-	}
-}
-
-void FNiagaraBlueprintNodesModule::UpdateBPHooks()
-{
-	if (FModuleManager::Get().IsModuleLoaded(TEXT("Kismet")))
-	{
-		for(auto It = InitializedCustomBPStructs.CreateIterator(); It; ++It)
-		{
-			if(It->IsValid() == false)
-			{
-				It.RemoveCurrent();
-			}
-		}
-
-		//Disable all make/break nodes for NDC Context classes
-		bool bRefresh = false;
-		for (TObjectIterator<UScriptStruct> It; It; ++It)
-		{
-			UScriptStruct* Struct = *It;
-			if (Struct->IsChildOf(FNDCAccessContextBase::StaticStruct()))
-			{
-				if(InitializedCustomBPStructs.Contains(Struct) == false)
-				{
-					bRefresh = true;	
-					InitializedCustomBPStructs.Add(Struct);
-				}		
-			}
-		}
-		
-		if(bRefresh)
-		{
-			FBlueprintActionDatabase& ActionDB = FBlueprintActionDatabase::Get();
-			ActionDB.RefreshClassActions(UK2Node_DataChannelAccessContextOperation::StaticClass());
-			ActionDB.RefreshClassActions(UK2Node_DataChannelAccessContext_Make::StaticClass());
-			ActionDB.RefreshClassActions(UK2Node_DataChannelAccessContext_GetMembers::StaticClass());
-			ActionDB.RefreshClassActions(UK2Node_DataChannelAccessContext_SetMembers::StaticClass());
-			UBlueprintEditorLibrary::RefreshAllOpenBlueprintEditors();
-		}
-	}
-}
-
-
 IMPLEMENT_MODULE(FNiagaraBlueprintNodesModule, NiagaraBlueprintNodes);
\ No newline at end of file
diff --git a/Engine/Plugins/FX/Niagara/Source/NiagaraEditor/Private/Customizations/NiagaraTypeCustomizations.cpp b/Engine/Plugins/FX/Niagara/Source/NiagaraEditor/Private/Customizations/NiagaraTypeCustomizations.cpp
--- a/Engine/Plugins/FX/Niagara/Source/NiagaraEditor/Private/Customizations/NiagaraTypeCustomizations.cpp
+++ b/Engine/Plugins/FX/Niagara/Source/NiagaraEditor/Private/Customizations/NiagaraTypeCustomizations.cpp
@@ -3205,128 +3205,5 @@ TSharedRef<SWidget> FNiagaraDataChannelVariableDetailsCustomization::GetTypeMenu
 	return TypeSelectMenu;
 }
 
-void FNDCAccessContextDetailsCustomization::CustomizeHeader(TSharedRef<IPropertyHandle> PropertyHandle, FDetailWidgetRow& HeaderRow, IPropertyTypeCustomizationUtils& CustomizationUtils)
-{
-	HeaderRow.NameContent()
-		[
-			PropertyHandle->CreatePropertyNameWidget()
-		];
-	HeaderRow.ValueContent()
-		[
-			PropertyHandle->CreatePropertyValueWidget()
-		];
-}
-
-void FNDCAccessContextDetailsCustomization::CustomizeChildren(TSharedRef<IPropertyHandle> PropertyHandle, IDetailChildrenBuilder& ChildBuilder, IPropertyTypeCustomizationUtils& CustomizationUtils)
-{
-	bool bShowInputContext = true;
-	bool bShowOutputContext = false;
-
-	//Walk up our parent chain until we get to a FNDCAccessContextInst and check it's meta data.
-	TSharedPtr<IPropertyHandle> ParentPtr = PropertyHandle->GetParentHandle();
-	while(ParentPtr.IsValid())
-	{
-		if (const FStructProperty* StructProp = CastField<FStructProperty>(ParentPtr->GetProperty()))
-		{
-			UScriptStruct* ParentStruct = StructProp->Struct;
-			if (ParentStruct == FNDCAccessContextInst::StaticStruct())
-			{
-				bShowInputContext = PropertyHandle->HasMetaData(TEXT("ShowNDCAccessContextInput"));
-				bShowOutputContext = PropertyHandle->HasMetaData(TEXT("ShowNDCAccessContextOutput"));
-
-				if (!bShowInputContext && !bShowOutputContext)
-				{
-					//If neither are set then assume the user has neglected to add the meta tags so just show input only.
-					bShowInputContext = true;
-					bShowOutputContext = false;
-				}
-				break;
-			}
-		}
-
-		ParentPtr = ParentPtr->GetParentHandle();
-	}
-
-	uint32 NumAccessContextChildren = 0;
-	PropertyHandle->GetNumChildren(NumAccessContextChildren);
-
-	for(uint32 ChildIdx = 0; ChildIdx < NumAccessContextChildren; ++ChildIdx)
-	{
-		TSharedPtr<IPropertyHandle> ChildPropertyHandle = PropertyHandle->GetChildHandle(ChildIdx);
-		if(!ChildPropertyHandle.IsValid())
-		{
-			continue;
-		}
-		FProperty* ChildProperty = ChildPropertyHandle->GetProperty();
-		if(!ChildProperty)
-		{
-			continue;
-		}
-
-		bool bIsInputContext = ChildPropertyHandle->HasMetaData(TEXT("NDCAccessContextInput")) && !ChildPropertyHandle->GetProperty()->HasAnyPropertyFlags(CPF_Transient);//Don't show transients in details panels.
-		bool bIsOutputContext = ChildPropertyHandle->HasMetaData(TEXT("NDCAccessContextOutput"));
-		if(!bIsInputContext && !bIsOutputContext)
-		{
-			//If neither are set then assume the user has neglected to add the meta tags so just show everything.
-			//A cluttered details display is better than an empty one.
-			bIsInputContext = true;
-			bIsOutputContext = true;
-		}
-
-		if((bShowInputContext && bIsInputContext) || (bShowOutputContext && bIsOutputContext))
-		{
-			ChildBuilder.AddProperty(ChildPropertyHandle.ToSharedRef());
-		}
-	}
-}
-
-
-void FNiagaraDataChannelReferenceDetailsCustomization::CustomizeHeader(TSharedRef<IPropertyHandle> PropertyHandle, FDetailWidgetRow& HeaderRow, IPropertyTypeCustomizationUtils& CustomizationUtils)
-{
-	HeaderRow.NameContent()
-		[
-			PropertyHandle->CreatePropertyNameWidget()
-		];
-	HeaderRow.ValueContent()
-		[
-			PropertyHandle->CreatePropertyValueWidget()
-		];
-}
-
-void FNiagaraDataChannelReferenceDetailsCustomization::CustomizeChildren(TSharedRef<IPropertyHandle> PropertyHandle, IDetailChildrenBuilder& ChildBuilder, IPropertyTypeCustomizationUtils& CustomizationUtils)
-{
-	PropertyPtr = PropertyHandle.ToSharedPtr();
-
-	uint32 NumChildren = 0;
-	PropertyHandle->GetNumChildren(NumChildren);
-
-	for (uint32 ChildIdx = 0; ChildIdx < NumChildren; ++ChildIdx)
-	{
-		TSharedPtr<IPropertyHandle> ChildPropertyHandle = PropertyHandle->GetChildHandle(ChildIdx);
-		if (!ChildPropertyHandle.IsValid())
-		{
-			continue;
-		}
-		FProperty* ChildProperty = ChildPropertyHandle->GetProperty();
-		if (!ChildProperty)
-		{
-			continue;
-		}
-
-		ChildBuilder.AddProperty(ChildPropertyHandle.ToSharedRef());
-
-		ChildPropertyHandle->SetOnPropertyValueChanged(FSimpleDelegate::CreateSP(this, &FNiagaraDataChannelReferenceDetailsCustomization::OnValueChanged));
-	}
-}
-
-void FNiagaraDataChannelReferenceDetailsCustomization::OnValueChanged()
-{
-	void* DataPtr = nullptr;
-	if (PropertyPtr.IsValid() && PropertyPtr->GetValueData(DataPtr) == FPropertyAccess::Success)
-	{
-		FNiagaraDataChannelReference* ValuePtr = static_cast<FNiagaraDataChannelReference*>(DataPtr);
-		ValuePtr->OnValueChanged();
-	}
-}
 
 #undef LOCTEXT_NAMESPACE
diff --git a/Engine/Plugins/FX/Niagara/Source/NiagaraEditor/Private/Customizations/NiagaraTypeCustomizations.h b/Engine/Plugins/FX/Niagara/Source/NiagaraEditor/Private/Customizations/NiagaraTypeCustomizations.h
--- a/Engine/Plugins/FX/Niagara/Source/NiagaraEditor/Private/Customizations/NiagaraTypeCustomizations.h
+++ b/Engine/Plugins/FX/Niagara/Source/NiagaraEditor/Private/Customizations/NiagaraTypeCustomizations.h
@@ -499,26 +499,3 @@ private:
 	TSharedRef<SWidget> GetTypeMenu(TSharedPtr<IPropertyHandle> InPropertyHandle, FNiagaraDataChannelVariable* Var);
 	TSharedPtr<SComboButton> ChangeTypeButton;
 };
-
-//** Properties customization for FNDCAccessContext. */
-class FNDCAccessContextDetailsCustomization : public IPropertyTypeCustomization
-{
-public:
-	static TSharedRef<IPropertyTypeCustomization> MakeInstance() { return MakeShared<FNDCAccessContextDetailsCustomization>(); }
-	virtual void CustomizeHeader(TSharedRef<IPropertyHandle> PropertyHandle, FDetailWidgetRow& HeaderRow, IPropertyTypeCustomizationUtils& CustomizationUtils) override;
-	virtual void CustomizeChildren(TSharedRef<IPropertyHandle> PropertyHandle, IDetailChildrenBuilder& ChildBuilder, IPropertyTypeCustomizationUtils& CustomizationUtils) override;
-};
-
-//** Properties customization for FNiagaraDataChannelReference. */
-class FNiagaraDataChannelReferenceDetailsCustomization : public IPropertyTypeCustomization
-{
-public:
-	static TSharedRef<IPropertyTypeCustomization> MakeInstance() { return MakeShared<FNiagaraDataChannelReferenceDetailsCustomization>(); }
-	virtual void CustomizeHeader(TSharedRef<IPropertyHandle> PropertyHandle, FDetailWidgetRow& HeaderRow, IPropertyTypeCustomizationUtils& CustomizationUtils) override;
-	virtual void CustomizeChildren(TSharedRef<IPropertyHandle> PropertyHandle, IDetailChildrenBuilder& ChildBuilder, IPropertyTypeCustomizationUtils& CustomizationUtils) override;
-
-private:
-	void OnValueChanged();
-
-	TSharedPtr<IPropertyHandle> PropertyPtr;
-};
diff --git a/Engine/Plugins/FX/Niagara/Source/NiagaraEditor/Private/NiagaraEditorModule.cpp b/Engine/Plugins/FX/Niagara/Source/NiagaraEditor/Private/NiagaraEditorModule.cpp
--- a/Engine/Plugins/FX/Niagara/Source/NiagaraEditor/Private/NiagaraEditorModule.cpp
+++ b/Engine/Plugins/FX/Niagara/Source/NiagaraEditor/Private/NiagaraEditorModule.cpp
@@ -1177,8 +1177,6 @@ void FNiagaraEditorModule::StartupModule()
 	FCoreDelegates::OnEnginePreExit.AddRaw(this, &FNiagaraEditorModule::OnPreExit);
 
 	OnAssetsPreDeleteHandle = FEditorDelegates::OnAssetsPreDelete.AddRaw(this, &FNiagaraEditorModule::OnAssetsPreDelete);
-
-	OnModulesChangedHandle = FModuleManager::Get().OnModulesChanged().AddRaw(this, &FNiagaraEditorModule::OnModulesChanged);
 	
 	// register details customization
 	FPropertyEditorModule& PropertyModule = FModuleManager::LoadModuleChecked<FPropertyEditorModule>("PropertyEditor");
@@ -1344,13 +1342,7 @@ void FNiagaraEditorModule::StartupModule()
 		FNiagaraRendererMaterialStaticBoolParameter::StaticStruct()->GetFName(),
 		FOnGetPropertyTypeCustomizationInstance::CreateStatic(&FNiagaraRendererMaterialStaticBoolParameterCustomization::MakeInstance));
 
-	PropertyModule.RegisterCustomPropertyTypeLayout(
-		FNiagaraDataChannelReference::StaticStruct()->GetFName(),
-		FOnGetPropertyTypeCustomizationInstance::CreateStatic(&FNiagaraDataChannelReferenceDetailsCustomization::MakeInstance));
-
-	RegisterDynamicCustomizations();
-
-	// details customizations end.
+	// Outliner Customizations end.
 
 	//Register Stack Object Issue Generators.
 	RegisterStackIssueGenerator(FNiagaraPlatformSet::StaticStruct()->GetFName(), new FNiagaraPlatformSetIssueGenerator());
@@ -1672,8 +1664,6 @@ void FNiagaraEditorModule::ShutdownModule()
 	FCoreUObjectDelegates::GetPostGarbageCollect().RemoveAll(this);
 	FCoreDelegates::OnPostEngineInit.RemoveAll(this);
 	FCoreDelegates::OnEnginePreExit.RemoveAll(this);
-
-	FModuleManager::Get().OnModulesChanged().Remove(OnModulesChangedHandle);
 	
 	if (GEditor)
 	{
@@ -2476,43 +2466,6 @@ void FNiagaraEditorModule::OnAssetsPreDelete(const TArray<UObject*>& Objects)
 	}
 }
 
-void FNiagaraEditorModule::OnModulesChanged(FName Module, EModuleChangeReason Reason)
-{	
-	//Refresh any dynamic customizations for potentially newly loaded types.
-	if (Reason == EModuleChangeReason::ModuleLoaded)
-	{
-		RegisterDynamicCustomizations();
-	}
-}
-
-void FNiagaraEditorModule::RegisterDynamicCustomizations()
-{
-	//Register any dynamic/discovered class customizations.
-	FPropertyEditorModule& PropertyModule = FModuleManager::LoadModuleChecked<FPropertyEditorModule>("PropertyEditor");
-	
-	for (TObjectIterator<UScriptStruct> It; It; ++It)
-	{
-		UScriptStruct* Struct = *It;
-		
-		//Ensure we don't register the struct again.
-		//If the module is unloaded / reloaded we'll still have it's correct customizer in place.
-		if(RegisteredDynamicStructs.Contains(Struct->GetFName()))
-		{
-			continue;
-		}
-
-		//Register customizations for Niagara Data Channel Access Context types.
-		if (Struct->IsChildOf(FNDCAccessContextBase::StaticStruct()))
-		{
-			PropertyModule.RegisterCustomPropertyTypeLayout(
-				Struct->GetFName(),
-				FOnGetPropertyTypeCustomizationInstance::CreateStatic(&FNDCAccessContextDetailsCustomization::MakeInstance));
-
-			RegisteredDynamicStructs.Add(Struct->GetFName());
-		}
-	}
-}
-
 FNiagaraEditorModule::FOnScriptApplied& FNiagaraEditorModule::OnScriptApplied()
 {
 	return OnScriptAppliedDelegate;
diff --git a/Engine/Plugins/FX/Niagara/Source/NiagaraEditor/Private/Widgets/DataChannel/NiagaraDataChannelWizard.cpp b/Engine/Plugins/FX/Niagara/Source/NiagaraEditor/Private/Widgets/DataChannel/NiagaraDataChannelWizard.cpp
--- a/Engine/Plugins/FX/Niagara/Source/NiagaraEditor/Private/Widgets/DataChannel/NiagaraDataChannelWizard.cpp
+++ b/Engine/Plugins/FX/Niagara/Source/NiagaraEditor/Private/Widgets/DataChannel/NiagaraDataChannelWizard.cpp
@@ -1231,7 +1231,6 @@ TSharedRef<FModuleWizardModel> DataChannel::CreateSpawnNDCModuleWizardModel()
 						if (UNiagaraDataInterfaceDataChannelRead* DataInterface = Cast<UNiagaraDataInterfaceDataChannelRead>(FunctionInput->Data))
 						{
 							DataInterface->Channel = Channel;
-							DataInterface->InitAccessContext();
 							DataInterface->bReadCurrentFrame = AssetPage->Data->bReadCurrentFrame;
 							DataInterface->bUpdateSourceDataEveryTick = AssetPage->Data->bUpdateSourceDataEveryTick;
 						}
diff --git a/Engine/Plugins/FX/Niagara/Source/NiagaraEditor/Public/NiagaraEditorModule.h b/Engine/Plugins/FX/Niagara/Source/NiagaraEditor/Public/NiagaraEditorModule.h
--- a/Engine/Plugins/FX/Niagara/Source/NiagaraEditor/Public/NiagaraEditorModule.h
+++ b/Engine/Plugins/FX/Niagara/Source/NiagaraEditor/Public/NiagaraEditorModule.h
@@ -445,10 +445,6 @@ private:
 	
 	void OnAssetsPreDelete(const TArray<UObject*>& Objects);
 
-	void OnModulesChanged(FName Module, EModuleChangeReason Reason);
-
-	void RegisterDynamicCustomizations();
-
 private:
 	TSharedPtr<FExtensibilityManager> MenuExtensibilityManager;
 	TSharedPtr<FExtensibilityManager> ToolBarExtensibilityManager;
@@ -492,8 +488,6 @@ private:
 	FDelegateHandle OnAssetsPreDeleteHandle;
 
 	FDelegateHandle DefaultTrackHandle;
-
-	FDelegateHandle OnModulesChangedHandle;
 	
 	TObjectPtr<USequencerSettings> SequencerSettings;
 
@@ -578,7 +572,4 @@ private:
 	UPackage* TempPackage = nullptr;
 
 	TSharedPtr<UE::Niagara::TraversalCache::FTraversalCache> TraversalCache;
-
-	//Set of dynamic structs we have done registrations for 
-	TSet<FName> RegisteredDynamicStructs;
 };
diff --git a/Engine/Source/Runtime/GameplayTags/Classes/BlueprintGameplayTagLibrary.h b/Engine/Source/Runtime/GameplayTags/Classes/BlueprintGameplayTagLibrary.h
--- a/Engine/Source/Runtime/GameplayTags/Classes/BlueprintGameplayTagLibrary.h
+++ b/Engine/Source/Runtime/GameplayTags/Classes/BlueprintGameplayTagLibrary.h
@@ -250,17 +250,6 @@ class UBlueprintGameplayTagLibrary : public UBlueprintFunctionLibrary
 	UFUNCTION(BlueprintPure, meta = (BlueprintInternalUseOnly = "TRUE"))
 	static GAMEPLAYTAGS_API bool HasAllMatchingGameplayTags(TScriptInterface<IGameplayTagAssetInterface> TagContainerInterface, const FGameplayTagContainer& OtherContainer);
 
-	/**
-	 * Check Gameplay tags in the interface has ANY of the specified tags in the tag container (expands to include parents of asset tags)
-	 *
-	 * @param TagContainerInterface		An Interface to a tag container
-	 * @param OtherContainer			A Tag Container
-	 *
-	 * @return True if the tagcontainer in the interface has ANY of the tags inside the container.
-	 */
-	UFUNCTION(BlueprintPure, meta = (BlueprintInternalUseOnly = "TRUE"))
-	static GAMEPLAYTAGS_API bool HasAnyMatchingGameplayTags(TScriptInterface<IGameplayTagAssetInterface> TagContainerInterface, const FGameplayTagContainer& OtherContainer);
-
 	/**
 	 * Check if the specified tag container has the specified tag, using the specified tag matching types
 	 *
diff --git a/Engine/Source/Runtime/GameplayTags/Private/BlueprintGameplayTagLibrary.cpp b/Engine/Source/Runtime/GameplayTags/Private/BlueprintGameplayTagLibrary.cpp
--- a/Engine/Source/Runtime/GameplayTags/Private/BlueprintGameplayTagLibrary.cpp
+++ b/Engine/Source/Runtime/GameplayTags/Private/BlueprintGameplayTagLibrary.cpp
@@ -220,15 +220,6 @@ bool UBlueprintGameplayTagLibrary::DoesTagAssetInterfaceHaveTag(TScriptInterface
 	return false;
 }
 
-bool UBlueprintGameplayTagLibrary::HasAnyMatchingGameplayTags(TScriptInterface<IGameplayTagAssetInterface> TagContainerInterface, const FGameplayTagContainer& OtherContainer)
-{
-	if (IGameplayTagAssetInterface* Interface = TagContainerInterface.GetInterface())
-	{
-		return Interface->HasAnyMatchingGameplayTags(OtherContainer);
-	}
-	return false;
-}
-
 void UBlueprintGameplayTagLibrary::AppendGameplayTagContainers(FGameplayTagContainer& InOutTagContainer, const FGameplayTagContainer& InTagContainer)
 {
 	InOutTagContainer.AppendTags(InTagContainer);
