Trim the unique emitter name down to 16 characters for the dump commandlet (avoids issus with crazy content exceeding filepath lengths)
#tests Ran dump commandlet over project

diff --git a/Engine/Plugins/FX/Niagara/Source/NiagaraEditor/Private/Commandlets/NiagaraDumpBytecodeCommandlet.cpp b/Engine/Plugins/FX/Niagara/Source/NiagaraEditor/Private/Commandlets/NiagaraDumpBytecodeCommandlet.cpp
--- a/Engine/Plugins/FX/Niagara/Source/NiagaraEditor/Private/Commandlets/NiagaraDumpBytecodeCommandlet.cpp
+++ b/Engine/Plugins/FX/Niagara/Source/NiagaraEditor/Private/Commandlets/NiagaraDumpBytecodeCommandlet.cpp
@@ -192,7 +192,7 @@ void UNiagaraDumpByteCodeCommandlet::ProcessBatch(TArray<FAssetData>& BatchAsset
 		DumpByteCode(NiagaraSystem->GetSystemSpawnScript(), SystemPathName, HashedPathName, TEXT("SystemSpawnScript"));
 		DumpByteCode(NiagaraSystem->GetSystemUpdateScript(), SystemPathName, HashedPathName, TEXT("SystemUpdateScript"));
 
-		for (const auto& EmitterHandle : NiagaraSystem->GetEmitterHandles())
+		for (const FNiagaraEmitterHandle& EmitterHandle : NiagaraSystem->GetEmitterHandles())
 		{
 			if (!EmitterHandle.GetIsEnabled())
 			{
@@ -203,7 +203,15 @@ void UNiagaraDumpByteCodeCommandlet::ProcessBatch(TArray<FAssetData>& BatchAsset
 			{
 				if (Emitter->SimTarget == ENiagaraSimTarget::CPUSim)
 				{
-					const FString EmitterName = EmitterHandle.GetUniqueInstanceName();
+					FString EmitterName = EmitterHandle.GetUniqueInstanceName();
+
+					// trim the emitter name down to 16 char at most
+					const int32 MaxEmitterNameLength = 16;
+					if (EmitterName.Len() > MaxEmitterNameLength)
+					{
+						UE_LOG(LogNiagaraDumpBytecodeCommandlet, Warning, TEXT("Emitter encountered with excessively long unique name: System: %s - Unique Emitter: %s"), *NiagaraSystem->GetFullName(), *EmitterName);
+						EmitterName.MidInline(0, MaxEmitterNameLength);
+					}
 
 					TArray<UNiagaraScript*> EmitterScripts;
 					Emitter->GetScripts(EmitterScripts);
