[Niagara]
* Fixed dll storage on FNiagaraDataChannelReference

diff --git a/Engine/Plugins/FX/Niagara/Source/Niagara/Public/NiagaraDataChannelPublic.h b/Engine/Plugins/FX/Niagara/Source/Niagara/Public/NiagaraDataChannelPublic.h
--- a/Engine/Plugins/FX/Niagara/Source/Niagara/Public/NiagaraDataChannelPublic.h
+++ b/Engine/Plugins/FX/Niagara/Source/Niagara/Public/NiagaraDataChannelPublic.h
@@ -582,14 +582,14 @@ public:
 
 /** A reference to a Niagara Data Channel with a pre-typed access context for convenience. */
 USTRUCT(BlueprintType)
-struct NIAGARA_API FNiagaraDataChannelReference
+struct FNiagaraDataChannelReference
 {
 	GENERATED_BODY()
 
 	[[nodiscard]] inline UNiagaraDataChannelAsset* GetDataChannel()const { return DataChannel; }
-	[[nodiscard]] FNDCAccessContextInst& GetAccessContext()const;
-	void SetDataChannel(UNiagaraDataChannelAsset* InDataChannel);	
-	void OnValueChanged();
+	NIAGARA_API [[nodiscard]] FNDCAccessContextInst& GetAccessContext()const;
+	NIAGARA_API void SetDataChannel(UNiagaraDataChannelAsset* InDataChannel);	
+	NIAGARA_API void OnValueChanged();
 
 	UPROPERTY(EditAnywhere, BlueprintReadWrite, Category = DataChannel)
 	TObjectPtr<UNiagaraDataChannelAsset> DataChannel;
