[Niagara]
* Moved dll storage macro to after nodiscard attribute to make mac happy

diff --git a/Engine/Plugins/FX/Niagara/Source/Niagara/Public/NiagaraDataChannelPublic.h b/Engine/Plugins/FX/Niagara/Source/Niagara/Public/NiagaraDataChannelPublic.h
--- a/Engine/Plugins/FX/Niagara/Source/Niagara/Public/NiagaraDataChannelPublic.h
+++ b/Engine/Plugins/FX/Niagara/Source/Niagara/Public/NiagaraDataChannelPublic.h
@@ -587,7 +587,7 @@ struct FNiagaraDataChannelReference
 	GENERATED_BODY()
 
 	[[nodiscard]] inline UNiagaraDataChannelAsset* GetDataChannel()const { return DataChannel; }
-	NIAGARA_API [[nodiscard]] FNDCAccessContextInst& GetAccessContext()const;
+	[[nodiscard]] NIAGARA_API FNDCAccessContextInst& GetAccessContext()const;
 	NIAGARA_API void SetDataChannel(UNiagaraDataChannelAsset* InDataChannel);	
 	NIAGARA_API void OnValueChanged();
 
