- Add icon to distinguish LW vs standard emitters

#tests Add stateful / stateless emittesr observer properties in normal / summary view mode that the right icons show

diff --git a/Engine/Plugins/FX/Niagara/Source/NiagaraEditor/Private/NiagaraEditorStyle.cpp b/Engine/Plugins/FX/Niagara/Source/NiagaraEditor/Private/NiagaraEditorStyle.cpp
--- a/Engine/Plugins/FX/Niagara/Source/NiagaraEditor/Private/NiagaraEditorStyle.cpp
+++ b/Engine/Plugins/FX/Niagara/Source/NiagaraEditor/Private/NiagaraEditorStyle.cpp
@@ -534,6 +534,7 @@ void FNiagaraEditorStyle::InitStackIcons()
 	//GPU/CPU icons
 	Set("NiagaraEditor.Stack.GPUIcon", new IMAGE_BRUSH("Icons/Simulate_GPU_x40", Icon16x16));
 	Set("NiagaraEditor.Stack.CPUIcon", new IMAGE_BRUSH("Icons/Simulate_CPU_x40", Icon16x16));
+	Set("NiagaraEditor.Stack.LWEIcon", new IMAGE_BRUSH("Icons/Simulate_LW_x40", Icon16x16));
 }
 
 void FNiagaraEditorStyle::InitStackWidgets()
diff --git a/Engine/Plugins/FX/Niagara/Source/NiagaraEditor/Private/ViewModels/Stack/NiagaraStackStatelessEmitterGroup.cpp b/Engine/Plugins/FX/Niagara/Source/NiagaraEditor/Private/ViewModels/Stack/NiagaraStackStatelessEmitterGroup.cpp
--- a/Engine/Plugins/FX/Niagara/Source/NiagaraEditor/Private/ViewModels/Stack/NiagaraStackStatelessEmitterGroup.cpp
+++ b/Engine/Plugins/FX/Niagara/Source/NiagaraEditor/Private/ViewModels/Stack/NiagaraStackStatelessEmitterGroup.cpp
@@ -2,15 +2,17 @@
 
 #include "ViewModels/Stack/NiagaraStackStatelessEmitterGroup.h"
 
-#include "IDetailTreeNode.h"
-#include "PropertyHandle.h"
 #include "Stateless/NiagaraStatelessEmitterDetails.h"
 #include "Stateless/NiagaraEmitterStatePropertyCustomization.h"
 #include "Stateless/NiagaraDistributionPropertyCustomization.h"
 #include "Stateless/NiagaraDistributionIntPropertyCustomization.h"
 #include "Stateless/NiagaraStatelessEmitter.h"
-#include "Styling/AppStyle.h"
 #include "ViewModels/Stack/NiagaraStackObject.h"
+#include "NiagaraEditorStyle.h"
+
+#include "IDetailTreeNode.h"
+#include "PropertyHandle.h"
+#include "Styling/AppStyle.h"
 
 #include UE_INLINE_GENERATED_CPP_BY_NAME(NiagaraStackStatelessEmitterGroup)
 
@@ -42,6 +44,16 @@ const FSlateBrush* UNiagaraStackStatelessEmitterGroup::GetIconBrush() const
 	return FAppStyle::Get().GetBrush("Icons.Details");
 }
 
+bool UNiagaraStackStatelessEmitterGroup::SupportsSecondaryIcon() const
+{
+	return true;
+}
+
+const FSlateBrush* UNiagaraStackStatelessEmitterGroup::GetSecondaryIconBrush() const
+{
+	return FNiagaraEditorStyle::Get().GetBrush("NiagaraEditor.Stack.LWEIcon");
+}
+
 void UNiagaraStackStatelessEmitterGroup::RefreshChildrenInternal(const TArray<UNiagaraStackEntry*>& CurrentChildren, TArray<UNiagaraStackEntry*>& NewChildren, TArray<FStackIssue>& NewIssues)
 {
 	Super::RefreshChildrenInternal(CurrentChildren, NewChildren, NewIssues);
diff --git a/Engine/Plugins/FX/Niagara/Source/NiagaraEditor/Public/ViewModels/Stack/NiagaraStackStatelessEmitterGroup.h b/Engine/Plugins/FX/Niagara/Source/NiagaraEditor/Public/ViewModels/Stack/NiagaraStackStatelessEmitterGroup.h
--- a/Engine/Plugins/FX/Niagara/Source/NiagaraEditor/Public/ViewModels/Stack/NiagaraStackStatelessEmitterGroup.h
+++ b/Engine/Plugins/FX/Niagara/Source/NiagaraEditor/Public/ViewModels/Stack/NiagaraStackStatelessEmitterGroup.h
@@ -22,6 +22,9 @@ public:
 	virtual EIconMode GetSupportedIconMode() const override { return EIconMode::Brush; }
 	virtual const FSlateBrush* GetIconBrush() const override;
 
+	virtual bool SupportsSecondaryIcon() const override;
+	virtual const FSlateBrush* GetSecondaryIconBrush() const override;
+
 	virtual bool GetCanExpandInOverview() const override { return false; }
 	virtual bool GetShouldShowInStack() const override { return false; }
 
diff --git a/Engine/Plugins/FX/Niagara/Source/NiagaraEditorWidgets/Private/Stack/SNiagaraSimTargetToggle.cpp b/Engine/Plugins/FX/Niagara/Source/NiagaraEditorWidgets/Private/Stack/SNiagaraSimTargetToggle.cpp
--- a/Engine/Plugins/FX/Niagara/Source/NiagaraEditorWidgets/Private/Stack/SNiagaraSimTargetToggle.cpp
+++ b/Engine/Plugins/FX/Niagara/Source/NiagaraEditorWidgets/Private/Stack/SNiagaraSimTargetToggle.cpp
@@ -14,9 +14,9 @@
 
 void SNiagaraSimTargetToggle::Construct(const FArguments& InArgs, TSharedPtr<FNiagaraEmitterHandleViewModel> InEmitterHandleViewModel)
 {
-	EmitterHandleViewModel = InEmitterHandleViewModel;
+	WeakEmitterHandleViewModel = InEmitterHandleViewModel;
 
-	if(EmitterHandleViewModel.IsValid() == false)
+	if (GetEmitterHandle() == nullptr)
 	{
 		return;
 	}
@@ -36,11 +36,13 @@ void SNiagaraSimTargetToggle::Construct(const FArguments& InArgs, TSharedPtr<FNi
 
 void SNiagaraSimTargetToggle::ToggleSimTarget() const
 {
-	if(EmitterHandleViewModel.IsValid() == false)
+	FNiagaraEmitterHandle* EmitterHandle = GetEmitterHandle();
+	TSharedPtr<FNiagaraEmitterHandleViewModel> EmitterHandleViewModel = WeakEmitterHandleViewModel.Pin();
+	if ( !EmitterHandle || EmitterHandle->GetEmitterMode() == ENiagaraEmitterMode::Stateless || !EmitterHandleViewModel.IsValid() )
 	{
 		return;
 	}
-	
+
 	FProperty* MatchingProperty = nullptr;
 	for(TFieldIterator<FProperty> It(FVersionedNiagaraEmitterData::StaticStruct()); It; ++It)
 	{
@@ -52,13 +54,13 @@ void SNiagaraSimTargetToggle::ToggleSimTarget() const
 	ensure(MatchingProperty != nullptr);
 
 	FScopedTransaction Transaction(LOCTEXT("ToggleSimTargetTransaction", "Simulation Target Changed"));
-	EmitterHandleViewModel.Pin()->GetEmitterViewModel()->GetEmitter().Emitter->PreEditChange(MatchingProperty);
+	EmitterHandleViewModel->GetEmitterViewModel()->GetEmitter().Emitter->PreEditChange(MatchingProperty);
 	
-	ENiagaraSimTarget CurrentSimTarget = EmitterHandleViewModel.Pin()->GetEmitterViewModel()->GetEmitter().GetEmitterData()->SimTarget;
-	EmitterHandleViewModel.Pin()->GetEmitterViewModel()->GetEmitter().GetEmitterData()->SimTarget = CurrentSimTarget == ENiagaraSimTarget::CPUSim ?  ENiagaraSimTarget::GPUComputeSim : ENiagaraSimTarget::CPUSim;
+	ENiagaraSimTarget CurrentSimTarget = EmitterHandleViewModel->GetEmitterViewModel()->GetEmitter().GetEmitterData()->SimTarget;
+	EmitterHandleViewModel->GetEmitterViewModel()->GetEmitter().GetEmitterData()->SimTarget = CurrentSimTarget == ENiagaraSimTarget::CPUSim ?  ENiagaraSimTarget::GPUComputeSim : ENiagaraSimTarget::CPUSim;
 
 	FPropertyChangedEvent PropertyChangedEvent(MatchingProperty);
-	UNiagaraSystem::RequestCompileForEmitter(EmitterHandleViewModel.Pin()->GetEmitterViewModel()->GetEmitter());
+	UNiagaraSystem::RequestCompileForEmitter(EmitterHandleViewModel->GetEmitterViewModel()->GetEmitter());
 }
 
 FReply SNiagaraSimTargetToggle::ToggleSimTargetInUI() const
@@ -69,26 +71,44 @@ FReply SNiagaraSimTargetToggle::ToggleSimTargetInUI() const
 
 const FSlateBrush* SNiagaraSimTargetToggle::GetSimTargetImage() const
 {
-	if(EmitterHandleViewModel.IsValid() == false)
+	if (FNiagaraEmitterHandle* EmitterHandle = GetEmitterHandle())
 	{
-		return FAppStyle::GetNoBrush();
+		if (EmitterHandle->GetEmitterMode() == ENiagaraEmitterMode::Stateless)
+		{
+			return FNiagaraEditorStyle::Get().GetBrush("NiagaraEditor.Stack.LWEIcon");
+		}
+		else
+		{
+			return EmitterHandle->GetInstance().GetEmitterData()->SimTarget == ENiagaraSimTarget::CPUSim
+				? FNiagaraEditorStyle::Get().GetBrush("NiagaraEditor.Stack.CPUIcon")
+				: FNiagaraEditorStyle::Get().GetBrush("NiagaraEditor.Stack.GPUIcon");
+		}
 	}
-	
-	return EmitterHandleViewModel.Pin()->GetEmitterViewModel()->GetEmitter().GetEmitterData()->SimTarget == ENiagaraSimTarget::CPUSim
-		? FNiagaraEditorStyle::Get().GetBrush("NiagaraEditor.Stack.CPUIcon")
-		: FNiagaraEditorStyle::Get().GetBrush("NiagaraEditor.Stack.GPUIcon");
+	return FAppStyle::GetNoBrush();
 }
 
 FText SNiagaraSimTargetToggle::GetSimTargetToggleTooltip() const
 {
-	if(EmitterHandleViewModel.IsValid() == false)
+	if (FNiagaraEmitterHandle* EmitterHandle = GetEmitterHandle())
 	{
-		return FText::GetEmpty();
+		if (EmitterHandle->GetEmitterMode() == ENiagaraEmitterMode::Stateless)
+		{
+			return LOCTEXT("ToggleSimTargetTooltip_LWE", "Lightweight emitter.");
+		}
+		else if (FVersionedNiagaraEmitterData* VersionedEmitterData = EmitterHandle->GetEmitterData())
+		{
+			return VersionedEmitterData->SimTarget == ENiagaraSimTarget::CPUSim
+				? LOCTEXT("ToggleSimTargetTooltip_CPU", "CPU emitter. Clicking will turn this emitter into a GPU emitter and recompile.")
+				: LOCTEXT("ToggleSimTargetTooltip_GPU", "GPU emitter. Clicking will turn this emitter into a CPU emitter and recompile.");		
+		}
 	}
+	return FText::GetEmpty();
+}
 
-	return EmitterHandleViewModel.Pin()->GetEmitterViewModel()->GetEmitter().GetEmitterData()->SimTarget == ENiagaraSimTarget::CPUSim
-		? LOCTEXT("ToggleSimTargetTooltip_CPU", "CPU emitter. Clicking will turn this emitter into a GPU emitter and recompile.")
-		: LOCTEXT("ToggleSimTargetTooltip_GPU", "GPU emitter. Clicking will turn this emitter into a CPU emitter and recompile.");		
+FNiagaraEmitterHandle* SNiagaraSimTargetToggle::GetEmitterHandle() const
+{
+	TSharedPtr<FNiagaraEmitterHandleViewModel> EmitterHandleViewModel = WeakEmitterHandleViewModel.Pin();
+	return EmitterHandleViewModel ? EmitterHandleViewModel->GetEmitterHandle() : nullptr;
 }
 
 #undef LOCTEXT_NAMESPACE
diff --git a/Engine/Plugins/FX/Niagara/Source/NiagaraEditorWidgets/Private/Stack/SNiagaraSimTargetToggle.h b/Engine/Plugins/FX/Niagara/Source/NiagaraEditorWidgets/Private/Stack/SNiagaraSimTargetToggle.h
--- a/Engine/Plugins/FX/Niagara/Source/NiagaraEditorWidgets/Private/Stack/SNiagaraSimTargetToggle.h
+++ b/Engine/Plugins/FX/Niagara/Source/NiagaraEditorWidgets/Private/Stack/SNiagaraSimTargetToggle.h
@@ -21,6 +21,8 @@ private:
 	const FSlateBrush* GetSimTargetImage() const;
 	FText GetSimTargetToggleTooltip() const;
 
+	FNiagaraEmitterHandle* GetEmitterHandle() const;
+
 private:
-	TWeakPtr<FNiagaraEmitterHandleViewModel> EmitterHandleViewModel;
+	TWeakPtr<FNiagaraEmitterHandleViewModel> WeakEmitterHandleViewModel;
 };
