NDCRef now can opt into having a custom AccessContext to save memory.
Transient/Default access context used if non provided.
Utility function to grab a valid, usable AccessContext from the NDCRef.

Gameplay Burst NDC now has a default system to spawn.
Each access can override this but doesn't need to.

diff --git a/Engine/Plugins/FX/Niagara/Source/Niagara/Private/NiagaraDataChannel.cpp b/Engine/Plugins/FX/Niagara/Source/Niagara/Private/NiagaraDataChannel.cpp
--- a/Engine/Plugins/FX/Niagara/Source/Niagara/Private/NiagaraDataChannel.cpp
+++ b/Engine/Plugins/FX/Niagara/Source/Niagara/Private/NiagaraDataChannel.cpp
@@ -1824,7 +1824,7 @@ UNiagaraDataChannelWriter* UNiagaraDataChannelLibrary::CreateDataChannelWriter_W
 {
 	check(IsInGameThread());
 	UWorld* World = (WorldContextObject != nullptr) ? GEngine->GetWorldFromContextObject(WorldContextObject, EGetWorldErrorMode::LogAndReturnNull) : nullptr;
-	if (Channel && World && Count > 0)
+	if (Channel && World && Count > 0 && AccessContext.IsValid())
 	{
 		if (FNiagaraWorldManager* WorldMan = FNiagaraWorldManager::Get(World))
 		{
@@ -1847,7 +1847,7 @@ UNiagaraDataChannelReader* UNiagaraDataChannelLibrary::CreateDataChannelReader_W
 {
 	check(IsInGameThread());
 	UWorld* World = (WorldContextObject != nullptr) ? GEngine->GetWorldFromContextObject(WorldContextObject, EGetWorldErrorMode::LogAndReturnNull) : nullptr;
-	if (Channel && World)
+	if (Channel && World && AccessContext.IsValid())
 	{
 		if (FNiagaraWorldManager* WorldMan = FNiagaraWorldManager::Get(World))
 		{
@@ -1866,6 +1866,20 @@ UNiagaraDataChannelReader* UNiagaraDataChannelLibrary::CreateDataChannelReader_W
 	return nullptr;
 }
 
+FNDCAccessContextInst& UNiagaraDataChannelLibrary::GetUsableAccessContextFromNDC(const UNiagaraDataChannelAsset* DataChannel)
+{
+	if(DataChannel && DataChannel->Get())
+	{
+		return DataChannel->Get()->GetTransientAccessContext();
+	}
+	static FNDCAccessContextInst Dummy;
+	return Dummy;
+}
+
+FNDCAccessContextInst& UNiagaraDataChannelLibrary::GetUsableAccessContextFromNDCRef(const FNiagaraDataChannelReference& NDCRef)
+{
+	return NDCRef.GetUsableAccessContext();
+}
 
 //////////////////////////////////////////////////////////////////////////
 // Functions for modifying dynamic objects.
@@ -1918,12 +1932,31 @@ void UNiagaraDataChannelAsset::PostEditChangeProperty(FPropertyChangedEvent& Pro
 
 //////////////////////////////////////////////////////////////////////////
 
-FNDCAccessContextInst& FNiagaraDataChannelReference::GetAccessContext()const
+const FNDCAccessContextInst& FNiagaraDataChannelReference::GetAccessContext()const
 {
 	InitAccessContext();
 	return AccessContext;
 }
 
+FNDCAccessContextInst& FNiagaraDataChannelReference::GetUsableAccessContext()const
+{
+	if (DataChannel)
+	{
+		if (UNiagaraDataChannel* NDC = DataChannel->Get())
+		{
+			InitAccessContext();
+			FNDCAccessContextInst& TransientContext = NDC->GetTransientAccessContext();
+			if(bCustomAccessContext && AccessContext.IsValid())
+			{
+				TransientContext = AccessContext;
+			}
+			return TransientContext;
+		}
+	}
+	static FNDCAccessContextInst Dummy;
+	return Dummy;
+}
+
 void FNiagaraDataChannelReference::InitAccessContext(bool bTransferExisting/* =true */)const
 {
 	//We need to verify that the access context has the correct type.
@@ -1933,7 +1966,7 @@ void FNiagaraDataChannelReference::InitAccessContext(bool bTransferExisting/* =t
 		if (UNiagaraDataChannel* NDC = DataChannel->Get())
 		{
 			const UScriptStruct* ExpectedType = NDC->GetAccessContextType().Get();
-			if(ExpectedType != FNDCAccessContextLegacy::StaticStruct())//Special case legacy type indicating we should not create a context in most places.
+			if(bCustomAccessContext && ExpectedType != FNDCAccessContextLegacy::StaticStruct())//Special case legacy type indicating we should not create a context in most places.
 			{
 				if (AccessContext.IsValid() && bTransferExisting)
 				{
diff --git a/Engine/Plugins/FX/Niagara/Source/Niagara/Private/NiagaraDataChannel_GameplayBurst.cpp b/Engine/Plugins/FX/Niagara/Source/Niagara/Private/NiagaraDataChannel_GameplayBurst.cpp
--- a/Engine/Plugins/FX/Niagara/Source/Niagara/Private/NiagaraDataChannel_GameplayBurst.cpp
+++ b/Engine/Plugins/FX/Niagara/Source/Niagara/Private/NiagaraDataChannel_GameplayBurst.cpp
@@ -39,6 +39,8 @@ void FNDCMapEntry_GameplayBurst::Init(FNDCAccessContextInst& AccessContext, UNia
 
 	Super::Init(AccessContext, InOwner, Key);
 
+	UObject* SysToSpawn = TypedOwner->GetSystemToSpawn(AccessContext);
+
 	bool bIsAttached = TypedContext.bAttachedToOwningComponent;
 	if (bIsAttached)
 	{
@@ -52,11 +54,11 @@ void FNDCMapEntry_GameplayBurst::Init(FNDCAccessContextInst& AccessContext, UNia
 		check(Data.IsValid());
 		Data->SetLwcTile(FLargeWorldRenderScalar::GetTileFor(Location));
 
-		if (TypedContext.SystemToSpawn)
+		if (SysToSpawn)
 		{
 			FBoxSphereBounds Bounds = Comp->CalcBounds(Comp->GetComponentToWorld());
 			FVector Extents = Bounds.BoxExtent + TypedContext.SystemBoundsPadding;
-			if (UNiagaraComponent* NewComp = SpawnSystem(AccessContext.GetChecked<FNDCAccessContext_MapBase>(), TypedContext.SystemToSpawn, Comp, Location, Extents, true))
+			if (UNiagaraComponent* NewComp = SpawnSystem(AccessContext.GetChecked<FNDCAccessContext_MapBase>(), SysToSpawn, Comp, Location, Extents, true))
 			{
 				//Add the attached state tag so we can be sure of how to handle FindData calls from it in future without having to re-gen attachment settings.
 				NewComp->ComponentTags.Add(UNiagaraDataChannelHandler_GameplayBurst::AttachedComponentTag);
@@ -70,8 +72,7 @@ void FNDCMapEntry_GameplayBurst::Init(FNDCAccessContextInst& AccessContext, UNia
 		check(Data.IsValid());
 		Data->SetLwcTile(FLargeWorldRenderScalar::GetTileFor(Location));
 
-		//TODO: Should this be a System collection. Or possibly a raw UObjet that could be either?
-		if (TypedContext.SystemToSpawn)
+		if (SysToSpawn)
 		{
 			int32 CellSize = TypedContext.bOverrideCellSize ? TypedContext.CellSizeOverride : TypedOwner->DefaultCellSize;
 			FVector CellExtents = FVector(CellSize * 0.5);
@@ -79,7 +80,7 @@ void FNDCMapEntry_GameplayBurst::Init(FNDCAccessContextInst& AccessContext, UNia
 			FVector CellCentre = FVector(CellCoord * CellSize) + CellExtents;
 			CellExtents += FVector(TypedContext.SystemBoundsPadding);
 			
-			SpawnSystem(AccessContext.GetChecked<FNDCAccessContext_MapBase>(), TypedContext.SystemToSpawn, nullptr, CellCentre, CellExtents, true);
+			SpawnSystem(AccessContext.GetChecked<FNDCAccessContext_MapBase>(), SysToSpawn, nullptr, CellCentre, CellExtents, true);
 		}
 	}
 }
diff --git a/Engine/Plugins/FX/Niagara/Source/Niagara/Private/NiagaraDataChannel_Map.cpp b/Engine/Plugins/FX/Niagara/Source/Niagara/Private/NiagaraDataChannel_Map.cpp
--- a/Engine/Plugins/FX/Niagara/Source/Niagara/Private/NiagaraDataChannel_Map.cpp
+++ b/Engine/Plugins/FX/Niagara/Source/Niagara/Private/NiagaraDataChannel_Map.cpp
@@ -32,6 +32,9 @@ UNiagaraDataChannelHandler_MapBase::UNiagaraDataChannelHandler_MapBase(FObjectIn
 void UNiagaraDataChannelHandler_MapBase::Init(const UNiagaraDataChannel* InChannel)
 {
 	Super::Init(InChannel);
+
+	const UNiagaraDataChannel_MapBase* NDCTyped = CastChecked<UNiagaraDataChannel_MapBase>(InChannel);
+	DefaultSystemToSpawn = NDCTyped->GetDefaultSystemToSpawn();
 }
 
 void UNiagaraDataChannelHandler_MapBase::Cleanup()
@@ -54,6 +57,8 @@ void UNiagaraDataChannelHandler_MapBase::Cleanup()
 	FreeEntries.Empty();
 	EntryPool.Empty();
 	SpawnedComponentsToActiveEntry.Empty();
+
+	DefaultSystemToSpawn = nullptr;
 }
 
 void UNiagaraDataChannelHandler_MapBase::BeginFrame(float DeltaTime)
@@ -141,6 +146,16 @@ FNiagaraDataChannelDataPtr UNiagaraDataChannelHandler_MapBase::FindData(FNDCAcce
 	return Entry.Get().GetData();
 }
 
+UObject* UNiagaraDataChannelHandler_MapBase::GetSystemToSpawn(FNDCAccessContextInst& AccessContext)const
+{
+	const FNDCAccessContext_MapBase& TypedContext = AccessContext.GetChecked<FNDCAccessContext_MapBase>();
+	if(TypedContext.bOverrideSystemToSpawn)
+	{
+		return TypedContext.SystemToSpawn;
+	}
+	return DefaultSystemToSpawn;
+}
+
 void UNiagaraDataChannelHandler_MapBase::GenerateKey(FNDCAccessContextInst& AccessContext, FNDCMapKeyWriter& KeyWriter)const
 {
 	const FNDCAccessContext_MapBase& TypedContext = AccessContext.GetChecked<FNDCAccessContext_MapBase>();
@@ -156,6 +171,10 @@ FNDCMapEntry& UNiagaraDataChannelHandler_MapBase::FindOrAddEntry(FNDCAccessConte
 {
 	FNDCAccessContext_MapBase& TypedContext = AccessContext.GetChecked<FNDCAccessContext_MapBase>();
 
+	//Clear out any transient output data from our context.
+	//TODO; Should possibly use virtual or iterate over properties to clear out all transient output properties here.
+	TypedContext.SpawnedSystems.Reset();
+
 	if (TypedContext.bIsAutoLinkingSystem)
 	{
 		//Only Niagara Components spawned by this NDC should be setting this flag. It is purely to allow auto-linking back to the originating NDC data without the need for generating a key.
@@ -182,10 +201,6 @@ FNDCMapEntry& UNiagaraDataChannelHandler_MapBase::FindOrAddEntry(FNDCAccessConte
 		}
 	}
 
-	//Clear out any transient output data from our context.
-	//TODO; Should possibly use virtual or iterate over properties to clear out all transient output properties here.
-	TypedContext.SpawnedSystems.Reset();
-
 	//Generate the key for this Access Context. Each child class will implement this as needed.
 	FNDCMapKey Key;
 	FNDCMapKeyWriter KeyWriter(Key);
diff --git a/Engine/Plugins/FX/Niagara/Source/Niagara/Public/NiagaraDataChannel.h b/Engine/Plugins/FX/Niagara/Source/Niagara/Public/NiagaraDataChannel.h
--- a/Engine/Plugins/FX/Niagara/Source/Niagara/Public/NiagaraDataChannel.h
+++ b/Engine/Plugins/FX/Niagara/Source/Niagara/Public/NiagaraDataChannel.h
@@ -492,6 +492,14 @@ class UNiagaraDataChannelLibrary : public UBlueprintFunctionLibrary
 	static UE_API UNiagaraDataChannelReader* CreateDataChannelReader_WithContext(const UObject* WorldContextObject, const UNiagaraDataChannel* Channel, FNDCAccessContextInst& AccessContext, bool bReadPreviousFrame);
 
 
+	/** Returns the access context for use when accessing a Niagara Data channel. */
+	UFUNCTION(BlueprintCallable, Category = "Niagara Data Channel")
+	static UE_API FNDCAccessContextInst& GetUsableAccessContextFromNDC(const UNiagaraDataChannelAsset* DataChannel);
+
+	/** Returns the access context for use when accessing a Niagara Data channel. */
+	UFUNCTION(BlueprintCallable, Category = "Niagara Data Channel")
+	static UE_API FNDCAccessContextInst& GetUsableAccessContextFromNDCRef(const FNiagaraDataChannelReference& NDCRef);
+
 	/** Creates a new NDCAccessContext with a specified internal struct type. Struct must be a child of FNDCAccessContextBase. */
 	UFUNCTION(BlueprintInternalUseOnly, BlueprintCallable, CustomThunk, Category = "Niagara Data Channel", meta = (HidePin = "ContextStruct"))
 	static UE_API FNDCAccessContextInst MakeNDCAccessContextInstance(UScriptStruct* ContextStruct);
diff --git a/Engine/Plugins/FX/Niagara/Source/Niagara/Public/NiagaraDataChannelPublic.h b/Engine/Plugins/FX/Niagara/Source/Niagara/Public/NiagaraDataChannelPublic.h
--- a/Engine/Plugins/FX/Niagara/Source/Niagara/Public/NiagaraDataChannelPublic.h
+++ b/Engine/Plugins/FX/Niagara/Source/Niagara/Public/NiagaraDataChannelPublic.h
@@ -193,6 +193,7 @@ struct FNDCAccessContext : public FNDCAccessContextBase
 		: FNDCAccessContextBase(Owner)
 		, Location(LocationOverride)
 		, bOverrideLocation(true)
+		, bOverrideSystemToSpawn(false)
 	{
 	}
 
@@ -200,6 +201,7 @@ struct FNDCAccessContext : public FNDCAccessContextBase
 		: FNDCAccessContextBase(nullptr)
 		, Location(InLocation)
 		, bOverrideLocation(true)
+		, bOverrideSystemToSpawn(false)
 	{
 	}
 
@@ -213,7 +215,11 @@ struct FNDCAccessContext : public FNDCAccessContextBase
 	UPROPERTY(BlueprintReadWrite, Transient, Category = "Input", meta=(NDCAccessContextInput))
 	uint32 bOverrideLocation : 1 = false;
 
-	UPROPERTY(EditAnywhere, BlueprintReadWrite, Category = "Input", meta = (NDCAccessContextInput, AllowedClasses = "/Script/Niagara.NiagaraSystem,/Script/Niagara.NiagaraSystemCollection"))
+	/** If true, we'll override the system to spawn for this NDC write. If false, we'll use the system defined in the NDC asset, if there is one. */
+	UPROPERTY(BlueprintReadWrite, Category = "Input", meta = (InlineEditCondition, NDCAccessContextInput))
+	uint32 bOverrideSystemToSpawn : 1 = false;
+
+	UPROPERTY(EditAnywhere, BlueprintReadWrite, Category = "Input", meta = (EditCondition="bOverrideSystemToSpawn", NDCAccessContextInput, AllowedClasses = "/Script/Niagara.NiagaraSystem,/Script/Niagara.NiagaraSystemCollection"))
 	TObjectPtr<UObject> SystemToSpawn;
 
 	/** The Niagara System spawned in response to this access, if any. Allows callers to set parameters etc on this system. It is not safe to store and use this component reference later. It can be reclaimed by the system at any time. */
@@ -587,7 +593,10 @@ struct FNiagaraDataChannelReference
 	GENERATED_BODY()
 
 	[[nodiscard]] inline UNiagaraDataChannelAsset* GetDataChannel()const { return DataChannel; }
-	[[nodiscard]] NIAGARA_API FNDCAccessContextInst& GetAccessContext()const;
+	/** Returns the internal access context instance. Can be invalid. */
+	[[nodiscard]] NIAGARA_API const FNDCAccessContextInst& GetAccessContext()const;
+	/** Returns a usable access context for accessing the referenced NDC. If the internal context is valid it will include it's input data. */
+	[[nodiscard]] NIAGARA_API FNDCAccessContextInst& GetUsableAccessContext()const;
 	NIAGARA_API void SetDataChannel(UNiagaraDataChannelAsset* InDataChannel);	
 	NIAGARA_API void OnValueChanged();
 
@@ -597,6 +606,10 @@ struct FNiagaraDataChannelReference
 	UPROPERTY(EditAnywhere, BlueprintReadWrite, Category = DataChannel, meta = (ShowOnlyInnerProperties, ShowNDCAccessContextInput))
 	mutable FNDCAccessContextInst AccessContext;
 
+	/** If true this reference will provide a custom access context. */
+	UPROPERTY(EditAnywhere, BlueprintReadOnly, Category = DataChannel)
+	bool bCustomAccessContext = false;
+
 private:
 
 	void InitAccessContext(bool bTransferExisting = true)const;
diff --git a/Engine/Plugins/FX/Niagara/Source/Niagara/Public/NiagaraDataChannel_Map.h b/Engine/Plugins/FX/Niagara/Source/Niagara/Public/NiagaraDataChannel_Map.h
--- a/Engine/Plugins/FX/Niagara/Source/Niagara/Public/NiagaraDataChannel_Map.h
+++ b/Engine/Plugins/FX/Niagara/Source/Niagara/Public/NiagaraDataChannel_Map.h
@@ -198,7 +198,6 @@ struct FNDCMapKeyWriter
 	FNDCMapKey& DestKey;
 };
 
-
 USTRUCT(BlueprintType, meta=(Hidden))
 struct FNDCAccessContext_MapBase : public FNDCAccessContext
 {
@@ -212,6 +211,16 @@ UCLASS(MinimalAPI, abstract)
 class UNiagaraDataChannel_MapBase : public UNiagaraDataChannel
 {
 	GENERATED_BODY()
+
+public:
+
+	[[nodiscard]] UObject* GetDefaultSystemToSpawn()const { return DefaultSystemToSpawn; }
+
+protected:
+
+	/** Niagara System (or System Collection) that we should spawn to listen for and handle NDC data. Can be overriden per via the SystemToSpawn property in the Access Context. */
+	UPROPERTY(EditAnywhere, Category = "Data Channel", meta = (AllowedClasses = "/Script/Niagara.NiagaraSystem,/Script/Niagara.NiagaraSystemCollection"))
+	TObjectPtr<UObject> DefaultSystemToSpawn;
 };
 
 /** Base Handler class for Data Channels that route and store subsets of internal data based in a map based on an arbitrary key they can generate in GenerateKey(). */
@@ -229,6 +238,8 @@ class UNiagaraDataChannelHandler_MapBase : public UNiagaraDataChannelHandler
 	[[nodiscard]] virtual FNiagaraDataChannelDataPtr FindData(FNDCAccessContextInst& AccessContext, ENiagaraResourceAccess AccessType)override;
 
 	[[nodiscard]] UWorld* GetWorld()const { return OwningWorld->GetWorld(); }
+
+	[[nodiscard]] UObject* GetSystemToSpawn(FNDCAccessContextInst& AccessContext)const;
 protected:
 
 	/** Returns the map entry data type for this handler. Each child class should implement this and return their own subtype of FNDCMapEntryBase that contains their needed per entry data. */
@@ -254,4 +265,7 @@ protected:
 
 	/** Mapping from spawned handler systems back to their owning entry. Allows faster access when handler systems are accessing the NDC. */
 	TMap<UNiagaraComponent*, int32> SpawnedComponentsToActiveEntry;
+
+	UPROPERTY()
+	TObjectPtr<UObject> DefaultSystemToSpawn;
 };
\ No newline at end of file
