- Remove UNiagaraEmitter from stateless emitter on cook
- The cooked build does not require them, various parts of the UI still make assumptions and will be cleaned up
- This saves ~4k per stateless emitter

#tests cooked replay run / editor testing

diff --git a/Engine/Plugins/FX/Niagara/Source/Niagara/Private/NiagaraSystem.cpp b/Engine/Plugins/FX/Niagara/Source/Niagara/Private/NiagaraSystem.cpp
--- a/Engine/Plugins/FX/Niagara/Source/Niagara/Private/NiagaraSystem.cpp
+++ b/Engine/Plugins/FX/Niagara/Source/Niagara/Private/NiagaraSystem.cpp
@@ -96,6 +96,14 @@ namespace NiagaraSystemPrivate
 	static const FName NAME_FixedBoundsSize("FixedBoundsSize");
 	static const FName NAME_NumEmitters("NumEmitters");
 
+	bool GCleanStatelessEmitterHandles = true;
+	static FAutoConsoleVariableRef CVarCleanStatelessEmitterHandles(
+		TEXT("fx.Niagara.CleanStatelessEmitterHandles"),
+		GCleanStatelessEmitterHandles,
+		TEXT("When enabled stateless emitter handles will clean the stateful emitter during cook to save memory."),
+		ECVF_Default
+	);
+
 #if WITH_EDITORONLY_DATA
 	bool GMarkPackageDirtyWhenNotSynchronized = false;
 	static FAutoConsoleVariableRef CVarMarkPackageDirtyWhenNotSynchronized(
@@ -1082,8 +1090,34 @@ void UNiagaraSystem::Serialize(FArchive& Ar)
 		bDeterminism = true;
 	}
 
+	// Temporary hack to remove the UNiagaraEmitter from stateless emitters
+	// Currently the UI and various parts of the code do not handle the lack of an UNiagaraEmitter that will be fixed, for the moment we just remove them from cooked
+	TArray<FVersionedNiagaraEmitter> TempVersionedEmitterStore;
+	if (NiagaraSystemPrivate::GCleanStatelessEmitterHandles && Ar.IsSaving() && Ar.IsCooking() && Ar.CookingTarget() && !Ar.CookingTarget()->HasEditorOnlyData())
+	{
+		TempVersionedEmitterStore.AddDefaulted(EmitterHandles.Num());
+		for (int32 i=0; i < EmitterHandles.Num(); ++i)
+		{
+			FNiagaraEmitterHandle& Handle = EmitterHandles[i];
+			TempVersionedEmitterStore[i] = Handle.GetInstance();
+			if ( Handle.GetEmitterMode() == ENiagaraEmitterMode::Stateless )
+			{
+				Handle.GetInstance() = FVersionedNiagaraEmitter();
+			}
+		}
+	}
+
 	Super::Serialize(Ar);
 
+	if (TempVersionedEmitterStore.Num() > 0)
+	{
+		for (int32 i=0; i < EmitterHandles.Num(); ++i)
+		{
+			EmitterHandles[i].GetInstance() = TempVersionedEmitterStore[i];
+		}
+		TempVersionedEmitterStore.Empty();
+	}
+
 	if (Ar.CustomVer(FNiagaraCustomVersion::GUID) >= FNiagaraCustomVersion::ChangeEmitterCompiledDataToSharedRefs)
 	{
 		UScriptStruct* NiagaraEmitterCompiledDataStruct = FNiagaraEmitterCompiledData::StaticStruct();
@@ -2355,7 +2389,8 @@ void UNiagaraSystem::ComputeEmittersExecutionOrder()
 	// likely been cooked out because of scalability
 	EmitterExecutionOrder.SetNum(Algo::StableRemoveIf(EmitterExecutionOrder, [this](FNiagaraEmitterExecutionIndex EmitterExecIdx)
 	{
-		return EmitterHandles[EmitterExecIdx.EmitterIndex].GetInstance().Emitter == nullptr;
+		const FNiagaraEmitterHandle& EmitterHandle = EmitterHandles[EmitterExecIdx.EmitterIndex];
+		return EmitterHandle.GetEmitterMode() == ENiagaraEmitterMode::Stateless ? EmitterHandle.GetStatelessEmitter() == nullptr : EmitterHandle.GetInstance().Emitter == nullptr;
 	}));
 }
 
