[Niagara/MegaLights] Always bound UniqueID as a Niagara light renderer attribute. This is needed to allow runtime switching all simple lights to be handled by MegaLights (i.e. r.MegaLights.SimpleLightMode 2)

#tests tested in editor

diff --git a/Engine/Plugins/FX/Niagara/Source/Niagara/Private/NiagaraLightRendererProperties.cpp b/Engine/Plugins/FX/Niagara/Source/Niagara/Private/NiagaraLightRendererProperties.cpp
--- a/Engine/Plugins/FX/Niagara/Source/Niagara/Private/NiagaraLightRendererProperties.cpp
+++ b/Engine/Plugins/FX/Niagara/Source/Niagara/Private/NiagaraLightRendererProperties.cpp
@@ -215,14 +215,7 @@ void UNiagaraLightRendererProperties::CacheFromCompiledData(const FNiagaraDataSe
 {
 	UpdateSourceModeDerivates(SourceMode);
 
-	if (bAllowMegaLights)
-	{
-		LightIdAccessor.Init(CompiledData, NiagaraLightRendererPropertiesPrivate::NAME_UniqueID);
-	}
-	else
-	{
-		LightIdAccessor.Init(nullptr, NAME_None);
-	}
+	LightIdAccessor.Init(CompiledData, NiagaraLightRendererPropertiesPrivate::NAME_UniqueID);
 
 	InitParticleDataSetAccessor(PositionDataSetAccessor, CompiledData, PositionBinding);
 	InitParticleDataSetAccessor(ColorDataSetAccessor, CompiledData, ColorBinding);
@@ -276,10 +269,8 @@ TArray<FNiagaraVariable> UNiagaraLightRendererProperties::GetBoundAttributes() c
 {
 	TArray<FNiagaraVariable> BoundAttributes = Super::GetBoundAttributes();
 
-	if (bAllowMegaLights)
-	{
-		BoundAttributes.AddUnique(NiagaraLightRendererPropertiesPrivate::GetUniqueIDVariable());
-	}
+	// UniqueID is needed by MegaLights to track light visibility
+	BoundAttributes.AddUnique(NiagaraLightRendererPropertiesPrivate::GetUniqueIDVariable());
 
 	return BoundAttributes;
 }
diff --git a/Engine/Plugins/FX/Niagara/Source/Niagara/Private/NiagaraRendererLights.cpp b/Engine/Plugins/FX/Niagara/Source/Niagara/Private/NiagaraRendererLights.cpp
--- a/Engine/Plugins/FX/Niagara/Source/Niagara/Private/NiagaraRendererLights.cpp
+++ b/Engine/Plugins/FX/Niagara/Source/Niagara/Private/NiagaraRendererLights.cpp
@@ -242,7 +242,7 @@ FNiagaraDynamicDataBase* FNiagaraRendererLights::GenerateDynamicData(const FNiag
 			LightData.LightEntry.VolumetricScatteringIntensity = LightScattering;
 			LightData.LightEntry.SpecularScale = Properties->SpecularScale;
 			LightData.LightEntry.DiffuseScale = Properties->DiffuseScale;
-			LightData.LightEntry.LightId = FSimpleLightId(Properties->bAllowMegaLights ? 0 : INDEX_NONE, SourcePrimitiveId);
+			LightData.LightEntry.LightId = FSimpleLightId(0, SourcePrimitiveId);
 
 			LightData.PerViewEntry.Position = SimToWorld.TransformPosition(FVector(SimPos));
 		}
