- Fix CIS issue on some platforms

#rnx

diff --git a/Engine/Plugins/FX/Niagara/Source/Niagara/Private/NiagaraSystem.cpp b/Engine/Plugins/FX/Niagara/Source/Niagara/Private/NiagaraSystem.cpp
--- a/Engine/Plugins/FX/Niagara/Source/Niagara/Private/NiagaraSystem.cpp
+++ b/Engine/Plugins/FX/Niagara/Source/Niagara/Private/NiagaraSystem.cpp
@@ -1099,10 +1099,12 @@ void UNiagaraSystem::Serialize(FArchive& Ar)
 		for (int32 i=0; i < EmitterHandles.Num(); ++i)
 		{
 			FNiagaraEmitterHandle& Handle = EmitterHandles[i];
-			TempVersionedEmitterStore[i] = Handle.GetInstance();
+			TempVersionedEmitterStore[i].Emitter = Handle.GetInstance().Emitter;
+			TempVersionedEmitterStore[i].Version = Handle.GetInstance().Version;
 			if ( Handle.GetEmitterMode() == ENiagaraEmitterMode::Stateless )
 			{
-				Handle.GetInstance() = FVersionedNiagaraEmitter();
+				Handle.GetInstance().Emitter = nullptr;
+				Handle.GetInstance().Version = FGuid();
 			}
 		}
 	}
@@ -1113,7 +1115,8 @@ void UNiagaraSystem::Serialize(FArchive& Ar)
 	{
 		for (int32 i=0; i < EmitterHandles.Num(); ++i)
 		{
-			EmitterHandles[i].GetInstance() = TempVersionedEmitterStore[i];
+			EmitterHandles[i].GetInstance().Emitter = TempVersionedEmitterStore[i].Emitter;
+			EmitterHandles[i].GetInstance().Version = TempVersionedEmitterStore[i].Version;
 		}
 		TempVersionedEmitterStore.Empty();
 	}
