- Fix for API refactored in Main vs earlier branches

#rnx

diff --git a/Engine/Plugins/FX/Niagara/Source/Niagara/Classes/NiagaraEmitterHandle.h b/Engine/Plugins/FX/Niagara/Source/Niagara/Classes/NiagaraEmitterHandle.h
--- a/Engine/Plugins/FX/Niagara/Source/Niagara/Classes/NiagaraEmitterHandle.h
+++ b/Engine/Plugins/FX/Niagara/Source/Niagara/Classes/NiagaraEmitterHandle.h
@@ -80,6 +80,7 @@ public:
 
 	/** Gets the copied instance of the emitter this handle references. */
 	[[nodiscard]] NIAGARA_API const FVersionedNiagaraEmitter GetInstance() const;
+	NIAGARA_API void SetInstance(const FVersionedNiagaraEmitter& VersionedData);
 	NIAGARA_API void SetInstanceVersion(const FGuid& InVersion);
 
 	NIAGARA_API FVersionedNiagaraEmitterData* GetEmitterData() const;
diff --git a/Engine/Plugins/FX/Niagara/Source/Niagara/Private/NiagaraEmitterHandle.cpp b/Engine/Plugins/FX/Niagara/Source/Niagara/Private/NiagaraEmitterHandle.cpp
--- a/Engine/Plugins/FX/Niagara/Source/Niagara/Private/NiagaraEmitterHandle.cpp
+++ b/Engine/Plugins/FX/Niagara/Source/Niagara/Private/NiagaraEmitterHandle.cpp
@@ -231,6 +231,11 @@ const FVersionedNiagaraEmitter FNiagaraEmitterHandle::GetInstance() const
 	return VersionedInstance;
 }
 
+void FNiagaraEmitterHandle::SetInstance(const FVersionedNiagaraEmitter& VersionedData)
+{
+	VersionedInstance = VersionedData;
+}
+
 void FNiagaraEmitterHandle::SetInstanceVersion(const FGuid& InVersion)
 {
 	VersionedInstance.Version = InVersion;
diff --git a/Engine/Plugins/FX/Niagara/Source/Niagara/Private/NiagaraSystem.cpp b/Engine/Plugins/FX/Niagara/Source/Niagara/Private/NiagaraSystem.cpp
--- a/Engine/Plugins/FX/Niagara/Source/Niagara/Private/NiagaraSystem.cpp
+++ b/Engine/Plugins/FX/Niagara/Source/Niagara/Private/NiagaraSystem.cpp
@@ -1099,12 +1099,10 @@ void UNiagaraSystem::Serialize(FArchive& Ar)
 		for (int32 i=0; i < EmitterHandles.Num(); ++i)
 		{
 			FNiagaraEmitterHandle& Handle = EmitterHandles[i];
-			TempVersionedEmitterStore[i].Emitter = Handle.GetInstance().Emitter;
-			TempVersionedEmitterStore[i].Version = Handle.GetInstance().Version;
+			TempVersionedEmitterStore[i] = Handle.GetInstance();
 			if ( Handle.GetEmitterMode() == ENiagaraEmitterMode::Stateless )
 			{
-				Handle.GetInstance().Emitter = nullptr;
-				Handle.GetInstance().Version = FGuid();
+				Handle.SetInstance(FVersionedNiagaraEmitter());
 			}
 		}
 	}
@@ -1115,8 +1113,7 @@ void UNiagaraSystem::Serialize(FArchive& Ar)
 	{
 		for (int32 i=0; i < EmitterHandles.Num(); ++i)
 		{
-			EmitterHandles[i].GetInstance().Emitter = TempVersionedEmitterStore[i].Emitter;
-			EmitterHandles[i].GetInstance().Version = TempVersionedEmitterStore[i].Version;
+			EmitterHandles[i].SetInstance(TempVersionedEmitterStore[i]);
 		}
 		TempVersionedEmitterStore.Empty();
 	}
