- Guard against using EmitterData as it can be nullptr

diff --git a/Engine/Plugins/FX/Niagara/Source/NiagaraEditor/Private/ViewModels/NiagaraParameterPanelViewModel.cpp b/Engine/Plugins/FX/Niagara/Source/NiagaraEditor/Private/ViewModels/NiagaraParameterPanelViewModel.cpp
--- a/Engine/Plugins/FX/Niagara/Source/NiagaraEditor/Private/ViewModels/NiagaraParameterPanelViewModel.cpp
+++ b/Engine/Plugins/FX/Niagara/Source/NiagaraEditor/Private/ViewModels/NiagaraParameterPanelViewModel.cpp
@@ -2264,7 +2264,10 @@ TArray<TWeakObjectPtr<UNiagaraGraph>> FNiagaraSystemToolkitParameterPanelViewMod
 	else
 	{
 		TArray<TWeakObjectPtr<UNiagaraGraph>> EditableEmitterScriptGraphs;
-		EditableEmitterScriptGraphs.Add(static_cast<UNiagaraScriptSource*>(SystemViewModel->GetEmitterHandleViewModels()[0]->GetEmitterHandle()->GetEmitterData()->GraphSource)->NodeGraph);
+		if (FVersionedNiagaraEmitterData* EmitterData = SystemViewModel->GetEmitterHandleViewModels()[0]->GetEmitterHandle()->GetEmitterData())
+		{
+			EditableEmitterScriptGraphs.Add(static_cast<UNiagaraScriptSource*>(EmitterData->GraphSource)->NodeGraph);
+		}
 		return EditableEmitterScriptGraphs;
 	}
 }
@@ -2405,7 +2408,8 @@ void FNiagaraSystemToolkitParameterPanelViewModel::OnParameterRenamedExternally(
 		const TArray<TSharedRef<FNiagaraEmitterHandleViewModel>>& EmitterHandleViewModels = SystemViewModel->GetEmitterHandleViewModels();
 		for (const TSharedRef<FNiagaraEmitterHandleViewModel>& EmitterHandleViewModel : EmitterHandleViewModels)
 		{
-			UNiagaraGraph* Graph = static_cast<UNiagaraScriptSource*>(EmitterHandleViewModel->GetEmitterHandle()->GetEmitterData()->GraphSource)->NodeGraph;
+			FVersionedNiagaraEmitterData* EmitterData = EmitterHandleViewModel->GetEmitterHandle()->GetEmitterData();
+			UNiagaraGraph* Graph = EmitterData ? CastChecked<UNiagaraScriptSource>(EmitterData->GraphSource)->NodeGraph : nullptr;
 			if (Graph)
 			{ 
 				Graphs.Add(Graph);
@@ -2486,7 +2490,8 @@ void FNiagaraSystemToolkitParameterPanelViewModel::OnParameterRemovedExternally(
 		const TArray<TSharedRef<FNiagaraEmitterHandleViewModel>>& EmitterHandleViewModels = SystemViewModel->GetEmitterHandleViewModels();
 		for (const TSharedRef<FNiagaraEmitterHandleViewModel>& EmitterHandleViewModel : EmitterHandleViewModels)
 		{
-			UNiagaraGraph* Graph = static_cast<UNiagaraScriptSource*>(EmitterHandleViewModel->GetEmitterHandle()->GetEmitterData()->GraphSource)->NodeGraph;
+			FVersionedNiagaraEmitterData* EmitterData = EmitterHandleViewModel->GetEmitterHandle()->GetEmitterData();
+			UNiagaraGraph* Graph = EmitterData && EmitterData->GraphSource ? static_cast<UNiagaraScriptSource*>(EmitterData->GraphSource)->NodeGraph : nullptr;
 			if (Graph)
 				Graphs.Add(Graph);
 
diff --git a/Engine/Plugins/FX/Niagara/Source/NiagaraEditor/Private/ViewModels/NiagaraSystemEditorDocumentsViewModel.cpp b/Engine/Plugins/FX/Niagara/Source/NiagaraEditor/Private/ViewModels/NiagaraSystemEditorDocumentsViewModel.cpp
--- a/Engine/Plugins/FX/Niagara/Source/NiagaraEditor/Private/ViewModels/NiagaraSystemEditorDocumentsViewModel.cpp
+++ b/Engine/Plugins/FX/Niagara/Source/NiagaraEditor/Private/ViewModels/NiagaraSystemEditorDocumentsViewModel.cpp
@@ -170,7 +170,10 @@ TArray<UNiagaraGraph*> UNiagaraSystemEditorDocumentsViewModel::GetEditableGraphs
 	}
 	else
 	{
-		EditableGraphs.Add(Cast<UNiagaraScriptSource>(SystemViewModel->GetEmitterHandleViewModels()[0]->GetEmitterHandle()->GetEmitterData()->GraphSource)->NodeGraph);
+		if (FVersionedNiagaraEmitterData* EmitterData = SystemViewModel->GetEmitterHandleViewModels()[0]->GetEmitterHandle()->GetEmitterData())
+		{
+			EditableGraphs.Add(CastChecked<UNiagaraScriptSource>(EmitterData->GraphSource)->NodeGraph);
+		}
 	}
 
 	return EditableGraphs;
diff --git a/Engine/Plugins/FX/Niagara/Source/NiagaraEditor/Private/ViewModels/NiagaraSystemViewModel.cpp b/Engine/Plugins/FX/Niagara/Source/NiagaraEditor/Private/ViewModels/NiagaraSystemViewModel.cpp
--- a/Engine/Plugins/FX/Niagara/Source/NiagaraEditor/Private/ViewModels/NiagaraSystemViewModel.cpp
+++ b/Engine/Plugins/FX/Niagara/Source/NiagaraEditor/Private/ViewModels/NiagaraSystemViewModel.cpp
@@ -1300,7 +1300,7 @@ UNiagaraEditorParametersAdapter* FNiagaraSystemViewModel::GetEditorOnlyParameter
 	}
 	/** else EditMode == ENiagaraSystemViewModelEditMode::EmitterAsset */
 	FVersionedNiagaraEmitterData* EmitterData = GetEmitterHandleViewModels()[0]->GetEmitterHandle()->GetEmitterData();
-	return CastChecked<UNiagaraEditorParametersAdapter>(EmitterData->GetEditorParameters());
+	return EmitterData ? CastChecked<UNiagaraEditorParametersAdapter>(EmitterData->GetEditorParameters()) : nullptr;
 }
 
 void FNiagaraSystemViewModel::GetOrderedScriptsForEmitterHandleId(FGuid EmitterHandleId, TArray<UNiagaraScript*>& OutScripts)
@@ -1605,7 +1605,10 @@ TArray<UNiagaraGraph*> FNiagaraSystemViewModel::GetSelectedGraphs()
 	}
 	else
 	{
-		SelectedGraphs.Add(Cast<UNiagaraScriptSource>(GetEmitterHandleViewModels()[0]->GetEmitterHandle()->GetEmitterData()->GraphSource)->NodeGraph);
+		if (FVersionedNiagaraEmitterData* EmitterData = GetEmitterHandleViewModels()[0]->GetEmitterHandle()->GetEmitterData())
+		{
+			SelectedGraphs.Add(CastChecked<UNiagaraScriptSource>(EmitterData->GraphSource)->NodeGraph);
+		}
 	}
 
 	TArray<UNiagaraGraph*> EdGraphs = GetDocumentViewModel()->GetEditableGraphsForActiveScriptDocument();
