- Fix binding buffer changes causing dispatches to serialize on some RHIs

#tests Spawn many LW emitters that bind to something like color and vary the binding data per frame, then observe stat gpu

diff --git a/Engine/Plugins/FX/Niagara/Source/Niagara/Private/Stateless/NiagaraStatelessComputeManager.cpp b/Engine/Plugins/FX/Niagara/Source/Niagara/Private/Stateless/NiagaraStatelessComputeManager.cpp
--- a/Engine/Plugins/FX/Niagara/Source/Niagara/Private/Stateless/NiagaraStatelessComputeManager.cpp
+++ b/Engine/Plugins/FX/Niagara/Source/Niagara/Private/Stateless/NiagaraStatelessComputeManager.cpp
@@ -26,6 +26,8 @@ DECLARE_DWORD_COUNTER_STAT(TEXT("CPU Particle Count"), STAT_NiagaraStateless_CPU
 DECLARE_DWORD_COUNTER_STAT(TEXT("GPU Particle Estimate"), STAT_NiagaraStateless_GPUParticleEstimate, STATGROUP_NiagaraStateless);
 DECLARE_MEMORY_STAT(TEXT("Buffer Memory"), STAT_NiagaraStateless_BufferMemory, STATGROUP_NiagaraStateless);
 
+DECLARE_GPU_STAT_NAMED(NiagaraStatelessGPUSimulation, TEXT("Niagara Stateless GPU Simulation"));
+
 namespace NiagaraStatelessComputeManagerPrivate
 {
 	enum class EComputeExecutionPath
@@ -131,6 +133,22 @@ namespace NiagaraStatelessComputeManagerPrivate
 					TransitionsBefore.Emplace(IntBuffer.Buffer, ERHIAccess::SRVMask, ERHIAccess::UAVCompute);
 					TransitionsAfter.Emplace(IntBuffer.Buffer, ERHIAccess::UAVCompute, ERHIAccess::SRVMask);
 				}
+
+				// Flush any pending updates to the binding buffer as this can break parallel dispatches on the GPU
+				const NiagaraStateless::FEmitterInstance_RT* EmitterInstance = GenerationRequest.EmitterInstance;
+				if (EmitterInstance->bBindingBufferDirty)
+				{
+					EmitterInstance->bBindingBufferDirty = false;
+					EmitterInstance->BindingBuffer.Release();
+
+					if (EmitterInstance->BindingBufferData.Num())
+					{
+						EmitterInstance->BindingBuffer.Initialize(RHICmdList, TEXT("FNiagaraStatelessEmitterInstance::BindingBuffer"), sizeof(uint32), EmitterInstance->BindingBufferData.Num() / sizeof(uint32), EPixelFormat::PF_R32_UINT, EBufferUsageFlags::Static);
+						void* LockedBuffer = RHICmdList.LockBuffer(EmitterInstance->BindingBuffer.Buffer, 0, EmitterInstance->BindingBuffer.NumBytes, RLM_WriteOnly);
+						FMemory::Memcpy(LockedBuffer, EmitterInstance->BindingBufferData.GetData(), EmitterInstance->BindingBuffer.NumBytes);
+						RHICmdList.UnlockBuffer(EmitterInstance->BindingBuffer.Buffer);
+					}
+				}
 			}
 		}
 
@@ -140,7 +158,6 @@ namespace NiagaraStatelessComputeManagerPrivate
 
 		// Execute Simulations
 		RHICmdList.Transition(TransitionsBefore);
-
 		RHICmdList.BeginUAVOverlap(CountBufferUAV);
 		for (const FNiagaraStatelessComputeManager::FStatelessDataGenerationRequest& GenerationRequest : GenerationRequests)
 		{
@@ -148,21 +165,6 @@ namespace NiagaraStatelessComputeManagerPrivate
 			const FNiagaraStatelessEmitterData* EmitterData = EmitterInstance->EmitterData.Get();
 			FNiagaraDataBuffer* DestinationData = GenerationRequest.DestinationData;
 
-			// Do we need to update the parameter buffer?
-			if (EmitterInstance->bBindingBufferDirty)
-			{
-				EmitterInstance->bBindingBufferDirty = false;
-				EmitterInstance->BindingBuffer.Release();
-
-				if (EmitterInstance->BindingBufferData.Num())
-				{
-					EmitterInstance->BindingBuffer.Initialize(RHICmdList, TEXT("FNiagaraStatelessEmitterInstance::BindingBuffer"), sizeof(uint32), EmitterInstance->BindingBufferData.Num() / sizeof(uint32), EPixelFormat::PF_R32_UINT, EBufferUsageFlags::Static);
-					void* LockedBuffer = RHICmdList.LockBuffer(EmitterInstance->BindingBuffer.Buffer, 0, EmitterInstance->BindingBuffer.NumBytes, RLM_WriteOnly);
-					FMemory::Memcpy(LockedBuffer, EmitterInstance->BindingBufferData.GetData(), EmitterInstance->BindingBuffer.NumBytes);
-					RHICmdList.UnlockBuffer(EmitterInstance->BindingBuffer.Buffer);
-				}
-			}
-
 			// Update parameters for this compute invocation
 			NiagaraStateless::FCommonShaderParameters* ShaderParameters = EmitterInstance->ShaderParameters.Get();
 			ShaderParameters->Common_SimulationTime = EmitterInstance->Age;
@@ -441,6 +443,7 @@ void FNiagaraStatelessComputeManager::OnPreRender(FRDGBuilder& GraphBuilder)
 	}
 
 	RDG_CSV_STAT_EXCLUSIVE_SCOPE(GraphBuilder, NiagaraStateless);
+	RDG_RHI_EVENT_SCOPE_STAT(GraphBuilder, NiagaraStatelessGPUSimulation, NiagaraStateless);
 	RDG_GPU_MASK_SCOPE(GraphBuilder, FRHIGPUMask::All());
 
 	// Execute dispatches
@@ -467,6 +470,7 @@ void FNiagaraStatelessComputeManager::OnPostPostRender(FRDGBuilder& GraphBuilder
 	}
 
 	RDG_CSV_STAT_EXCLUSIVE_SCOPE(GraphBuilder, NiagaraStateless);
+	RDG_RHI_EVENT_SCOPE_STAT(GraphBuilder, NiagaraStatelessGPUSimulation, NiagaraStateless);
 	RDG_GPU_MASK_SCOPE(GraphBuilder, FRHIGPUMask::All());
 
 	AddPass(
