Fix warning V623: Consider inspecting the '?:' operator. A temporary object of the 'TObjectPtr' type is being created and subsequently destroyed. Check third operand.
#rnx

diff --git a/Engine/Plugins/FX/Niagara/Source/NiagaraEditor/Private/ViewModels/NiagaraParameterPanelViewModel.cpp b/Engine/Plugins/FX/Niagara/Source/NiagaraEditor/Private/ViewModels/NiagaraParameterPanelViewModel.cpp
--- a/Engine/Plugins/FX/Niagara/Source/NiagaraEditor/Private/ViewModels/NiagaraParameterPanelViewModel.cpp
+++ b/Engine/Plugins/FX/Niagara/Source/NiagaraEditor/Private/ViewModels/NiagaraParameterPanelViewModel.cpp
@@ -2491,10 +2491,13 @@ void FNiagaraSystemToolkitParameterPanelViewModel::OnParameterRemovedExternally(
 		for (const TSharedRef<FNiagaraEmitterHandleViewModel>& EmitterHandleViewModel : EmitterHandleViewModels)
 		{
 			FVersionedNiagaraEmitterData* EmitterData = EmitterHandleViewModel->GetEmitterHandle()->GetEmitterData();
-			UNiagaraGraph* Graph = EmitterData && EmitterData->GraphSource ? static_cast<UNiagaraScriptSource*>(EmitterData->GraphSource)->NodeGraph : nullptr;
-			if (Graph)
-				Graphs.Add(Graph);
-
+			if (EmitterData && EmitterData->GraphSource)
+			{
+				if (UNiagaraGraph* Graph = static_cast<UNiagaraScriptSource*>(EmitterData->GraphSource)->NodeGraph)
+				{
+					Graphs.Add(Graph);
+				}
+			}
 		}
 
 		bool bFound = false;
