Made module notes expanded by default and introduced editor setting

diff --git a/Engine/Plugins/FX/Niagara/Source/NiagaraEditor/Private/ViewModels/Stack/NiagaraStackModuleItem.cpp b/Engine/Plugins/FX/Niagara/Source/NiagaraEditor/Private/ViewModels/Stack/NiagaraStackModuleItem.cpp
--- a/Engine/Plugins/FX/Niagara/Source/NiagaraEditor/Private/ViewModels/Stack/NiagaraStackModuleItem.cpp
+++ b/Engine/Plugins/FX/Niagara/Source/NiagaraEditor/Private/ViewModels/Stack/NiagaraStackModuleItem.cpp
@@ -1105,7 +1105,7 @@ void UNiagaraStackModuleItem::RefreshIssues(TArray<FStackIssue>& NewIssues)
 				ScriptData->NoteMessage,
 				GetStackEditorDataKey(),
 				true);
-			NoteIssue.SetIsExpandedByDefault(false);
+			NoteIssue.SetIsExpandedByDefault(GetDefault<UNiagaraEditorSettings>()->bExpandModuleNotesByDefault);
 			NewIssues.Add(NoteIssue);
 		}
 	}
diff --git a/Engine/Plugins/FX/Niagara/Source/NiagaraEditor/Public/NiagaraEditorSettings.h b/Engine/Plugins/FX/Niagara/Source/NiagaraEditor/Public/NiagaraEditorSettings.h
--- a/Engine/Plugins/FX/Niagara/Source/NiagaraEditor/Public/NiagaraEditorSettings.h
+++ b/Engine/Plugins/FX/Niagara/Source/NiagaraEditor/Public/NiagaraEditorSettings.h
@@ -417,6 +417,10 @@ public:
 
 	UPROPERTY(config, EditAnywhere, Category = "Niagara Graph")
 	EGraphZoomLimitHandling ScriptGraphZoomLimitHandling = EGraphZoomLimitHandling::Default;
+
+	/** Module Notes that were defined in a Module Script will expand by default, if true. Consider setting this to false if you are familiar with Niagara and don't need module notes. */
+	UPROPERTY(config, EditAnywhere, Category = "Niagara Stack")
+	bool bExpandModuleNotesByDefault = true;
 	
 	NIAGARAEDITOR_API TArray<float> GetPlaybackSpeeds() const;
 
