- Fix potential OOB reads if a binding is invalid

#tests Create a new NS, add a LWE, add decal attributes module, in initialize particle set position to binding

diff --git a/Engine/Plugins/FX/Niagara/Source/Niagara/Internal/Stateless/NiagaraStatelessEmitterDataBuildContext.h b/Engine/Plugins/FX/Niagara/Source/Niagara/Internal/Stateless/NiagaraStatelessEmitterDataBuildContext.h
--- a/Engine/Plugins/FX/Niagara/Source/Niagara/Internal/Stateless/NiagaraStatelessEmitterDataBuildContext.h
+++ b/Engine/Plugins/FX/Niagara/Source/Niagara/Internal/Stateless/NiagaraStatelessEmitterDataBuildContext.h
@@ -117,6 +117,11 @@ public:
 				FNiagaraStatelessBuiltDistribution::SetIsBinding(BuiltDistribution);
 				FNiagaraStatelessBuiltDistribution::SetLookupParameters(BuiltDistribution, ParameterOffset);
 			}
+			else
+			{
+				const auto DefaultValues = { TDistribution::GetDefaultValue(), TDistribution::GetDefaultValue() };
+				BuiltDistribution = AddDistribution(ENiagaraDistributionMode::NonUniformConstant, MakeArrayView(DefaultValues), FVector2f::ZeroVector);
+			}
 		}
 		else if (Distribution.Mode == ENiagaraDistributionMode::Expression)
 		{
@@ -127,6 +132,11 @@ public:
 				FNiagaraStatelessBuiltDistribution::SetIsBinding(BuiltDistribution);
 				FNiagaraStatelessBuiltDistribution::SetLookupParameters(BuiltDistribution, ParameterOffset);
 			}
+			else
+			{
+				const auto DefaultValues = { TDistribution::GetDefaultValue(), TDistribution::GetDefaultValue() };
+				BuiltDistribution = AddDistribution(ENiagaraDistributionMode::NonUniformConstant, MakeArrayView(DefaultValues), FVector2f::ZeroVector);
+			}
 		}
 		else
 		{
diff --git a/Engine/Plugins/FX/Niagara/Source/Niagara/Public/Stateless/NiagaraStatelessDistribution.h b/Engine/Plugins/FX/Niagara/Source/Niagara/Public/Stateless/NiagaraStatelessDistribution.h
--- a/Engine/Plugins/FX/Niagara/Source/Niagara/Public/Stateless/NiagaraStatelessDistribution.h
+++ b/Engine/Plugins/FX/Niagara/Source/Niagara/Public/Stateless/NiagaraStatelessDistribution.h
@@ -88,6 +88,8 @@ struct FNiagaraDistributionRangeInt
 	FNiagaraDistributionRangeInt() = default;
 	explicit FNiagaraDistributionRangeInt(int32 ConstantValue) { InitConstant(ConstantValue); }
 
+	static constexpr int32 GetDefaultValue() { return 0; }
+
 	UPROPERTY(EditAnywhere, Category = "Parameters")
 	ENiagaraDistributionMode Mode = ENiagaraDistributionMode::UniformConstant;
 
@@ -98,13 +100,13 @@ struct FNiagaraDistributionRangeInt
 	FInstancedStruct ParameterExpression;
 
 	UPROPERTY(EditAnywhere, Category = "Parameters")
-	int32 Min = 0;
+	int32 Min = GetDefaultValue();
 
 	UPROPERTY(EditAnywhere, Category = "Parameters")
-	int32 Max = 0;
+	int32 Max = GetDefaultValue();
 
 	NIAGARA_API void InitConstant(int32 Value);
-	NIAGARA_API FNiagaraStatelessRangeInt CalculateRange(const int32 Default = 0) const;
+	NIAGARA_API FNiagaraStatelessRangeInt CalculateRange(const int32 Default = GetDefaultValue()) const;
 
 	bool IsBinding() const { return Mode == ENiagaraDistributionMode::Binding; }
 	bool IsExpression() const { return Mode == ENiagaraDistributionMode::Expression; }
@@ -128,15 +130,17 @@ struct FNiagaraDistributionRangeFloat : public FNiagaraDistributionBase
 	explicit FNiagaraDistributionRangeFloat(float ConstantValue) { InitConstant(ConstantValue); }
 	explicit FNiagaraDistributionRangeFloat(float MinValue, float MaxValue) { InitRange(MinValue, MaxValue); }
 
+	static constexpr float GetDefaultValue() { return 0.0f; }
+
 	UPROPERTY(EditAnywhere, Category = "Parameters")
-	float Min = 0.0f;
+	float Min = GetDefaultValue();
 
 	UPROPERTY(EditAnywhere, Category = "Parameters")
-	float Max = 0.0f;
+	float Max = GetDefaultValue();
 
 	NIAGARA_API void InitConstant(float Value);
 	NIAGARA_API void InitRange(float MinValue, float MaxValue);
-	NIAGARA_API FNiagaraStatelessRangeFloat CalculateRange(const float Default = 0.0f) const;
+	NIAGARA_API FNiagaraStatelessRangeFloat CalculateRange(const float Default = GetDefaultValue()) const;
 
 #if WITH_EDITORONLY_DATA
 	virtual bool AllowCurves() const override { return false; }
@@ -155,14 +159,16 @@ struct FNiagaraDistributionRangeVector2 : public FNiagaraDistributionBase
 	FNiagaraDistributionRangeVector2() = default;
 	explicit FNiagaraDistributionRangeVector2(const FVector2f& ConstantValue) { InitConstant(ConstantValue); }
 
+	static constexpr FVector2f GetDefaultValue() { return FVector2f::ZeroVector; }
+
 	UPROPERTY(EditAnywhere, Category = "Parameters")
-	FVector2f Min = FVector2f::ZeroVector;
+	FVector2f Min = GetDefaultValue();
 
 	UPROPERTY(EditAnywhere, Category = "Parameters")
-	FVector2f Max = FVector2f::ZeroVector;
+	FVector2f Max = GetDefaultValue();
 
 	NIAGARA_API void InitConstant(const FVector2f& Value);
-	NIAGARA_API FNiagaraStatelessRangeVector2 CalculateRange(const FVector2f& Default = FVector2f::ZeroVector) const;
+	NIAGARA_API FNiagaraStatelessRangeVector2 CalculateRange(const FVector2f& Default = GetDefaultValue()) const;
 
 #if WITH_EDITORONLY_DATA
 	virtual bool AllowCurves() const override { return false; }
@@ -181,14 +187,16 @@ struct FNiagaraDistributionRangeVector3 : public FNiagaraDistributionBase
 	FNiagaraDistributionRangeVector3() = default;
 	explicit FNiagaraDistributionRangeVector3(const FVector3f& ConstantValue) { InitConstant(ConstantValue); }
 
+	static constexpr FVector3f GetDefaultValue() { return FVector3f::ZeroVector; }
+
 	UPROPERTY(EditAnywhere, Category = "Parameters")
-	FVector3f Min = FVector3f::ZeroVector;
+	FVector3f Min = GetDefaultValue();
 
 	UPROPERTY(EditAnywhere, Category = "Parameters")
-	FVector3f Max = FVector3f::ZeroVector;
+	FVector3f Max = GetDefaultValue();
 
 	NIAGARA_API void InitConstant(const FVector3f& Value);
-	NIAGARA_API FNiagaraStatelessRangeVector3 CalculateRange(const FVector3f& Default = FVector3f::ZeroVector) const;
+	NIAGARA_API FNiagaraStatelessRangeVector3 CalculateRange(const FVector3f& Default = GetDefaultValue()) const;
 
 #if WITH_EDITORONLY_DATA
 	virtual bool AllowCurves() const override { return false; }
@@ -207,14 +215,16 @@ struct FNiagaraDistributionRangeColor : public FNiagaraDistributionBase
 	FNiagaraDistributionRangeColor() = default;
 	explicit FNiagaraDistributionRangeColor(const FLinearColor& ConstantValue) { InitConstant(ConstantValue); }
 
+	static constexpr FLinearColor GetDefaultValue() { return FLinearColor::White; }
+
 	UPROPERTY(EditAnywhere, Category = "Parameters")
-	FLinearColor Min = FLinearColor::White;
+	FLinearColor Min = GetDefaultValue();
 
 	UPROPERTY(EditAnywhere, Category = "Parameters")
-	FLinearColor Max = FLinearColor::White;
+	FLinearColor Max = GetDefaultValue();
 
 	NIAGARA_API void InitConstant(const FLinearColor& Value);
-	NIAGARA_API FNiagaraStatelessRangeColor CalculateRange(const FLinearColor& Default = FLinearColor::White) const;
+	NIAGARA_API FNiagaraStatelessRangeColor CalculateRange(const FLinearColor& Default = GetDefaultValue()) const;
 
 #if WITH_EDITORONLY_DATA
 	virtual bool AllowCurves() const override { return false; }
@@ -260,6 +270,8 @@ struct FNiagaraDistributionFloat : public FNiagaraDistributionBase
 	explicit FNiagaraDistributionFloat(float ConstantValue) { InitConstant(ConstantValue); }
 	explicit FNiagaraDistributionFloat(std::initializer_list<float> CurvePoints) { InitCurve(CurvePoints); }
 
+	static constexpr float GetDefaultValue() { return 0.0f; }
+
 	UPROPERTY(EditAnywhere, Category = "Parameters")
 	TArray<float> Values;
 
@@ -271,7 +283,7 @@ struct FNiagaraDistributionFloat : public FNiagaraDistributionBase
 #if WITH_EDITORONLY_DATA
 	NIAGARA_API void InitCurve(const TArray<FRichCurveKey>& CurveKeys);
 #endif
-	NIAGARA_API FNiagaraStatelessRangeFloat CalculateRange(const float Default = 0.0f) const;
+	NIAGARA_API FNiagaraStatelessRangeFloat CalculateRange(const float Default = GetDefaultValue()) const;
 
 #if WITH_EDITORONLY_DATA
 	bool operator==(const FNiagaraDistributionFloat& Other) const
@@ -293,6 +305,8 @@ struct FNiagaraDistributionVector2 : public FNiagaraDistributionBase
 	explicit FNiagaraDistributionVector2(const float ConstantValue) { InitConstant(ConstantValue); }
 	explicit FNiagaraDistributionVector2(const FVector2f& ConstantValue) { InitConstant(ConstantValue); }
 
+	static constexpr FVector2f GetDefaultValue() { return FVector2f::ZeroVector; }
+
 	UPROPERTY(EditAnywhere, Category = "Parameters")
 	TArray<FVector2f> Values;
 
@@ -301,7 +315,7 @@ struct FNiagaraDistributionVector2 : public FNiagaraDistributionBase
 
 	NIAGARA_API void InitConstant(const float Value);
 	NIAGARA_API void InitConstant(const FVector2f& Value);
-	NIAGARA_API FNiagaraStatelessRangeVector2 CalculateRange(const FVector2f& Default = FVector2f::ZeroVector) const;
+	NIAGARA_API FNiagaraStatelessRangeVector2 CalculateRange(const FVector2f& Default = GetDefaultValue()) const;
 
 #if WITH_EDITORONLY_DATA
 	virtual int32 GetBaseNumberOfChannels() const override { return 2; }
@@ -321,6 +335,8 @@ struct FNiagaraDistributionVector3 : public FNiagaraDistributionBase
 	explicit FNiagaraDistributionVector3(std::initializer_list<float> CurvePoints) { InitCurve(CurvePoints); }
 	explicit FNiagaraDistributionVector3(std::initializer_list<FVector3f> CurvePoints) { InitCurve(CurvePoints); }
 
+	static constexpr FVector3f GetDefaultValue() { return FVector3f::ZeroVector; }
+
 	UPROPERTY(EditAnywhere, Category = "Parameters")
 	TArray<FVector3f> Values;
 
@@ -331,7 +347,7 @@ struct FNiagaraDistributionVector3 : public FNiagaraDistributionBase
 	NIAGARA_API void InitConstant(const FVector3f& Value);
 	NIAGARA_API void InitCurve(std::initializer_list<float> CurvePoints);
 	NIAGARA_API void InitCurve(std::initializer_list<FVector3f> CurvePoints);
-	NIAGARA_API FNiagaraStatelessRangeVector3 CalculateRange(const FVector3f& Default = FVector3f::ZeroVector) const;
+	NIAGARA_API FNiagaraStatelessRangeVector3 CalculateRange(const FVector3f& Default = GetDefaultValue()) const;
 
 #if WITH_EDITORONLY_DATA
 	virtual int32 GetBaseNumberOfChannels() const override { return 3; }
@@ -362,6 +378,8 @@ struct FNiagaraDistributionColor : public FNiagaraDistributionBase
 	FNiagaraDistributionColor() = default;
 	explicit FNiagaraDistributionColor(const FLinearColor& ConstantValue) { InitConstant(ConstantValue); }
 
+	static constexpr FLinearColor GetDefaultValue() { return FLinearColor::White; }
+
 	UPROPERTY(EditAnywhere, Category = "Parameters")
 	TArray<FLinearColor> Values;
 
@@ -369,7 +387,7 @@ struct FNiagaraDistributionColor : public FNiagaraDistributionBase
 	FVector2f ValuesTimeRange = FVector2f(0.0f, 1.0f);
 
 	NIAGARA_API void InitConstant(const FLinearColor& Value);
-	NIAGARA_API FNiagaraStatelessRangeColor CalculateRange(const FLinearColor& Default = FLinearColor::White) const;
+	NIAGARA_API FNiagaraStatelessRangeColor CalculateRange(const FLinearColor& Default = GetDefaultValue()) const;
 
 #if WITH_EDITORONLY_DATA
 	virtual bool DisplayAsColor() const override { return true; }
