Fixed dynamic input reassignment if check

diff --git a/Engine/Plugins/FX/Niagara/Source/NiagaraEditor/Private/ViewModels/Stack/NiagaraStackFunctionInput.cpp b/Engine/Plugins/FX/Niagara/Source/NiagaraEditor/Private/ViewModels/Stack/NiagaraStackFunctionInput.cpp
--- a/Engine/Plugins/FX/Niagara/Source/NiagaraEditor/Private/ViewModels/Stack/NiagaraStackFunctionInput.cpp
+++ b/Engine/Plugins/FX/Niagara/Source/NiagaraEditor/Private/ViewModels/Stack/NiagaraStackFunctionInput.cpp
@@ -2897,7 +2897,7 @@ void UNiagaraStackFunctionInput::ReassignDynamicInputScript(UNiagaraScript* Dyna
 			TArray<UNiagaraStackScriptHierarchyRoot*> DynamicInputRoots;
 			GetUnfilteredChildrenOfType(DynamicInputRoots);
 
-			if (ConversionUtility && DynamicInputRoots.Num() == 0)
+			if (ConversionUtility && DynamicInputRoots.Num() == 1)
 			{
 				FText ConvertMessage;
 				bool bConverted = ConversionUtility->Convert(OldScript, OldClipboardContent, DynamicInputScript, DynamicInputRoots[0], NewClipboardContent, InputValues.DynamicNode.Get(), ConvertMessage);
