- Change pin color based on Data Interface vs UObject to reduce confusion

#tests Make a map get node and add both a DI & a UObject input

diff --git a/Engine/Plugins/FX/Niagara/Source/NiagaraEditor/Private/EdGraphSchema_Niagara.cpp b/Engine/Plugins/FX/Niagara/Source/NiagaraEditor/Private/EdGraphSchema_Niagara.cpp
--- a/Engine/Plugins/FX/Niagara/Source/NiagaraEditor/Private/EdGraphSchema_Niagara.cpp
+++ b/Engine/Plugins/FX/Niagara/Source/NiagaraEditor/Private/EdGraphSchema_Niagara.cpp
@@ -1567,7 +1567,9 @@ FLinearColor UEdGraphSchema_Niagara::GetPinTypeColor(const FEdGraphPinType& PinT
 	}
 	if (PinType.PinCategory == PinCategoryClass || PinType.PinCategory == PinCategoryStaticClass)
 	{
-		return Settings->ObjectPinTypeColor;
+		UClass* PinClass = Cast<UClass>(PinType.PinSubCategoryObject.Get());
+		// Note: we use the interface color for UObjects as most Niagara work involves DIs not UObjects and we didn't want to confuse users
+		return PinClass && PinClass->IsChildOf(UNiagaraDataInterface::StaticClass()) ? Settings->ObjectPinTypeColor : Settings->InterfacePinTypeColor;
 	}
 	return Settings->WildcardPinTypeColor;
 }
@@ -1600,10 +1602,15 @@ FLinearColor UEdGraphSchema_Niagara::GetTypeColor(const FNiagaraTypeDefinition&
 	{
 		return Settings->ExecutionPinTypeColor;
 	}
-	else if (Type.IsUObject() || Type.IsDataInterface())
+	else if (Type.IsDataInterface())
 	{
 		return Settings->ObjectPinTypeColor;
 	}
+	else if (Type.IsUObject())
+	{
+		// Note: we use the interface color for UObjects as most Niagara work involves DIs not UObjects and we didn't want to confuse users
+		return Settings->InterfacePinTypeColor;
+	}
 	else if (Type.IsEnum())
 	{
 		return Settings->BytePinTypeColor;
