- Adjust display name for SKM / SM DIs to remove confusion between the asset type

#tests EngineTests + Create a SKM / SM DI and observe name change

diff --git a/Engine/Plugins/FX/Niagara/Source/Niagara/Classes/NiagaraDataInterfaceSkeletalMesh.h b/Engine/Plugins/FX/Niagara/Source/Niagara/Classes/NiagaraDataInterfaceSkeletalMesh.h
--- a/Engine/Plugins/FX/Niagara/Source/Niagara/Classes/NiagaraDataInterfaceSkeletalMesh.h
+++ b/Engine/Plugins/FX/Niagara/Source/Niagara/Classes/NiagaraDataInterfaceSkeletalMesh.h
@@ -684,7 +684,7 @@ struct FNDISkeletalMesh_InstanceData
 };
 
 /** Data Interface allowing sampling of skeletal meshes. */
-UCLASS(EditInlineNew, Category = "Meshes", meta = (DisplayName = "Skeletal Mesh"), MinimalAPI)
+UCLASS(EditInlineNew, Category = "Meshes", meta = (DisplayName = "Skeletal Mesh Reader"), MinimalAPI)
 class UNiagaraDataInterfaceSkeletalMesh : public UNiagaraDataInterface
 {
 	GENERATED_UCLASS_BODY()
diff --git a/Engine/Plugins/FX/Niagara/Source/Niagara/Internal/DataInterface/NiagaraDataInterfaceStaticMesh.h b/Engine/Plugins/FX/Niagara/Source/Niagara/Internal/DataInterface/NiagaraDataInterfaceStaticMesh.h
--- a/Engine/Plugins/FX/Niagara/Source/Niagara/Internal/DataInterface/NiagaraDataInterfaceStaticMesh.h
+++ b/Engine/Plugins/FX/Niagara/Source/Niagara/Internal/DataInterface/NiagaraDataInterfaceStaticMesh.h
@@ -49,7 +49,7 @@ struct FNDIStaticMeshSectionFilter
 };
 
 /** Data Interface allowing sampling of static meshes. */
-UCLASS(EditInlineNew, Category = "Meshes", meta = (DisplayName = "Static Mesh"), MinimalAPI)
+UCLASS(EditInlineNew, Category = "Meshes", meta = (DisplayName = "Static Mesh Reader"), MinimalAPI)
 class UNiagaraDataInterfaceStaticMesh : public UNiagaraDataInterface
 {
 	GENERATED_UCLASS_BODY()
