- Fix a potential use-after-free on world tear down
- If the compute interface is marked pending kill we can ignore sending the command to remove from the computer interface as it will be killed before

#tests EngineTests.  Have a NS which uses a cull proxy, place in the level, move very far away from the system then unload the world

diff --git a/Engine/Plugins/FX/Niagara/Source/Niagara/Private/NiagaraGpuComputeDispatch.cpp b/Engine/Plugins/FX/Niagara/Source/Niagara/Private/NiagaraGpuComputeDispatch.cpp
--- a/Engine/Plugins/FX/Niagara/Source/Niagara/Private/NiagaraGpuComputeDispatch.cpp
+++ b/Engine/Plugins/FX/Niagara/Source/Niagara/Private/NiagaraGpuComputeDispatch.cpp
@@ -546,7 +546,7 @@ void FNiagaraGpuComputeDispatch::FinishDispatches()
 	{
 		for (FNiagaraSystemGpuComputeProxy* ComputeProxy : ProxiesPerStage[iTickStage])
 		{
-			ComputeProxy->ReleaseTicks(GPUInstanceCounterManager, MaxTicksToFlush, bIsLastViewFamily);
+			ComputeProxy->ReleaseTicks(MaxTicksToFlush, bIsLastViewFamily);
 		}
 	}
 
diff --git a/Engine/Plugins/FX/Niagara/Source/Niagara/Private/NiagaraSystemGpuComputeProxy.cpp b/Engine/Plugins/FX/Niagara/Source/Niagara/Private/NiagaraSystemGpuComputeProxy.cpp
--- a/Engine/Plugins/FX/Niagara/Source/Niagara/Private/NiagaraSystemGpuComputeProxy.cpp
+++ b/Engine/Plugins/FX/Niagara/Source/Niagara/Private/NiagaraSystemGpuComputeProxy.cpp
@@ -115,11 +115,20 @@ void FNiagaraSystemGpuComputeProxy::RemoveFromRenderThread(FNiagaraGpuComputeDis
 	check(DebugOwnerComputeDispatchInterface == ComputeDispatchInterface);
 	DebugOwnerComputeDispatchInterface = nullptr;
 
+	// If the interface is pending kill we can't use it in the render command
+	if (ComputeDispatchInterface->IsPendingKill())
+	{
+		ComputeDispatchInterface = nullptr;
+	}
+
 	ENQUEUE_RENDER_COMMAND(RemoveFromRenderThread)(
 		[this, ComputeDispatchInterface, bDeleteProxy](FRHICommandListImmediate& RHICmdList)
 		{
-			ComputeDispatchInterface->RemoveGpuComputeProxy(this);
-			ReleaseTicks(ComputeDispatchInterface->GetGPUInstanceCounterManager(), TNumericLimits<int32>::Max(), true);
+			if (ComputeDispatchInterface)
+			{
+				ComputeDispatchInterface->RemoveGpuComputeProxy(this);
+			}
+			ReleaseTicks(TNumericLimits<int32>::Max(), true);
 
 			for (FNiagaraComputeExecutionContext* ComputeContext : ComputeContexts)
 			{
@@ -151,7 +160,7 @@ void FNiagaraSystemGpuComputeProxy::ClearTicksFromRenderThread(FNiagaraGpuComput
 	ENQUEUE_RENDER_COMMAND(ClearTicksFromProxy)(
 		[this, ComputeDispatchInterface](FRHICommandListImmediate& RHICmdList)
 		{
-			ReleaseTicks(ComputeDispatchInterface->GetGPUInstanceCounterManager(), TNumericLimits<int32>::Max(), true);
+			ReleaseTicks(TNumericLimits<int32>::Max(), true);
 		}
 	);
 }
@@ -187,7 +196,7 @@ void FNiagaraSystemGpuComputeProxy::QueueTick(const FNiagaraGPUSystemTick& Tick)
 	}
 }
 
-void FNiagaraSystemGpuComputeProxy::ReleaseTicks(FNiagaraGPUInstanceCountManager& GPUInstanceCountManager, int32 NumTicksToRelease, bool bLastViewFamily)
+void FNiagaraSystemGpuComputeProxy::ReleaseTicks(int32 NumTicksToRelease, bool bLastViewFamily)
 {
 	check(IsInRenderingThread());
 
diff --git a/Engine/Plugins/FX/Niagara/Source/Niagara/Public/NiagaraSystemGpuComputeProxy.h b/Engine/Plugins/FX/Niagara/Source/Niagara/Public/NiagaraSystemGpuComputeProxy.h
--- a/Engine/Plugins/FX/Niagara/Source/Niagara/Public/NiagaraSystemGpuComputeProxy.h
+++ b/Engine/Plugins/FX/Niagara/Source/Niagara/Public/NiagaraSystemGpuComputeProxy.h
@@ -26,7 +26,7 @@ public:
 	FNiagaraSystemInstanceID GetSystemInstanceID() const { return SystemInstanceID; }
 	ENiagaraGpuComputeTickStage::Type GetComputeTickStage() const { return ComputeTickStage; }
 	void QueueTick(const FNiagaraGPUSystemTick& Tick);
-	void ReleaseTicks(FNiagaraGPUInstanceCountManager& GPUInstanceCountManager, int32 NumTicksToRelease, bool bLastViewFamily);
+	void ReleaseTicks(int32 NumTicksToRelease, bool bLastViewFamily);
 
 	bool RequiresGlobalDistanceField() const { return bRequiresGlobalDistanceField; }
 	bool RequiresDepthBuffer() const { return bRequiresDepthBuffer; }
