- Send the current bounds to the GPU for the emitter properties DI
- This makes the result of GetBounds consistent between both paths when using fixed bounds on the emitter

#tests Create an emitter set bounds to dynamic, read the bounds and observe a different between CPU & GPU

diff --git a/Engine/Plugins/FX/Niagara/Shaders/Private/NiagaraDataInterfaceEmitterPropertiesTemplate.ush b/Engine/Plugins/FX/Niagara/Shaders/Private/NiagaraDataInterfaceEmitterPropertiesTemplate.ush
--- a/Engine/Plugins/FX/Niagara/Shaders/Private/NiagaraDataInterfaceEmitterPropertiesTemplate.ush
+++ b/Engine/Plugins/FX/Niagara/Shaders/Private/NiagaraDataInterfaceEmitterPropertiesTemplate.ush
@@ -1,6 +1,9 @@
 // Copyright Epic Games, Inc. All Rights Reserved.
 
 uint			{ParameterName}_LocalSpace;
+uint			{ParameterName}_BoundsValid;
+float3			{ParameterName}_BoundsMin;
+float3			{ParameterName}_BoundsMax;
 uint			{ParameterName}_FixedBoundsValid;
 float3			{ParameterName}_FixedBoundsMin;
 float3			{ParameterName}_FixedBoundsMax;
@@ -13,9 +16,9 @@ void GetLocalSpace_{ParameterName}(out bool bLocalSpace)
 void GetBounds_{ParameterName}(out bool bValid, out float3 Min, out float3 Max)
 {
 	// GPU is always fixed no dynamic bounds currently
-	bValid = {ParameterName}_FixedBoundsValid != 0;
-	Min = {ParameterName}_FixedBoundsMin;
-	Max = {ParameterName}_FixedBoundsMax;
+	bValid = {ParameterName}_BoundsValid != 0;
+	Min = {ParameterName}_BoundsMin;
+	Max = {ParameterName}_BoundsMax;
 }
 
 void GetFixedBounds_{ParameterName}(out bool bValid, out float3 Min, out float3 Max)
diff --git a/Engine/Plugins/FX/Niagara/Source/Niagara/Private/DataInterface/NiagaraDataInterfaceEmitterProperties.cpp b/Engine/Plugins/FX/Niagara/Source/Niagara/Private/DataInterface/NiagaraDataInterfaceEmitterProperties.cpp
--- a/Engine/Plugins/FX/Niagara/Source/Niagara/Private/DataInterface/NiagaraDataInterfaceEmitterProperties.cpp
+++ b/Engine/Plugins/FX/Niagara/Source/Niagara/Private/DataInterface/NiagaraDataInterfaceEmitterProperties.cpp
@@ -26,6 +26,7 @@ namespace NDIEmitterPropertiesLocal
 
 	struct FInstanceData_GameToRender
 	{
+		FBox	Bounds;
 		FBox	FixedBounds;
 	};
 
@@ -38,6 +39,7 @@ namespace NDIEmitterPropertiesLocal
 	struct FInstanceData_RenderThread
 	{
 		bool	bLocalSpace = false;
+		FBox	Bounds = FBox(EForceInit::ForceInit);
 		FBox	FixedBounds = FBox(EForceInit::ForceInit);
 	};
 
@@ -48,7 +50,11 @@ namespace NDIEmitterPropertiesLocal
 			FInstanceData_GameToRender* InstanceData_ForRT = new(DataForRenderThread) FInstanceData_GameToRender();
 			FInstanceData_GameThread* InstanceData_GT = reinterpret_cast<FInstanceData_GameThread*>(PerInstanceData);
 
-			InstanceData_ForRT->FixedBounds = InstanceData_GT->EmitterInstance ? InstanceData_GT->EmitterInstance->GetFixedBounds() : FBox(EForceInit::ForceInit);
+			if (InstanceData_GT->EmitterInstance)
+			{
+				InstanceData_ForRT->Bounds = InstanceData_GT->EmitterInstance->GetBounds();
+				InstanceData_ForRT->FixedBounds = InstanceData_GT->EmitterInstance->GetFixedBounds();
+			}
 		}
 
 		virtual void ConsumePerInstanceDataFromGameThread(void* PerInstanceData, const FNiagaraSystemInstanceID& InstanceID) override
@@ -56,6 +62,7 @@ namespace NDIEmitterPropertiesLocal
 			FInstanceData_GameToRender* InstanceData_FromGT = static_cast<FInstanceData_GameToRender*>(PerInstanceData);
 			if ( FInstanceData_RenderThread* InstanceData_RT = PerInstanceData_RenderThread.Find(InstanceID) )
 			{
+				InstanceData_RT->Bounds = InstanceData_FromGT->Bounds;
 				InstanceData_RT->FixedBounds = InstanceData_FromGT->FixedBounds;
 			}
 			InstanceData_FromGT->~FInstanceData_GameToRender();
@@ -318,6 +325,9 @@ void UNiagaraDataInterfaceEmitterProperties::SetShaderParameters(const FNiagaraD
 
 	FShaderParameters* ShaderParameters	= Context.GetParameterNestedStruct<FShaderParameters>();
 	ShaderParameters->LocalSpace		= InstanceData_RT.bLocalSpace ? 1 : 0;
+	ShaderParameters->BoundsValid		= InstanceData_RT.Bounds.IsValid;
+	ShaderParameters->BoundsMin			= FVector3f(InstanceData_RT.Bounds.Min);
+	ShaderParameters->BoundsMax			= FVector3f(InstanceData_RT.Bounds.Max);
 	ShaderParameters->FixedBoundsValid	= InstanceData_RT.FixedBounds.IsValid;
 	ShaderParameters->FixedBoundsMin	= FVector3f(InstanceData_RT.FixedBounds.Min);
 	ShaderParameters->FixedBoundsMax	= FVector3f(InstanceData_RT.FixedBounds.Max);
diff --git a/Engine/Plugins/FX/Niagara/Source/Niagara/Private/DataInterface/NiagaraDataInterfaceEmitterProperties.h b/Engine/Plugins/FX/Niagara/Source/Niagara/Private/DataInterface/NiagaraDataInterfaceEmitterProperties.h
--- a/Engine/Plugins/FX/Niagara/Source/Niagara/Private/DataInterface/NiagaraDataInterfaceEmitterProperties.h
+++ b/Engine/Plugins/FX/Niagara/Source/Niagara/Private/DataInterface/NiagaraDataInterfaceEmitterProperties.h
@@ -22,6 +22,9 @@ class UNiagaraDataInterfaceEmitterProperties : public UNiagaraDataInterface
 
 	BEGIN_SHADER_PARAMETER_STRUCT(FShaderParameters, )
 		SHADER_PARAMETER(uint32,	LocalSpace)
+		SHADER_PARAMETER(uint32,	BoundsValid)
+		SHADER_PARAMETER(FVector3f,	BoundsMin)
+		SHADER_PARAMETER(FVector3f,	BoundsMax)
 		SHADER_PARAMETER(uint32,	FixedBoundsValid)
 		SHADER_PARAMETER(FVector3f,	FixedBoundsMin)
 		SHADER_PARAMETER(FVector3f,	FixedBoundsMax)
