- Add missing postload calls for stateless emitters

#tests Add a LWE, bind a system color to the renderer, run asset validation

diff --git a/Engine/Plugins/FX/Niagara/Source/Niagara/Private/NiagaraEmitterHandle.cpp b/Engine/Plugins/FX/Niagara/Source/Niagara/Private/NiagaraEmitterHandle.cpp
--- a/Engine/Plugins/FX/Niagara/Source/Niagara/Private/NiagaraEmitterHandle.cpp
+++ b/Engine/Plugins/FX/Niagara/Source/Niagara/Private/NiagaraEmitterHandle.cpp
@@ -340,6 +340,10 @@ void FNiagaraEmitterHandle::ConditionalPostLoad(int32 NiagaraCustomVersion)
 			}
 		}
 	}
+	if (StatelessEmitter != nullptr)
+	{
+		StatelessEmitter->ConditionalPostLoad();
+	}
 }
 
 bool FNiagaraEmitterHandle::UsesEmitter(const FVersionedNiagaraEmitter& InEmitter) const
diff --git a/Engine/Plugins/FX/Niagara/Source/Niagara/Private/NiagaraSystem.cpp b/Engine/Plugins/FX/Niagara/Source/Niagara/Private/NiagaraSystem.cpp
--- a/Engine/Plugins/FX/Niagara/Source/Niagara/Private/NiagaraSystem.cpp
+++ b/Engine/Plugins/FX/Niagara/Source/Niagara/Private/NiagaraSystem.cpp
@@ -1322,7 +1322,7 @@ void UNiagaraSystem::PostLoad()
 		GetObjectsWithPackage(PackageOuter, ObjectsInPackage);
 		for (UObject* ObjectInPackage : ObjectsInPackage)
 		{
-			UNiagaraEmitter* Emitter = Cast<UNiagaraEmitter>(ObjectInPackage);
+			UNiagaraEmitterBase* Emitter = Cast<UNiagaraEmitterBase>(ObjectInPackage);
 			if (Emitter != nullptr)
 			{
 				Emitter->ConditionalPostLoad();
@@ -1364,7 +1364,7 @@ void UNiagaraSystem::PostLoad()
 #if WITH_EDITORONLY_DATA
 		EmitterHandle.ConditionalPostLoad(NiagaraVer);
 #else
-		if (UNiagaraEmitter* NiagaraEmitter = EmitterHandle.GetInstance().Emitter)
+		if (UNiagaraEmitterBase* NiagaraEmitter = EmitterHandle.GetEmitterBase())
 		{
 			NiagaraEmitter->ConditionalPostLoad();
 		}
