- LWE Shape Location Changes
- Add shape scale
- Fix plane shape not applying rotation
- Add surface mode to cylinder

#tests Run ET, plus test various shapes with rotation & scaling applied

diff --git a/Engine/Plugins/FX/Niagara/Shaders/Private/Stateless/Modules/NiagaraStatelessModule_ShapeLocation.ush b/Engine/Plugins/FX/Niagara/Shaders/Private/Stateless/Modules/NiagaraStatelessModule_ShapeLocation.ush
--- a/Engine/Plugins/FX/Niagara/Shaders/Private/Stateless/Modules/NiagaraStatelessModule_ShapeLocation.ush
+++ b/Engine/Plugins/FX/Niagara/Shaders/Private/Stateless/Modules/NiagaraStatelessModule_ShapeLocation.ush
@@ -49,13 +49,29 @@ float3 ShapeLocation_GetLocation()
 		const float3 AxisY		= ShapeLocation_Parameters1.xyz;
 		const float3 AxisZ		= ShapeLocation_Parameters2.xyz;
 
+		const bool bOnSurface	= (ShapeLocation_Mode.y & 0x1) != 0;
 		const float4 Random = RandomFloat4(0);
-		const float2 XYLocation = SafeNormalize(Random.xy - 0.5f) * Radius * Random.z;
+
+		float XYOffset = Radius * Random.z;
+		float ZOffset = Random.w * HeightScale + HeightBias;
+		if (bOnSurface)
+		{
+			const bool bOnSurfaceIncludeEndCaps = (ShapeLocation_Mode.y & 0x2) != 0;
+			const float ThicknessScale = asfloat(ShapeLocation_Mode.z);
+			const float ThicknessBias = asfloat(ShapeLocation_Mode.w);
+
+			const uint Edge = (RandomUInt(1) & 0x3) + (bOnSurfaceIncludeEndCaps ? 0 : 2);
+			ZOffset  = Edge == 0 ? HeightBias - ThicknessBias - (ThicknessScale * Random.w) : ZOffset;
+			ZOffset  = Edge == 1 ? HeightBias + HeightScale + ThicknessBias + (ThicknessScale * Random.w) : ZOffset;
+			XYOffset = Edge >= 2 ? Radius + ThicknessBias + (ThicknessScale * Random.z) : (Radius + ThicknessBias + ThicknessScale) * Random.z;
+		}
+
+		const float2 XYLocation = SafeNormalize(Random.xy - 0.5f) * XYOffset;
 
 		float3 Location;
 		Location  = AxisX * XYLocation.x;
 		Location += AxisY * XYLocation.y;
-		Location += AxisZ * (Random.w * HeightScale + HeightBias);
+		Location += AxisZ * ZOffset;
 		return Location;
 	}
 
@@ -78,11 +94,15 @@ float3 ShapeLocation_GetLocation()
 
 	// ENSM_ShapePrimitive::Sphere:
 	{
-		const float SphereScale	= ShapeLocation_Parameters0.x;
-		const float SphereBias	= ShapeLocation_Parameters0.y;
+		const float SphereScale	= ShapeLocation_Parameters0.w;
+		const float SphereBias	= ShapeLocation_Parameters1.w;
+		const float3 AxisX		= ShapeLocation_Parameters0.xyz;
+		const float3 AxisY		= ShapeLocation_Parameters1.xyz;
+		const float3 AxisZ		= ShapeLocation_Parameters2.xyz;
 
 		const float3 Vector = RandomUnitFloat3(0);
-		return Vector * RandomScaleBiasFloat(1, SphereScale, SphereBias);
+		const float3 Location = Vector * RandomScaleBiasFloat(1, SphereScale, SphereBias);
+		return (Location.x * AxisX) + (Location.y * AxisY) + (Location.z * AxisZ);
 	}
 }
 
diff --git a/Engine/Plugins/FX/Niagara/Source/Niagara/Internal/Stateless/Modules/NiagaraStatelessModule_ShapeLocation.h b/Engine/Plugins/FX/Niagara/Source/Niagara/Internal/Stateless/Modules/NiagaraStatelessModule_ShapeLocation.h
--- a/Engine/Plugins/FX/Niagara/Source/Niagara/Internal/Stateless/Modules/NiagaraStatelessModule_ShapeLocation.h
+++ b/Engine/Plugins/FX/Niagara/Source/Niagara/Internal/Stateless/Modules/NiagaraStatelessModule_ShapeLocation.h
@@ -62,6 +62,14 @@ public:
 	FNiagaraDistributionRangeFloat CylinderRadius = FNiagaraDistributionRangeFloat(100.0f);
 	UPROPERTY(EditAnywhere, Category = "Parameters", meta = (DisableRangeDistribution, EditConditionHides, EditCondition = "ShapePrimitive == ENSM_ShapePrimitive::Cylinder"))
 	FNiagaraDistributionRangeFloat CylinderHeightMidpoint = FNiagaraDistributionRangeFloat(0.5f);
+	UPROPERTY(EditAnywhere, Category = "Parameters", meta = (EditConditionHides, EditCondition = "ShapePrimitive == ENSM_ShapePrimitive::Cylinder"))
+	bool bCylinderSurfaceOnly = false;
+	UPROPERTY(EditAnywhere, Category = "Parameters", meta = (EditConditionHides, EditCondition = "ShapePrimitive == ENSM_ShapePrimitive::Cylinder && bCylinderSurfaceOnly"))
+	bool bCylinderSurfaceOnlyIncludeEndCaps = true;
+	UPROPERTY(EditAnywhere, Category = "Parameters", meta = (EditConditionHides, EditCondition = "ShapePrimitive == ENSM_ShapePrimitive::Cylinder && bCylinderSurfaceOnly"))
+	ENSM_SurfaceExpansionMode CylinderSurfaceExpansion = ENSM_SurfaceExpansionMode::Centered;
+	UPROPERTY(EditAnywhere, Category = "Parameters", meta = (DisableRangeDistribution, EditConditionHides, EditCondition = "ShapePrimitive == ENSM_ShapePrimitive::Cylinder && bCylinderSurfaceOnly"))
+	FNiagaraDistributionRangeFloat CylinderSurfaceThickness = FNiagaraDistributionRangeFloat(0.0f);
 
 	UPROPERTY(EditAnywhere, Category = "Parameters", meta = (DisableRangeDistribution, EditConditionHides, EditCondition = "ShapePrimitive == ENSM_ShapePrimitive::Ring"))
 	FNiagaraDistributionRangeFloat RingRadius = FNiagaraDistributionRangeFloat(100.0f);
@@ -79,6 +87,9 @@ public:
 	UPROPERTY(EditAnywhere, Category = "Parameters", meta = (DisableRangeDistribution, DisableUniformDistribution))
 	FNiagaraDistributionRangeRotator ShapeRotation;
 
+	UPROPERTY(EditAnywhere, Category = "Parameters", meta = (DisableRangeDistribution, DisableUniformDistribution))
+	FNiagaraDistributionRangeVector3 ShapeScale = FNiagaraDistributionRangeVector3(FVector3f::OneVector);
+
 private:
 #if WITH_EDITORONLY_DATA
 	UPROPERTY()
diff --git a/Engine/Plugins/FX/Niagara/Source/Niagara/Internal/Stateless/NiagaraStatelessDrawDebugContext.h b/Engine/Plugins/FX/Niagara/Source/Niagara/Internal/Stateless/NiagaraStatelessDrawDebugContext.h
--- a/Engine/Plugins/FX/Niagara/Source/Niagara/Internal/Stateless/NiagaraStatelessDrawDebugContext.h
+++ b/Engine/Plugins/FX/Niagara/Source/Niagara/Internal/Stateless/NiagaraStatelessDrawDebugContext.h
@@ -15,10 +15,30 @@ public:
 	NIAGARA_API void DrawArrow(const FVector& Origin, const FVector& DirectionWithLength, const FColor& Color = FColor::White) const;
 	NIAGARA_API void DrawBox(const FVector& Center, const FQuat& Rotation, const FVector& Extent, const FColor& Color = FColor::White) const;
 	NIAGARA_API void DrawCone(const FVector& Origin, const FQuat& Rotation, float Angle, float Length, const FColor& Color = FColor::White) const;
-	NIAGARA_API void DrawCylinder(const FVector& Center, const FQuat& Rotation, float CylinderHeight, float CylinderRadius, float CylinderHeightMidpoint, const FColor& Color = FColor::White) const;
-	NIAGARA_API void DrawCircle(const FVector& Center, const FQuat& Rotation, const float Radius, const FColor& Color = FColor::White) const;
-	NIAGARA_API void DrawRing(const FVector& Center, const FQuat& Rotation, const float InnerRadius, const float OuterRadius, float Distribution, const FColor& Color = FColor::White) const;
-	NIAGARA_API void DrawSphere(const FVector& Center, const float Radius, const FColor& Color = FColor::White) const;
+	NIAGARA_API void DrawCylinder(const FVector& Center, const FQuat& Rotation, const FVector& Scale, float CylinderHeight, float CylinderRadius, float CylinderHeightMidpoint, const FColor& Color = FColor::White) const;
+	NIAGARA_API void DrawCircle(const FVector& Center, const FQuat& Rotation, const FVector& Scale, const float Radius, const FColor& Color = FColor::White) const;
+	NIAGARA_API void DrawRing(const FVector& Center, const FQuat& Rotation, const FVector& Scale, const float InnerRadius, const float OuterRadius, float Distribution, const FColor& Color = FColor::White) const;
+	NIAGARA_API void DrawSphere(const FVector& Center, const FQuat& Rotation, const FVector& Scale, const float Radius, const FColor& Color = FColor::White) const;
+
+	void DrawCylinder(const FVector& Center, const FQuat& Rotation, float CylinderHeight, float CylinderRadius, float CylinderHeightMidpoint, const FColor& Color = FColor::White) const
+	{
+		DrawCylinder(Center, Rotation, FVector::OneVector, CylinderHeight, CylinderRadius, CylinderHeightMidpoint, Color);
+	}
+
+	void DrawCircle(const FVector& Center, const FQuat& Rotation, const float Radius, const FColor& Color = FColor::White) const
+	{
+		DrawCircle(Center, Rotation, FVector::OneVector, Radius, Color);
+	}
+
+	void DrawRing(const FVector& Center, const FQuat& Rotation, const float InnerRadius, const float OuterRadius, float Distribution, const FColor& Color = FColor::White) const
+	{
+		DrawRing(Center, Rotation, FVector::OneVector, InnerRadius, OuterRadius, Distribution, Color);
+	}
+
+	void DrawSphere(const FVector& Center, const float Radius, const FColor& Color = FColor::White) const
+	{
+		DrawSphere(Center, FQuat::Identity, FVector::OneVector, Radius, Color);
+	}
 
 	// Helper functions to transform into World space for debug rendering
 	NIAGARA_API FQuat   TransformRotation(ENiagaraCoordinateSpace SourceSpace, const FQuat4f& Rotation) const;
diff --git a/Engine/Plugins/FX/Niagara/Source/Niagara/Private/Stateless/Modules/NiagaraStatelessModule_ShapeLocation.cpp b/Engine/Plugins/FX/Niagara/Source/Niagara/Private/Stateless/Modules/NiagaraStatelessModule_ShapeLocation.cpp
--- a/Engine/Plugins/FX/Niagara/Source/Niagara/Private/Stateless/Modules/NiagaraStatelessModule_ShapeLocation.cpp
+++ b/Engine/Plugins/FX/Niagara/Source/Niagara/Private/Stateless/Modules/NiagaraStatelessModule_ShapeLocation.cpp
@@ -15,6 +15,7 @@ namespace NSMShapeLocationPrivate
 		ENSM_ShapePrimitive				ShapePrimitive = ENSM_ShapePrimitive::Max;
 		ENiagaraCoordinateSpace			CoordinateSpace = ENiagaraCoordinateSpace::Simulation;
 		FNiagaraStatelessRangeRotator	ShapeRotation;
+		FNiagaraStatelessRangeVector3	ShapeScale;
 		int32							PositionVariableOffset = INDEX_NONE;
 		int32							PreviousPositionVariableOffset = INDEX_NONE;
 	};
@@ -32,6 +33,10 @@ namespace NSMShapeLocationPrivate
 		FNiagaraStatelessRangeFloat	Height;
 		FNiagaraStatelessRangeFloat	Radius;
 		FNiagaraStatelessRangeFloat	MidPoint;
+		bool						bDoSurfaceOnly = false;
+		bool						bDoSurfaceOnlyIncludeEndCaps = true;
+		ENSM_SurfaceExpansionMode	SurfaceExpansionMode = ENSM_SurfaceExpansionMode::Centered;
+		FNiagaraStatelessRangeFloat	SurfaceThickness;
 	};
 
 	struct FModuleBuiltDataPlane
@@ -67,6 +72,19 @@ namespace NSMShapeLocationPrivate
 		return MakeTuple(0.0f, 0.0f);
 	}
 
+	TTuple<float, float> CalculateCylinderThicknessParameters(ENSM_SurfaceExpansionMode ExpansionMode, float Thickness)
+	{
+		const float HThickness = Thickness * 0.5f;
+		switch (ExpansionMode)
+		{
+			case ENSM_SurfaceExpansionMode::Inner:		return MakeTuple(Thickness, -Thickness);
+			case ENSM_SurfaceExpansionMode::Centered:	return MakeTuple(Thickness, -HThickness);
+			case ENSM_SurfaceExpansionMode::Outside:	return MakeTuple(Thickness,	0.0f);
+		}
+		checkNoEntry();
+		return MakeTuple(0.0f, 0.0f);
+	}
+
 	FVector3f ShapeLocation_GetLocation(const NiagaraStateless::FParticleSimulationContext& ParticleSimulationContext, const UNiagaraStatelessModule_ShapeLocation::FParameters* ShaderParameters, uint32 iInstance)
 	{
 		// ENSM_ShapePrimitive::Box | ENSM_ShapePrimitive::Plane
@@ -112,13 +130,28 @@ namespace NSMShapeLocationPrivate
 			const FVector3f AxisY	= FVector3f(ShaderParameters->ShapeLocation_Parameters1.X, ShaderParameters->ShapeLocation_Parameters1.Y, ShaderParameters->ShapeLocation_Parameters1.Z);
 			const FVector3f AxisZ	= FVector3f(ShaderParameters->ShapeLocation_Parameters2.X, ShaderParameters->ShapeLocation_Parameters2.Y, ShaderParameters->ShapeLocation_Parameters2.Z);
 
+			const bool bOnSurface = (ShaderParameters->ShapeLocation_Mode.Y & 0x1) != 0;
 			const FVector4f Random = ParticleSimulationContext.RandomFloat4(iInstance, 0);
-			const FVector2f XYLocation = ParticleSimulationContext.SafeNormalize(FVector2f(Random.X - 0.5f, Random.Y - 0.5f)) * Radius * Random.Z;
+
+			float XYOffset = Radius * Random.Z;
+			float ZOffset = Random.W * HeightScale + HeightBias;
+			if (bOnSurface)
+			{
+				const bool bOnSurfaceIncludeEndCaps = (ShaderParameters->ShapeLocation_Mode.Y & 0x2) != 0;
+				const float ThicknessScale = reinterpret_cast<const float&>(ShaderParameters->ShapeLocation_Mode.Z);
+				const float ThicknessBias = reinterpret_cast<const float&>(ShaderParameters->ShapeLocation_Mode.W);
+
+				const uint32 Edge = (ParticleSimulationContext.RandomUInt(iInstance, 1) & 0x3) + (bOnSurfaceIncludeEndCaps ? 0 : 2);
+				ZOffset = Edge == 0 ? HeightBias - ThicknessBias - (ThicknessScale * Random.W) : ZOffset;
+				ZOffset = Edge == 1 ? HeightBias + HeightScale + ThicknessBias + (ThicknessScale * Random.W) : ZOffset;
+				XYOffset = Edge >= 2 ? Radius + ThicknessBias + (ThicknessScale * Random.Z) : (Radius + ThicknessBias + ThicknessScale) * Random.Z;
+			}
+			const FVector2f XYLocation = ParticleSimulationContext.SafeNormalize(FVector2f(Random.X - 0.5f, Random.Y - 0.5f)) * XYOffset;
 
 			FVector3f Location;
 			Location  = AxisX * XYLocation.X;
 			Location += AxisY * XYLocation.Y;
-			Location += AxisZ * (Random.W * HeightScale + HeightBias);
+			Location += AxisZ * ZOffset;
 			return Location;
 		}
 	
@@ -142,11 +175,15 @@ namespace NSMShapeLocationPrivate
 	
 		// ENSM_ShapePrimitive::Sphere:
 		{
-			const float SphereScale	= ShaderParameters->ShapeLocation_Parameters0.X;
-			const float SphereBias	= ShaderParameters->ShapeLocation_Parameters0.Y;
-	
+			const float SphereScale	= ShaderParameters->ShapeLocation_Parameters0.W;
+			const float SphereBias	= ShaderParameters->ShapeLocation_Parameters1.W;
+			const FVector3f AxisX = FVector3f(ShaderParameters->ShapeLocation_Parameters0.X, ShaderParameters->ShapeLocation_Parameters0.Y, ShaderParameters->ShapeLocation_Parameters0.Z);
+			const FVector3f AxisY = FVector3f(ShaderParameters->ShapeLocation_Parameters1.X, ShaderParameters->ShapeLocation_Parameters1.Y, ShaderParameters->ShapeLocation_Parameters1.Z);
+			const FVector3f AxisZ = FVector3f(ShaderParameters->ShapeLocation_Parameters2.X, ShaderParameters->ShapeLocation_Parameters2.Y, ShaderParameters->ShapeLocation_Parameters2.Z);
+
 			const FVector3f Vector = ParticleSimulationContext.RandomUnitFloat3(iInstance, 0);
-			return Vector * ParticleSimulationContext.RandomScaleBiasFloat(iInstance, 1, SphereScale, SphereBias);
+			const FVector3f Location = Vector * ParticleSimulationContext.RandomScaleBiasFloat(iInstance, 1, SphereScale, SphereBias);
+			return (Location.X * AxisX) + (Location.Y * AxisY) + (Location.Z * AxisZ);
 		}
 	}
 
@@ -190,6 +227,7 @@ void UNiagaraStatelessModule_ShapeLocation::BuildEmitterData(const FNiagaraState
 	BuiltData->ShapePrimitive	= ShapePrimitive;
 	BuiltData->CoordinateSpace	= CoordinateSpace;
 	BuiltData->ShapeRotation	= BuildContext.ConvertDistributionToRange(ShapeRotation, FRotator3f::ZeroRotator);
+	BuiltData->ShapeScale		= BuildContext.ConvertDistributionToRange(ShapeScale, FVector3f::OneVector);
 
 	switch (ShapePrimitive)
 	{
@@ -208,6 +246,11 @@ void UNiagaraStatelessModule_ShapeLocation::BuildEmitterData(const FNiagaraState
 			BuiltDataCylinder->Height	= BuildContext.ConvertDistributionToRange(CylinderHeight, 0.0f);
 			BuiltDataCylinder->Radius	= BuildContext.ConvertDistributionToRange(CylinderRadius, 0.0f);
 			BuiltDataCylinder->MidPoint	= BuildContext.ConvertDistributionToRange(CylinderHeightMidpoint, 0.0f);
+
+			BuiltDataCylinder->bDoSurfaceOnly				= bCylinderSurfaceOnly;
+			BuiltDataCylinder->bDoSurfaceOnlyIncludeEndCaps	= bCylinderSurfaceOnlyIncludeEndCaps;
+			BuiltDataCylinder->SurfaceExpansionMode			= CylinderSurfaceExpansion;
+			BuiltDataCylinder->SurfaceThickness				= BuildContext.ConvertDistributionToRange(CylinderSurfaceThickness, 0.0f);
 			break;
 		}
 		case ENSM_ShapePrimitive::Plane:
@@ -256,6 +299,8 @@ void UNiagaraStatelessModule_ShapeLocation::SetShaderParameters(const FNiagaraSt
 	const FQuat4f ResolvedShapeRotation = SetShaderParameterContext.GetRendererParameterValue(ModuleBuiltData->ShapeRotation.ParameterOffset, ModuleBuiltData->ShapeRotation.Min.Quaternion());
 	const FQuat4f Rotation = SetShaderParameterContext.GetSpaceTransforms().TransformRotation(ModuleBuiltData->CoordinateSpace, ResolvedShapeRotation);
 
+	const FVector3f Scale = SetShaderParameterContext.ConvertRangeToValue(ModuleBuiltData->ShapeScale);
+
 	Parameters->ShapeLocation_Mode			= FUintVector4::ZeroValue;
 	Parameters->ShapeLocation_Parameters0	= FVector4f::Zero();
 	Parameters->ShapeLocation_Parameters1	= FVector4f::Zero();
@@ -266,7 +311,7 @@ void UNiagaraStatelessModule_ShapeLocation::SetShaderParameters(const FNiagaraSt
 		case ENSM_ShapePrimitive::Box:
 		{
 			const FModuleBuiltDataBox* BuiltDataBox = SetShaderParameterContext.ReadBuiltData<FModuleBuiltDataBox>();
-			const FVector3f Size	= SetShaderParameterContext.ConvertRangeToValue(BuiltDataBox->Size);
+			const FVector3f Size	= SetShaderParameterContext.ConvertRangeToValue(BuiltDataBox->Size) * Scale;
 			const float Thickness	= SetShaderParameterContext.ConvertRangeToValue(BuiltDataBox->SurfaceThickness);
 			const TTuple<float, float> ThicknessParameters = CalculateThicknessParameters(BuiltDataBox->SurfaceExpansionMode, Thickness);
 
@@ -285,17 +330,23 @@ void UNiagaraStatelessModule_ShapeLocation::SetShaderParameters(const FNiagaraSt
 			const float Height		= SetShaderParameterContext.ConvertRangeToValue(BuiltDataCylinder->Height);
 			const float Radius		= SetShaderParameterContext.ConvertRangeToValue(BuiltDataCylinder->Radius);
 			const float MidPoint	= SetShaderParameterContext.ConvertRangeToValue(BuiltDataCylinder->MidPoint);
+			const float Thickness	= SetShaderParameterContext.ConvertRangeToValue(BuiltDataCylinder->SurfaceThickness);
+			const TTuple<float, float> ThicknessParameters = CalculateCylinderThicknessParameters(BuiltDataCylinder->SurfaceExpansionMode, Thickness);
 
 			Parameters->ShapeLocation_Mode.X		= 1;
-			Parameters->ShapeLocation_Parameters0	= FVector4f(Rotation.GetAxisX(), Height);
-			Parameters->ShapeLocation_Parameters1	= FVector4f(Rotation.GetAxisY(), Height * -MidPoint);
-			Parameters->ShapeLocation_Parameters2	= FVector4f(Rotation.GetAxisZ(), Radius);
+			Parameters->ShapeLocation_Mode.Y		= (bCylinderSurfaceOnly ? 1 : 0) | (bCylinderSurfaceOnlyIncludeEndCaps ? 2 : 0);
+			reinterpret_cast<float&>(Parameters->ShapeLocation_Mode.Z) = ThicknessParameters.Key;
+			reinterpret_cast<float&>(Parameters->ShapeLocation_Mode.W) = ThicknessParameters.Value;
+
+			Parameters->ShapeLocation_Parameters0	= FVector4f(Rotation.GetAxisX() * Scale.X, Height);
+			Parameters->ShapeLocation_Parameters1	= FVector4f(Rotation.GetAxisY() * Scale.Y, Height * -MidPoint);
+			Parameters->ShapeLocation_Parameters2	= FVector4f(Rotation.GetAxisZ() * Scale.Z, Radius);
 			break;
 		}
 		case ENSM_ShapePrimitive::Plane:
 		{
 			const FModuleBuiltDataPlane* BuiltDataPlane = SetShaderParameterContext.ReadBuiltData<FModuleBuiltDataPlane>();
-			const FVector2f Size	= SetShaderParameterContext.ConvertRangeToValue(BuiltDataPlane->Size);
+			const FVector2f Size	= SetShaderParameterContext.ConvertRangeToValue(BuiltDataPlane->Size) * FVector2f(Scale.X, Scale.Y);
 			const float Thickness	= SetShaderParameterContext.ConvertRangeToValue(BuiltDataPlane->EdgeThickness);
 			const TTuple<float, float> ThicknessParameters = CalculateThicknessParameters(BuiltDataPlane->EdgeExpansionMode, Thickness);
 
@@ -304,6 +355,7 @@ void UNiagaraStatelessModule_ShapeLocation::SetShaderParameters(const FNiagaraSt
 			Parameters->ShapeLocation_Mode.Z		= 4;
 			Parameters->ShapeLocation_Parameters0	= FVector4f(Size.X + ThicknessParameters.Key, Size.Y + ThicknessParameters.Key, 0.0f, 0.0f);
 			Parameters->ShapeLocation_Parameters1	= FVector4f(ThicknessParameters.Value, ThicknessParameters.Value, 0.0f, 0.0f);
+			Parameters->ShapeLocation_Parameters2	= FVector4f(Rotation.X, Rotation.Y, Rotation.Z, Rotation.W);
 			break;
 		}
 		case ENSM_ShapePrimitive::Ring:
@@ -321,8 +373,8 @@ void UNiagaraStatelessModule_ShapeLocation::SetShaderParameters(const FNiagaraSt
 			Parameters->ShapeLocation_Parameters0.Y = Radius * SDC;
 			Parameters->ShapeLocation_Parameters0.Z = -UE_TWO_PI * (1.0f - Distribution);
 			Parameters->ShapeLocation_Parameters0.W = 0.0f;
-			Parameters->ShapeLocation_Parameters1 = FVector4f(Rotation.GetAxisX());
-			Parameters->ShapeLocation_Parameters2 = FVector4f(Rotation.GetAxisY());
+			Parameters->ShapeLocation_Parameters1 = FVector4f(Rotation.GetAxisX() * Scale.X);
+			Parameters->ShapeLocation_Parameters2 = FVector4f(Rotation.GetAxisY() * Scale.Y);
 			break;
 		}
 		case ENSM_ShapePrimitive::Sphere:
@@ -333,8 +385,9 @@ void UNiagaraStatelessModule_ShapeLocation::SetShaderParameters(const FNiagaraSt
 			SetShaderParameterContext.ConvertRangeToScaleBias(BuiltDataSphere->Radius, SphereScale, SphereBias);
 
 			Parameters->ShapeLocation_Mode.X		= 3;
-			Parameters->ShapeLocation_Parameters0.X = SphereScale;
-			Parameters->ShapeLocation_Parameters0.Y = SphereBias;
+			Parameters->ShapeLocation_Parameters0	= FVector4f(Rotation.GetAxisX() * Scale.X, SphereScale);
+			Parameters->ShapeLocation_Parameters1	= FVector4f(Rotation.GetAxisY() * Scale.Y, SphereBias);
+			Parameters->ShapeLocation_Parameters2	= FVector4f(Rotation.GetAxisZ() * Scale.Z, 0.0f);
 			break;
 		}
 		default:
@@ -345,7 +398,7 @@ void UNiagaraStatelessModule_ShapeLocation::SetShaderParameters(const FNiagaraSt
 #if WITH_EDITOR
 void UNiagaraStatelessModule_ShapeLocation::DrawDebug(const FNiagaraStatelessDrawDebugContext& DrawDebugContext) const
 {
-	if (ShapeRotation.IsBinding())
+	if (ShapeRotation.IsBinding() || ShapeScale.IsBinding())
 	{
 		return;
 	}
@@ -353,13 +406,15 @@ void UNiagaraStatelessModule_ShapeLocation::DrawDebug(const FNiagaraStatelessDra
 	const FQuat4f ResolvedRotation = ShapeRotation.CalculateRange().Min.Quaternion();
 	const FVector DrawPosition = DrawDebugContext.TransformPosition(FVector3f::ZeroVector);
 	const FQuat DrawRotation = DrawDebugContext.TransformRotation(CoordinateSpace, ResolvedRotation);
+	const FVector DrawScale = FVector(ShapeScale.CalculateRange().Min);
+
 	switch (ShapePrimitive)
 	{
 		case ENSM_ShapePrimitive::Box:
 		{
 			if (BoxSize.IsBinding() == false)
 			{
-				DrawDebugContext.DrawBox(DrawPosition, DrawRotation, FVector(BoxSize.Min * 0.5f));
+				DrawDebugContext.DrawBox(DrawPosition, DrawRotation, FVector(BoxSize.Min * 0.5f) * DrawScale);
 			}
 			break;
 		}
@@ -367,7 +422,7 @@ void UNiagaraStatelessModule_ShapeLocation::DrawDebug(const FNiagaraStatelessDra
 		{
 			if (PlaneSize.IsBinding() == false)
 			{
-				DrawDebugContext.DrawBox(DrawPosition, DrawRotation, FVector(PlaneSize.Min.X * 0.5f, PlaneSize.Min.Y * 0.5f, 0.0f));
+				DrawDebugContext.DrawBox(DrawPosition, DrawRotation, FVector(PlaneSize.Min.X * 0.5f, PlaneSize.Min.Y * 0.5f, 0.0f) * DrawScale);
 			}
 			break;
 		}
@@ -375,7 +430,7 @@ void UNiagaraStatelessModule_ShapeLocation::DrawDebug(const FNiagaraStatelessDra
 		{
 			if (CylinderHeight.IsBinding() == false && CylinderRadius.IsBinding() == false && CylinderHeightMidpoint.IsBinding() == false)
 			{
-				DrawDebugContext.DrawCylinder(DrawPosition, DrawRotation, CylinderHeight.Min, CylinderRadius.Min, CylinderHeightMidpoint.Min);
+				DrawDebugContext.DrawCylinder(DrawPosition, DrawRotation, DrawScale, CylinderHeight.Min, CylinderRadius.Min, CylinderHeightMidpoint.Min);
 			}
 			break;
 		}
@@ -385,7 +440,7 @@ void UNiagaraStatelessModule_ShapeLocation::DrawDebug(const FNiagaraStatelessDra
 			{
 				const float DC = FMath::Clamp(1.0f - DiscCoverage.Min, 0.0f, 1.0f);
 				const float SDC = DC > 0.0f ? FMath::Sqrt(DC) : 0.0f;
-				DrawDebugContext.DrawRing(DrawPosition, DrawRotation, RingRadius.Min * SDC, RingRadius.Min, RingUDistribution.Min);
+				DrawDebugContext.DrawRing(DrawPosition, DrawRotation, DrawScale, RingRadius.Min * SDC, RingRadius.Min, RingUDistribution.Min);
 			}
 			break;
 		}
@@ -393,10 +448,10 @@ void UNiagaraStatelessModule_ShapeLocation::DrawDebug(const FNiagaraStatelessDra
 		{
 			if (SphereRadius.IsBinding() == false)
 			{
-				DrawDebugContext.DrawSphere(DrawPosition, SphereRadius.Min);
+				DrawDebugContext.DrawSphere(DrawPosition, DrawRotation, DrawScale, SphereRadius.Min);
 				if (SphereRadius.IsRange())
 				{
-					DrawDebugContext.DrawSphere(DrawPosition, SphereRadius.Max);
+					DrawDebugContext.DrawSphere(DrawPosition, DrawRotation, DrawScale, SphereRadius.Max);
 				}
 			}
 			break;
diff --git a/Engine/Plugins/FX/Niagara/Source/Niagara/Private/Stateless/NiagaraStatelessDrawDebugContext.cpp b/Engine/Plugins/FX/Niagara/Source/Niagara/Private/Stateless/NiagaraStatelessDrawDebugContext.cpp
--- a/Engine/Plugins/FX/Niagara/Source/Niagara/Private/Stateless/NiagaraStatelessDrawDebugContext.cpp
+++ b/Engine/Plugins/FX/Niagara/Source/Niagara/Private/Stateless/NiagaraStatelessDrawDebugContext.cpp
@@ -64,55 +64,79 @@ void FNiagaraStatelessDrawDebugContext::DrawCone(const FVector& Origin, const FQ
 	);
 }
 
-void FNiagaraStatelessDrawDebugContext::DrawCylinder(const FVector& Center, const FQuat& Rotation, float CylinderHeight, float CylinderRadius, float CylinderHeightMidpoint, const FColor& Color) const
+void FNiagaraStatelessDrawDebugContext::DrawCylinder(const FVector& Center, const FQuat& Rotation, const FVector& Scale, float CylinderHeight, float CylinderRadius, float CylinderHeightMidpoint, const FColor& Color) const
 {
-	const FVector Axis = Rotation.RotateVector(FVector::ZAxisVector);
+	ULineBatchComponent* LineBatcher = World->GetLineBatcher(UWorld::ELineBatcherType::World);
+	if (LineBatcher == nullptr)
+	{
+		return;
+	}
+
+	constexpr int NumSegments = 12;
+	constexpr float AngleStep = 2.f * UE_PI / float(NumSegments);
+
+	const FVector AxisX = Rotation.RotateVector(FVector::XAxisVector) * Scale.X * CylinderRadius;
+	const FVector AxisY = Rotation.RotateVector(FVector::YAxisVector) * Scale.Y * CylinderRadius;
+	const FVector AxisZ = Rotation.RotateVector(FVector::ZAxisVector) * Scale.Z;
+
 	const float Offset = -CylinderHeightMidpoint * CylinderHeight;
+	const FVector Start = Center + (AxisZ * Offset);
+	const FVector End = Center + (AxisZ * (Offset + CylinderHeight));
 
-	DrawDebugCylinder(
-		World,
-		Center + (Axis * Offset),
-		Center + (Axis * (Offset + CylinderHeight)),
-		CylinderRadius,
-		12,
-		Color
-	);
+	TArray<FBatchedLine> Lines;
+	Lines.Empty(NumSegments * 3);
+
+	float Angle = 0.f;
+	for (int iSegment=0; iSegment < NumSegments; ++iSegment)
+	{
+		const FVector Offset1 = (AxisX * FMath::Cos(Angle) + AxisY * FMath::Sin(Angle));
+		Angle += AngleStep;
+		const FVector Offset2 = (AxisX * FMath::Cos(Angle) + AxisY * FMath::Sin(Angle));
+
+		Lines.Emplace(Start + Offset1, Start + Offset2, Color, -1.0f, 0.0f, 0);
+		Lines.Emplace(Start + Offset1, End + Offset1, Color, -1.0f, 0.0f, 0);
+		Lines.Emplace(End + Offset1, End + Offset2, Color, -1.0f, 0.0f, 0);
+	}
+	LineBatcher->DrawLines(Lines);
 }
 
-void FNiagaraStatelessDrawDebugContext::DrawCircle(const FVector& Center, const FQuat& Rotation, const float Radius, const FColor& Color) const
+void FNiagaraStatelessDrawDebugContext::DrawCircle(const FVector& Center, const FQuat& Rotation, const FVector& Scale, const float Radius, const FColor& Color) const
 {
-	if (ULineBatchComponent* LineBatcher = World->GetLineBatcher(UWorld::ELineBatcherType::World))
+	ULineBatchComponent* LineBatcher = World->GetLineBatcher(UWorld::ELineBatcherType::World);
+	if (LineBatcher == nullptr)
 	{
-		int32 Segments = 16;
-		const float AngleStep = 2.f * UE_PI / float(Segments);
+		return;
+	}
 
-		const FVector AxisX = Rotation.RotateVector(FVector::XAxisVector);
-		const FVector AxisY = Rotation.RotateVector(FVector::YAxisVector);
+	int32 Segments = 16;
+	const float AngleStep = 2.f * UE_PI / float(Segments);
 
-		TArray<FBatchedLine> Lines;
-		Lines.Empty(Segments);
+	const FVector AxisX = Rotation.RotateVector(FVector::XAxisVector) * Scale.X;
+	const FVector AxisY = Rotation.RotateVector(FVector::YAxisVector) * Scale.Y;
 
-		float Angle = 0.f;
-		while (Segments--)
-		{
-			const FVector Vertex1 = Center + Radius * (AxisX * FMath::Cos(Angle) + AxisY * FMath::Sin(Angle));
-			Angle += AngleStep;
-			const FVector Vertex2 = Center + Radius * (AxisX * FMath::Cos(Angle) + AxisY * FMath::Sin(Angle));
-			Lines.Emplace(Vertex1, Vertex2, Color, -1.0f, 0.0f, 0);
-		}
-		LineBatcher->DrawLines(Lines);
+	TArray<FBatchedLine> Lines;
+	Lines.Empty(Segments);
+
+	float Angle = 0.f;
+	while (Segments--)
+	{
+		const FVector Vertex1 = Center + Radius * (AxisX * FMath::Cos(Angle) + AxisY * FMath::Sin(Angle));
+		Angle += AngleStep;
+		const FVector Vertex2 = Center + Radius * (AxisX * FMath::Cos(Angle) + AxisY * FMath::Sin(Angle));
+		Lines.Emplace(Vertex1, Vertex2, Color, -1.0f, 0.0f, 0);
 	}
+	LineBatcher->DrawLines(Lines);
 }
 
-void FNiagaraStatelessDrawDebugContext::DrawRing(const FVector& Center, const FQuat& Rotation, const float InnerRadius, const float OuterRadius, float Distribution, const FColor& Color) const
+void FNiagaraStatelessDrawDebugContext::DrawRing(const FVector& Center, const FQuat& Rotation, const FVector& Scale, const float InnerRadius, const float OuterRadius, float Distribution, const FColor& Color) const
 {
 	if (ULineBatchComponent* LineBatcher = World->GetLineBatcher(UWorld::ELineBatcherType::World))
 	{
 		const int32 NumSegments = 16;
 		const float AngleStep = -2.f * UE_PI * FMath::Clamp(1.0f - Distribution, 0.0f, 1.0f) / float(NumSegments);
 
-		const FVector AxisX = Rotation.RotateVector(FVector::XAxisVector);
-		const FVector AxisY = Rotation.RotateVector(FVector::YAxisVector);
+		const FVector AxisX = Rotation.RotateVector(FVector::XAxisVector) * Scale.X;
+		const FVector AxisY = Rotation.RotateVector(FVector::YAxisVector) * Scale.Y;
 
 		TArray<FBatchedLine> Lines;
 		Lines.Empty(2 + (NumSegments * 2));
@@ -143,15 +167,64 @@ void FNiagaraStatelessDrawDebugContext::DrawRing(const FVector& Center, const FQ
 	}
 }
 
-void FNiagaraStatelessDrawDebugContext::DrawSphere(const FVector& Center, const float Radius, const FColor& Color) const
+void FNiagaraStatelessDrawDebugContext::DrawSphere(const FVector& Center, const FQuat& Rotation, const FVector& Scale, const float Radius, const FColor& Color) const
 {
-	DrawDebugSphere(
-		World,
-		Center,
-		Radius,
-		16,
-		Color
-	);
+	ULineBatchComponent* LineBatcher = World->GetLineBatcher(UWorld::ELineBatcherType::World);
+	if (LineBatcher == nullptr)
+	{
+		return;
+	}
+
+	constexpr int NumSegments = 16;
+
+	const FVector AxisX = Rotation.RotateVector(FVector::XAxisVector) * Scale.X;
+	const FVector AxisY = Rotation.RotateVector(FVector::YAxisVector) * Scale.Y;
+	const FVector AxisZ = Rotation.RotateVector(FVector::ZAxisVector) * Scale.Z;
+
+	const float AngleInc = 2.f * UE_PI / NumSegments;
+	int32 NumSegmentsY = NumSegments;
+	float Latitude = AngleInc;
+	float SinY1 = 0.0f, CosY1 = 1.0f;
+
+	TArray<FBatchedLine> Lines;
+	Lines.Empty(NumSegments * NumSegments);
+
+	while (NumSegmentsY--)
+	{
+		const float SinY2 = FMath::Sin(Latitude);
+		const float CosY2 = FMath::Cos(Latitude);
+
+		FVector Vertex1 = FVector(SinY1, 0.0f, CosY1) * Radius + Center;
+		FVector Vertex3 = FVector(SinY2, 0.0f, CosY2) * Radius + Center;
+		Vertex1 = (AxisX * Vertex1.X) + (AxisY * Vertex1.Y) + (AxisZ * Vertex1.Z);
+		Vertex3 = (AxisX * Vertex3.X) + (AxisY * Vertex3.Y) + (AxisZ * Vertex3.Z);
+
+		float Longitude = AngleInc;
+
+		int32 NumSegmentsX = NumSegments;
+		while (NumSegmentsX--)
+		{
+			const float SinX = FMath::Sin(Longitude);
+			const float CosX = FMath::Cos(Longitude);
+
+			FVector Vertex2 = FVector((CosX * SinY1), (SinX * SinY1), CosY1) * Radius + Center;
+			FVector Vertex4 = FVector((CosX * SinY2), (SinX * SinY2), CosY2) * Radius + Center;
+			Vertex2 = (AxisX * Vertex2.X) + (AxisY * Vertex2.Y) + (AxisZ * Vertex2.Z);
+			Vertex4 = (AxisX * Vertex4.X) + (AxisY * Vertex4.Y) + (AxisZ * Vertex4.Z);
+
+			Lines.Emplace(Vertex1, Vertex2, Color, -1.0f, 0.0f, 0);
+			Lines.Emplace(Vertex1, Vertex3, Color, -1.0f, 0.0f, 0);
+
+			Vertex1 = Vertex2;
+			Vertex3 = Vertex4;
+			Longitude += AngleInc;
+		}
+		SinY1 = SinY2;
+		CosY1 = CosY2;
+		Latitude += AngleInc;
+	}
+
+	LineBatcher->DrawLines(Lines);
 }
 
 FQuat FNiagaraStatelessDrawDebugContext::TransformRotation(ENiagaraCoordinateSpace SourceSpace, const FQuat4f& Rotation) const
