- Read from built data not class data

#rnx

diff --git a/Engine/Plugins/FX/Niagara/Source/Niagara/Private/Stateless/Modules/NiagaraStatelessModule_ShapeLocation.cpp b/Engine/Plugins/FX/Niagara/Source/Niagara/Private/Stateless/Modules/NiagaraStatelessModule_ShapeLocation.cpp
--- a/Engine/Plugins/FX/Niagara/Source/Niagara/Private/Stateless/Modules/NiagaraStatelessModule_ShapeLocation.cpp
+++ b/Engine/Plugins/FX/Niagara/Source/Niagara/Private/Stateless/Modules/NiagaraStatelessModule_ShapeLocation.cpp
@@ -334,7 +334,7 @@ void UNiagaraStatelessModule_ShapeLocation::SetShaderParameters(const FNiagaraSt
 			const TTuple<float, float> ThicknessParameters = CalculateCylinderThicknessParameters(BuiltDataCylinder->SurfaceExpansionMode, Thickness);
 
 			Parameters->ShapeLocation_Mode.X		= 1;
-			Parameters->ShapeLocation_Mode.Y		= (bCylinderSurfaceOnly ? 1 : 0) | (bCylinderSurfaceOnlyIncludeEndCaps ? 2 : 0);
+			Parameters->ShapeLocation_Mode.Y		= (BuiltDataCylinder->bDoSurfaceOnly ? 1 : 0) | (BuiltDataCylinder->bDoSurfaceOnlyIncludeEndCaps ? 2 : 0);
 			reinterpret_cast<float&>(Parameters->ShapeLocation_Mode.Z) = ThicknessParameters.Key;
 			reinterpret_cast<float&>(Parameters->ShapeLocation_Mode.W) = ThicknessParameters.Value;
 
