Fix for deleted stateless emitter module being rest to an invalid state leading to a check()
#tests Added Scale Sprite Size module to a lightweight emitter, delete the module and then re-add a Scale Sprite Size module, a check would fire

diff --git a/Engine/Plugins/FX/Niagara/Source/NiagaraEditor/Private/ViewModels/Stack/NiagaraStackStatelessEmitterSimulateGroup.cpp b/Engine/Plugins/FX/Niagara/Source/NiagaraEditor/Private/ViewModels/Stack/NiagaraStackStatelessEmitterSimulateGroup.cpp
--- a/Engine/Plugins/FX/Niagara/Source/NiagaraEditor/Private/ViewModels/Stack/NiagaraStackStatelessEmitterSimulateGroup.cpp
+++ b/Engine/Plugins/FX/Niagara/Source/NiagaraEditor/Private/ViewModels/Stack/NiagaraStackStatelessEmitterSimulateGroup.cpp
@@ -356,6 +356,9 @@ void UNiagaraStackStatelessModuleItem::Delete()
 		CopyParams.bDoDelta = false;
 		UEngine::CopyPropertiesForUnrelatedObjects(StatelessModule->GetClass()->GetDefaultObject(), StatelessModule, CopyParams);
 
+		// we want to reset the module to a valid initial state which may require any changes done in PostInitProperties()
+		StatelessModule->PostInitProperties();
+
 		StatelessModule->SetIsModuleEnabled(false);
 		StatelessModule->PostEditChange();
 		OnDataObjectModified().Broadcast({ StatelessModule }, ENiagaraDataObjectChange::Changed);
