Fixing a bug where the async trace DI would always incorrectly use the ray count from the first NS using the DI rather than it's own ray count when using either the inline raytracing path or the "fake indirect dispatch" path

This fixes an issue where the collision module would either stutter or sometime have completely incorrect collisions (go through surfaces or collide on invisible surfaces)

#tests Validated async traces worked with multiple instances of generators

diff --git a/Engine/Plugins/FX/Niagara/Shaders/Private/NiagaraRayTracingShaders.usf b/Engine/Plugins/FX/Niagara/Shaders/Private/NiagaraRayTracingShaders.usf
--- a/Engine/Plugins/FX/Niagara/Shaders/Private/NiagaraRayTracingShaders.usf
+++ b/Engine/Plugins/FX/Niagara/Shaders/Private/NiagaraRayTracingShaders.usf
@@ -16,7 +16,8 @@ RWStructuredBuffer<FNiagaraAsyncGpuTraceResult>	CollisionOutput;
 uint											CollisionOutputOffset;
 
 #if NIAGARA_RAYTRACE_FAKE_INDIRECT || COMPUTESHADER
-Buffer<uint>                                  RayTraceCounts;
+Buffer<uint> RayTraceCounts;
+uint RayTraceCountsOffset;
 #endif
 
 Buffer<uint> HashToCollisionGroups;
@@ -48,7 +49,7 @@ RAY_TRACING_ENTRY_RAYGEN(NiagaraCollisionRayTraceRG)
 #endif
 
 #if NIAGARA_RAYTRACE_FAKE_INDIRECT || COMPUTESHADER
-	if (RayIndex >= RayTraceCounts[0]) 
+	if (RayIndex >= RayTraceCounts[RayTraceCountsOffset]) 
 	{
 		return;
 	}
diff --git a/Engine/Plugins/FX/Niagara/Source/NiagaraShader/Private/NiagaraAsyncGpuTraceProviderHwrt.cpp b/Engine/Plugins/FX/Niagara/Source/NiagaraShader/Private/NiagaraAsyncGpuTraceProviderHwrt.cpp
--- a/Engine/Plugins/FX/Niagara/Source/NiagaraShader/Private/NiagaraAsyncGpuTraceProviderHwrt.cpp
+++ b/Engine/Plugins/FX/Niagara/Source/NiagaraShader/Private/NiagaraAsyncGpuTraceProviderHwrt.cpp
@@ -77,6 +77,7 @@ public:
 		SHADER_PARAMETER_UAV(Buffer<FNiagaraAsyncGpuTraceResult>, CollisionOutput)
 		SHADER_PARAMETER(uint32, CollisionOutputOffset)
 		SHADER_PARAMETER_SRV(Buffer<UINT>, RayTraceCounts)
+		SHADER_PARAMETER(uint32, RayTraceCountsOffset)
 		SHADER_PARAMETER(uint32, MaxRetraces)
 		END_SHADER_PARAMETER_STRUCT()
 };
@@ -437,6 +438,7 @@ void FNiagaraAsyncGpuTraceProviderHwrt::IssueTraces(FRHICommandList& RHICmdList,
 	Params.CollisionOutput = Request.ResultsBuffer->UAV;
 	Params.CollisionOutputOffset = Request.ResultsOffset;
 	Params.MaxRetraces = Request.MaxRetraceCount;
+	Params.RayTraceCountsOffset = Request.TraceCountsOffset;
 
 	if (SupportsNiagaraAsyncGpuTraceHwrtInline())
 	{
