virtual texture sample DI

diff --git a/Engine/Plugins/FX/Niagara/Shaders/Private/NiagaraDataInterfaceVirtualTextureSampleTemplate.ush b/Engine/Plugins/FX/Niagara/Shaders/Private/NiagaraDataInterfaceVirtualTextureSampleTemplate.ush
new file mode 100644
--- /dev/null
+++ b/Engine/Plugins/FX/Niagara/Shaders/Private/NiagaraDataInterfaceVirtualTextureSampleTemplate.ush
@@ -0,0 +1,41 @@
+// Copyright Epic Games, Inc. All Rights Reserved.
+
+#define NUM_VIRTUALTEXTURE_SAMPLES 0
+// Workaround for unbound View uniform buffer
+#define VT_DISABLE_VIEW_UNIFORM_BUFFER 1
+
+#include "/Engine/Private/VirtualTextureCommon.ush"
+
+int2			{ParameterName}_TextureSize;
+int				{ParameterName}_MipLevels;
+
+Texture2D<uint4> {ParameterName}_PageTableTexture0;
+Texture2D<uint4> {ParameterName}_PageTableTexture1;
+Texture2D {ParameterName}_PhysicalTexture;
+
+uint4 {ParameterName}_VTPackedPageTableUniform[2];
+uint4 {ParameterName}_VTPackedUniform;
+
+SamplerState {ParameterName}_TextureSampler;
+
+void SampleTexture2D_{ParameterName}(in float2 UV, in float MipLevel, out float4 OutValue)
+{
+	uint LayerIndex = 0;
+	float4 Sample;
+
+	VTPageTableResult PageTableResult = TextureLoadVirtualPageTableLevel({ParameterName}_PageTableTexture0, {ParameterName}_PageTableTexture1, VTPageTableUniform_Unpack({ParameterName}_VTPackedPageTableUniform[0], {ParameterName}_VTPackedPageTableUniform[1]), UV, VTADDRESSMODE_CLAMP, VTADDRESSMODE_CLAMP, MipLevel);
+	OutValue = TextureVirtualSample( {ParameterName}_PhysicalTexture, {ParameterName}_TextureSampler, PageTableResult, LayerIndex, VTUniform_Unpack({ParameterName}_VTPackedUniform));	 
+}
+
+void GetTextureDimensions_{ParameterName}(int MipLevel, out float2 OutValue)
+{
+	OutValue.x = float(max({ParameterName}_TextureSize.x >> MipLevel, 1));
+	OutValue.y = float(max({ParameterName}_TextureSize.y >> MipLevel, 1));
+}
+
+void GetNumMipLevels_{ParameterName}(out int OutMipLevels)
+{
+	OutMipLevels = {ParameterName}_MipLevels;
+}
+
+
diff --git a/Engine/Plugins/FX/Niagara/Source/Niagara/Classes/NiagaraDataInterfaceVirtualTextureSample.h b/Engine/Plugins/FX/Niagara/Source/Niagara/Classes/NiagaraDataInterfaceVirtualTextureSample.h
new file mode 100644
--- /dev/null
+++ b/Engine/Plugins/FX/Niagara/Source/Niagara/Classes/NiagaraDataInterfaceVirtualTextureSample.h
@@ -0,0 +1,77 @@
+// Copyright Epic Games, Inc. All Rights Reserved.
+
+#pragma once
+
+#include "NiagaraDataInterface.h"
+#include "NiagaraCommon.h"
+#include "VectorVM.h"
+#include "NiagaraDataInterfaceVirtualTextureSample.generated.h"
+
+/** Data Interface allowing sampling of a texture */
+UCLASS(EditInlineNew, Category = "Texture", CollapseCategories, meta = (DisplayName = "Virtual Texture Sample", Experimental), MinimalAPI)
+class UNiagaraDataInterfaceVirtualTextureSample : public UNiagaraDataInterface
+{
+	GENERATED_UCLASS_BODY()
+
+	BEGIN_SHADER_PARAMETER_STRUCT(FShaderParameters, )
+		SHADER_PARAMETER(FIntPoint,				TextureSize)
+		SHADER_PARAMETER(int32,					MipLevels)				
+		SHADER_PARAMETER_TEXTURE(Texture2D,		PageTableTexture0)
+		SHADER_PARAMETER_TEXTURE(Texture2D,		PageTableTexture1)
+		SHADER_PARAMETER_SRV(Texture2D,			PhysicalTexture)
+		SHADER_PARAMETER_ARRAY(FUintVector4,	VTPackedPageTableUniform, [2])
+		SHADER_PARAMETER(FUintVector4,			VTPackedUniform)
+		SHADER_PARAMETER_SAMPLER(SamplerState,	TextureSampler)
+	END_SHADER_PARAMETER_STRUCT()
+
+public:
+	UPROPERTY(EditAnywhere, Category = "Texture")
+	TObjectPtr<UTexture> Texture;
+
+	UPROPERTY(EditAnywhere, Category = "Texture", meta = (ToolTip = "When valid the user parameter is used as the texture rather than the one on the data interface"))
+	FNiagaraUserParameterBinding TextureUserParameter;
+
+	//UObject Interface
+	NIAGARA_API virtual void PostInitProperties()override;
+	virtual bool CanBeInCluster() const override { return false; }	// Note: Due to BP functionality we can change a UObject property on this DI we can not put into a cluster
+	//UObject Interface End
+
+	//UNiagaraDataInterface Interface
+	NIAGARA_API virtual void GetVMExternalFunction(const FVMExternalFunctionBindingInfo& BindingInfo, void* InstanceData, FVMExternalFunction &OutFunc) override;
+	virtual bool CanExecuteOnTarget(ENiagaraSimTarget Target)const override { return Target==ENiagaraSimTarget::GPUComputeSim; }
+
+	NIAGARA_API virtual int32 PerInstanceDataSize() const override;
+	NIAGARA_API virtual bool InitPerInstanceData(void* PerInstanceData, FNiagaraSystemInstance* SystemInstance) override;
+	NIAGARA_API virtual void DestroyPerInstanceData(void* PerInstanceData, FNiagaraSystemInstance* SystemInstance) override;
+
+	virtual bool HasPreSimulateTick() const override { return true; }
+	NIAGARA_API virtual bool PerInstanceTick(void* PerInstanceData, FNiagaraSystemInstance* SystemInstance, float DeltaSeconds) override;
+	//UNiagaraDataInterface Interface End
+
+	NIAGARA_API void VMGetTextureDimensions(FVectorVMExternalFunctionContext& Context);
+	NIAGARA_API void VMGetNumMipLevels(FVectorVMExternalFunctionContext& Context);
+
+	NIAGARA_API virtual bool Equals(const UNiagaraDataInterface* Other) const override;
+
+	// GPU sim functionality
+#if WITH_EDITORONLY_DATA
+	NIAGARA_API virtual bool AppendCompileHash(FNiagaraCompileHashVisitor* InVisitor) const override;
+	NIAGARA_API virtual void GetParameterDefinitionHLSL(const FNiagaraDataInterfaceGPUParamInfo& ParamInfo, FString& OutHLSL) override;
+	NIAGARA_API virtual bool GetFunctionHLSL(const FNiagaraDataInterfaceGPUParamInfo& ParamInfo, const FNiagaraDataInterfaceGeneratedFunction& FunctionInfo, int FunctionInstanceIndex, FString& OutHLSL) override;	
+#endif
+	NIAGARA_API virtual void BuildShaderParameters(FNiagaraShaderParametersBuilder& ShaderParametersBuilder) const override;
+	NIAGARA_API virtual void SetShaderParameters(const FNiagaraDataInterfaceSetShaderParametersContext& Context) const override;
+
+protected:
+#if WITH_EDITORONLY_DATA
+	NIAGARA_API virtual void GetFunctionsInternal(TArray<FNiagaraFunctionSignature>& OutFunctions) const override;
+#endif
+	NIAGARA_API virtual bool CopyToInternal(UNiagaraDataInterface* Destination) const override;
+
+protected:
+	static NIAGARA_API const TCHAR* TemplateShaderFilePath;
+	static NIAGARA_API const FName LoadTexture2DName;	
+	static NIAGARA_API const FName SampleTexture2DName;	
+	static NIAGARA_API const FName GetTextureDimensionsName;
+	static NIAGARA_API const FName GetNumMipLevelsName;
+};
diff --git a/Engine/Plugins/FX/Niagara/Source/Niagara/Private/NiagaraDataInterfaceVirtualTextureSample.cpp b/Engine/Plugins/FX/Niagara/Source/Niagara/Private/NiagaraDataInterfaceVirtualTextureSample.cpp
new file mode 100644
--- /dev/null
+++ b/Engine/Plugins/FX/Niagara/Source/Niagara/Private/NiagaraDataInterfaceVirtualTextureSample.cpp
@@ -0,0 +1,402 @@
+// Copyright Epic Games, Inc. All Rights Reserved.
+
+#include "NiagaraDataInterfaceVirtualTextureSample.h"
+#include "Engine/Texture.h"
+#include "Engine/Texture2D.h"
+#include "Engine/World.h"
+#include "EngineModule.h"
+#include "Engine/TextureRenderTarget2D.h"
+#include "NiagaraRenderer.h"
+#include "NiagaraCompileHashVisitor.h"
+#include "NiagaraCustomVersion.h"
+#include "NiagaraGpuComputeDispatchInterface.h"
+#include "NiagaraShaderParametersBuilder.h"
+#include "NiagaraSystemInstance.h"
+#include "RenderGraphBuilder.h"
+#include "GlobalRenderResources.h"
+#include "RHIStaticStates.h"
+#include "RenderGraphUtils.h"
+#include "ShaderCompilerCore.h"
+#include "TextureResource.h"
+
+#include UE_INLINE_GENERATED_CPP_BY_NAME(NiagaraDataInterfaceVirtualTextureSample)
+
+#define LOCTEXT_NAMESPACE "UNiagaraDataInterfaceVirtualTextureSample"
+
+const TCHAR* UNiagaraDataInterfaceVirtualTextureSample::TemplateShaderFilePath = TEXT("/Plugin/FX/Niagara/Private/NiagaraDataInterfaceVirtualTextureSampleTemplate.ush");
+const FName UNiagaraDataInterfaceVirtualTextureSample::SampleTexture2DName(TEXT("SampleTexture2D"));
+const FName UNiagaraDataInterfaceVirtualTextureSample::GetTextureDimensionsName(TEXT("GetTextureDimensions"));
+const FName UNiagaraDataInterfaceVirtualTextureSample::GetNumMipLevelsName(TEXT("GetNumMipLevels"));
+
+struct FNDIVirtualTextureInstanceData_GameThread
+{
+	TWeakObjectPtr<UTexture> CurrentTexture = nullptr;
+	FIntPoint CurrentTextureSize = FIntPoint::ZeroValue;
+	int32 CurrentTextureMipLevels = 0;
+	FNiagaraParameterDirectBinding<UObject*> UserParamBinding;
+};
+
+struct FNDIVirtualTextureSampleInstanceData_RenderThread
+{
+	FVirtualTexture2DResource *VirtualTextureResource;
+	uint8					VirtualTextureResource_SRGB;	
+	FIntPoint				TextureSize = FIntPoint(0, 0);
+	int32					MipLevels = 0;	
+};
+
+struct FNiagaraDataInterfaceProxyVirtualTextureSample : public FNiagaraDataInterfaceProxy
+{
+	virtual void ConsumePerInstanceDataFromGameThread(void* PerInstanceData, const FNiagaraSystemInstanceID& Instance) override { checkNoEntry(); }
+	virtual int32 PerInstanceDataPassedToRenderThreadSize() const override { return 0; }
+
+	TMap<FNiagaraSystemInstanceID, FNDIVirtualTextureSampleInstanceData_RenderThread> InstanceData_RT;
+};
+
+UNiagaraDataInterfaceVirtualTextureSample::UNiagaraDataInterfaceVirtualTextureSample(FObjectInitializer const& ObjectInitializer)
+	: Super(ObjectInitializer)
+	, Texture(nullptr)
+{
+	Proxy.Reset(new FNiagaraDataInterfaceProxyVirtualTextureSample());
+
+	FNiagaraTypeDefinition Def(UTexture::StaticClass());
+	TextureUserParameter.Parameter.SetType(Def);
+}
+
+void UNiagaraDataInterfaceVirtualTextureSample::PostInitProperties()
+{
+	Super::PostInitProperties();
+
+	if (HasAnyFlags(RF_ClassDefaultObject))
+	{
+		ENiagaraTypeRegistryFlags Flags = ENiagaraTypeRegistryFlags::AllowAnyVariable | ENiagaraTypeRegistryFlags::AllowParameter;
+		FNiagaraTypeRegistry::Register(FNiagaraTypeDefinition(GetClass()), Flags);
+	}
+}
+
+bool UNiagaraDataInterfaceVirtualTextureSample::CopyToInternal(UNiagaraDataInterface* Destination) const
+{
+	if (!Super::CopyToInternal(Destination))
+	{
+		return false;
+	}
+	UNiagaraDataInterfaceVirtualTextureSample* DestinationTexture = CastChecked<UNiagaraDataInterfaceVirtualTextureSample>(Destination);
+	DestinationTexture->Texture = Texture;
+	DestinationTexture->TextureUserParameter = TextureUserParameter;
+
+	return true;
+}
+
+bool UNiagaraDataInterfaceVirtualTextureSample::Equals(const UNiagaraDataInterface* Other) const
+{
+	if (!Super::Equals(Other))
+	{
+		return false;
+	}
+	const UNiagaraDataInterfaceVirtualTextureSample* OtherTexture = CastChecked<const UNiagaraDataInterfaceVirtualTextureSample>(Other);
+	return
+		OtherTexture->Texture == Texture &&
+		OtherTexture->TextureUserParameter == TextureUserParameter;
+}
+
+#if WITH_EDITORONLY_DATA
+void UNiagaraDataInterfaceVirtualTextureSample::GetFunctionsInternal(TArray<FNiagaraFunctionSignature>& OutFunctions) const
+{
+	FNiagaraFunctionSignature DefaultGpuSig;
+	DefaultGpuSig.bExperimental = true;
+	DefaultGpuSig.bMemberFunction = true;
+	DefaultGpuSig.bRequiresContext = false;
+	DefaultGpuSig.bSupportsCPU = false;
+	DefaultGpuSig.bSupportsGPU = true;
+	DefaultGpuSig.Inputs.Emplace(FNiagaraTypeDefinition(GetClass()), TEXT("VirtualTexture"));	
+
+	{
+		FNiagaraFunctionSignature& Sig = OutFunctions.Add_GetRef(DefaultGpuSig);
+		Sig.Name = SampleTexture2DName;
+		Sig.Inputs.Emplace(FNiagaraTypeDefinition::GetVec2Def(), TEXT("UV"));
+		Sig.Inputs.Emplace(FNiagaraTypeDefinition::GetFloatDef(), TEXT("MipLevel"));
+		Sig.Outputs.Emplace(FNiagaraTypeDefinition::GetVec4Def(), TEXT("Value"));
+		Sig.SetDescription(LOCTEXT("VirtualTextureSampleTexture2DDesc", "Sample supplied mip level from input 2d virtual texture at the specified UV coordinates. The UV origin (0,0) is in the upper left hand corner of the image."));
+	}
+
+	{
+		FNiagaraFunctionSignature& Sig = OutFunctions.AddDefaulted_GetRef();
+		Sig.Name = GetTextureDimensionsName;
+		Sig.bMemberFunction = true;
+		Sig.bRequiresContext = false;
+		Sig.Inputs.Emplace(FNiagaraTypeDefinition(GetClass()), TEXT("VirtualTexture"));
+		Sig.Inputs.Emplace(FNiagaraTypeDefinition::GetIntDef(), TEXT("MipLevel"));
+		Sig.Outputs.Emplace(FNiagaraTypeDefinition::GetVec2Def(), TEXT("Dimensions2D"));
+		Sig.SetDescription(LOCTEXT("VirtualTextureDimsDesc", "Get the dimensions of the provided Mip level for the virtual texture."));		
+	}
+	{
+		FNiagaraFunctionSignature& Sig = OutFunctions.AddDefaulted_GetRef();
+		Sig.Name = GetNumMipLevelsName;
+		Sig.bMemberFunction = true;
+		Sig.bRequiresContext = false;
+		Sig.Inputs.Emplace(FNiagaraTypeDefinition(GetClass()), TEXT("VirtualTexture"));
+		Sig.Outputs.Emplace(FNiagaraTypeDefinition::GetIntDef(), TEXT("NumMipLevels"));
+		Sig.SetDescription(LOCTEXT("GetNumMipLevelsDesc", "Get the number of Mip Levels for the virtual texture."));		
+	}
+}
+#endif
+
+void UNiagaraDataInterfaceVirtualTextureSample::GetVMExternalFunction(const FVMExternalFunctionBindingInfo& BindingInfo, void* InstanceData, FVMExternalFunction &OutFunc)
+{
+	if (BindingInfo.Name == GetTextureDimensionsName)
+	{
+		check(BindingInfo.GetNumInputs() == 2 && BindingInfo.GetNumOutputs() == 2);
+		OutFunc = FVMExternalFunction::CreateUObject(this, &UNiagaraDataInterfaceVirtualTextureSample::VMGetTextureDimensions);
+	}
+	else if (BindingInfo.Name == GetNumMipLevelsName)
+	{
+		check(BindingInfo.GetNumInputs() == 1 && BindingInfo.GetNumOutputs() == 1);
+		OutFunc = FVMExternalFunction::CreateUObject(this, &UNiagaraDataInterfaceVirtualTextureSample::VMGetNumMipLevels);
+	}
+}
+
+int32 UNiagaraDataInterfaceVirtualTextureSample::PerInstanceDataSize() const
+{
+	return sizeof(FNDIVirtualTextureInstanceData_GameThread);
+}
+
+bool UNiagaraDataInterfaceVirtualTextureSample::InitPerInstanceData(void* PerInstanceData, FNiagaraSystemInstance* SystemInstance)
+{
+	FNDIVirtualTextureInstanceData_GameThread* InstanceData = new (PerInstanceData) FNDIVirtualTextureInstanceData_GameThread();
+	InstanceData->UserParamBinding.Init(SystemInstance->GetInstanceParameters(), TextureUserParameter.Parameter);
+	return true;
+}
+
+void UNiagaraDataInterfaceVirtualTextureSample::DestroyPerInstanceData(void* PerInstanceData, FNiagaraSystemInstance* SystemInstance)
+{
+	FNDIVirtualTextureInstanceData_GameThread* InstanceData = static_cast<FNDIVirtualTextureInstanceData_GameThread*>(PerInstanceData);
+	InstanceData->~FNDIVirtualTextureInstanceData_GameThread();
+
+	ENQUEUE_RENDER_COMMAND(NDIVirtualTexture_RemoveInstance)
+	(
+		[RT_Proxy=GetProxyAs<FNiagaraDataInterfaceProxyVirtualTextureSample>(), RT_InstanceID=SystemInstance->GetId()](FRHICommandListImmediate&)
+		{
+			RT_Proxy->InstanceData_RT.Remove(RT_InstanceID);
+		}
+	);
+}
+
+bool UNiagaraDataInterfaceVirtualTextureSample::PerInstanceTick(void* PerInstanceData, FNiagaraSystemInstance* SystemInstance, float DeltaSeconds)
+{
+	FNDIVirtualTextureInstanceData_GameThread* InstanceData = static_cast<FNDIVirtualTextureInstanceData_GameThread*>(PerInstanceData);
+
+	UTexture* CurrentTexture = InstanceData->UserParamBinding.GetValueOrDefault<UTexture>(Texture);
+	FIntPoint CurrentTextureSize = FIntPoint::ZeroValue;
+	int32 CurrentTextureMipLevels = 0;
+
+	bool IsValid = true;
+
+	if (CurrentTexture != nullptr)
+	{
+		CurrentTextureSize.X = int32(CurrentTexture->GetSurfaceWidth());
+		CurrentTextureSize.Y = int32(CurrentTexture->GetSurfaceHeight());
+		if (UTexture2D* CurrentTexture2D = Cast<UTexture2D>(CurrentTexture))
+		{
+			CurrentTextureMipLevels = CurrentTexture2D->GetNumMips();
+		}
+		else if (UTextureRenderTarget2D* CurrentTexture2DRT = Cast<UTextureRenderTarget2D>(CurrentTexture))
+		{
+			CurrentTextureMipLevels = CurrentTexture2DRT->GetNumMips();
+		}
+		else
+		{
+			CurrentTextureMipLevels = 1;
+		}
+
+		// error if the texture isn't virtual
+		if (CurrentTexture->VirtualTextureStreaming == false)
+		{			
+			IsValid = false;		
+		}
+
+		// error if texture isn't valid
+		if (!CurrentTexture->IsValidLowLevel())
+		{		
+			IsValid = false;			
+		}
+	}
+	else
+	{		
+		IsValid = false;		
+	}
+
+	IsValid = IsValid && CurrentTexture->IsFullyStreamedIn() && !CurrentTexture->HasPendingInitOrStreaming() && !CurrentTexture->IsCompiling();
+
+	if (IsValid)
+	{
+		InstanceData->CurrentTexture = CurrentTexture;
+		InstanceData->CurrentTextureSize = CurrentTextureSize;
+		InstanceData->CurrentTextureMipLevels = CurrentTextureMipLevels;
+	
+		ENQUEUE_RENDER_COMMAND(NDIVirtualTexture_UpdateInstance)
+		(
+			[RT_Proxy=GetProxyAs<FNiagaraDataInterfaceProxyVirtualTextureSample>(), RT_InstanceID=SystemInstance->GetId(), RT_Texture=CurrentTexture, RT_TextureSize=CurrentTextureSize, RT_MipLevels=CurrentTextureMipLevels](FRHICommandListImmediate &RHI)
+			{
+				FNDIVirtualTextureSampleInstanceData_RenderThread& InstanceData = RT_Proxy->InstanceData_RT.FindOrAdd(RT_InstanceID);
+				
+				InstanceData.TextureSize = RT_TextureSize;
+				InstanceData.MipLevels = RT_MipLevels;
+				InstanceData.VirtualTextureResource = static_cast<FVirtualTexture2DResource*>(RT_Texture->GetResource());
+				InstanceData.VirtualTextureResource_SRGB = RT_Texture->SRGB;				
+			}
+		);
+	}
+	else
+	{
+		InstanceData->CurrentTexture = nullptr;
+		InstanceData->CurrentTextureSize = FIntPoint(0,0);
+		InstanceData->CurrentTextureMipLevels = 0;
+
+		ENQUEUE_RENDER_COMMAND(NDIVirtualTexture_UpdateInstance)
+		(
+			[RT_Proxy=GetProxyAs<FNiagaraDataInterfaceProxyVirtualTextureSample>(), RT_InstanceID=SystemInstance->GetId()](FRHICommandListImmediate &RHI)
+			{
+				FNDIVirtualTextureSampleInstanceData_RenderThread& InstanceData = RT_Proxy->InstanceData_RT.FindOrAdd(RT_InstanceID);
+				
+				InstanceData.TextureSize = 0;
+				InstanceData.MipLevels = 0;
+				InstanceData.VirtualTextureResource = nullptr;
+				InstanceData.VirtualTextureResource_SRGB = false;				
+			}
+		);
+	}
+
+	return false;
+}
+
+void UNiagaraDataInterfaceVirtualTextureSample::VMGetTextureDimensions(FVectorVMExternalFunctionContext& Context)
+{
+	VectorVM::FUserPtrHandler<FNDIVirtualTextureInstanceData_GameThread> InstData(Context);
+	FNDIInputParam<int32>		InMipLevel(Context);
+	FNDIOutputParam<FVector2f>	OutSize(Context);
+
+	for (int32 i = 0; i < Context.GetNumInstances(); ++i)
+	{
+		const int32 MipLevel = InMipLevel.GetAndAdvance();
+		const FVector2f TextureSize(
+			float(FMath::Max(InstData->CurrentTextureSize.X >> MipLevel, 1)),
+			float(FMath::Max(InstData->CurrentTextureSize.Y >> MipLevel, 1))
+		);
+		OutSize.SetAndAdvance(TextureSize);
+	}
+}
+
+void UNiagaraDataInterfaceVirtualTextureSample::VMGetNumMipLevels(FVectorVMExternalFunctionContext& Context)
+{
+	VectorVM::FUserPtrHandler<FNDIVirtualTextureInstanceData_GameThread> InstData(Context);
+	FNDIOutputParam<int32> OutNumMipLevels(Context);
+
+	for (int32 i = 0; i < Context.GetNumInstances(); ++i)
+	{
+		OutNumMipLevels.SetAndAdvance(InstData->CurrentTextureMipLevels);
+	}
+}
+
+#if WITH_EDITORONLY_DATA
+bool UNiagaraDataInterfaceVirtualTextureSample::AppendCompileHash(FNiagaraCompileHashVisitor* InVisitor) const
+{
+	bool bSuccess = Super::AppendCompileHash(InVisitor);
+	InVisitor->UpdateShaderFile(TemplateShaderFilePath);
+	bSuccess &= InVisitor->UpdateShaderParameters<FShaderParameters>();
+	return bSuccess;
+}
+
+void UNiagaraDataInterfaceVirtualTextureSample::GetParameterDefinitionHLSL(const FNiagaraDataInterfaceGPUParamInfo& ParamInfo, FString& OutHLSL)
+{
+	const TMap<FString, FStringFormatArg> TemplateArgs =
+	{
+		{TEXT("ParameterName"),	ParamInfo.DataInterfaceHLSLSymbol},
+	};
+	AppendTemplateHLSL(OutHLSL, TemplateShaderFilePath, TemplateArgs);
+}
+
+bool UNiagaraDataInterfaceVirtualTextureSample::GetFunctionHLSL(const FNiagaraDataInterfaceGPUParamInfo& ParamInfo, const FNiagaraDataInterfaceGeneratedFunction& FunctionInfo, int FunctionInstanceIndex, FString& OutHLSL)
+{
+	static const TSet<FName> ValidGpuFunctions =
+	{			
+		SampleTexture2DName,
+		GetTextureDimensionsName,
+		GetNumMipLevelsName,
+	};
+
+	return ValidGpuFunctions.Contains(FunctionInfo.DefinitionName);
+}
+
+#endif
+
+void UNiagaraDataInterfaceVirtualTextureSample::BuildShaderParameters(FNiagaraShaderParametersBuilder& ShaderParametersBuilder) const
+{
+	ShaderParametersBuilder.AddNestedStruct<FShaderParameters>();
+}
+
+void UNiagaraDataInterfaceVirtualTextureSample::SetShaderParameters(const FNiagaraDataInterfaceSetShaderParametersContext& Context) const
+{
+	FNiagaraDataInterfaceProxyVirtualTextureSample& TextureProxy = Context.GetProxy<FNiagaraDataInterfaceProxyVirtualTextureSample>();
+	FNDIVirtualTextureSampleInstanceData_RenderThread* InstanceData = TextureProxy.InstanceData_RT.Find(Context.GetSystemInstanceID());
+	FRDGBuilder& GraphBuilder = Context.GetGraphBuilder();
+
+	bool IsValid = false;
+
+	FShaderParameters* Parameters = Context.GetParameterNestedStruct<FShaderParameters>();
+	if (InstanceData)
+	{
+		FVirtualTexture2DResource* VirtualTextureResource = InstanceData->VirtualTextureResource;
+
+		if (VirtualTextureResource && VirtualTextureResource->IsInitialized())
+		{
+			IAllocatedVirtualTexture* AllocatedVT = VirtualTextureResource->AcquireAllocatedVT();
+
+			if (AllocatedVT)
+			{
+				FRHIShaderResourceView* InPhysicalTexture = AllocatedVT->GetPhysicalTextureSRV((uint32)0, InstanceData->VirtualTextureResource_SRGB == 1);
+
+				if (InPhysicalTexture && InPhysicalTexture->IsValid())
+				{
+					Parameters->TextureSize = InstanceData->TextureSize;
+					Parameters->MipLevels = InstanceData->MipLevels;
+
+					Parameters->PhysicalTexture = InPhysicalTexture;
+					Parameters->TextureSampler = VirtualTextureResource->SamplerStateRHI.GetReference();
+
+					Parameters->PageTableTexture0 = AllocatedVT->GetPageTableTexture(0);
+					Parameters->PageTableTexture1 = AllocatedVT->GetNumPageTableTextures() > 1 ? AllocatedVT->GetPageTableTexture(1) : GBlackTexture->TextureRHI.GetReference();
+
+					if (Parameters->TextureSampler->IsValid() && Parameters->PageTableTexture0->IsValid() && Parameters->PageTableTexture0->IsValid())
+					{
+						FUintVector4 VTPackedPageTableUniform[2];
+						FUintVector4 VTPackedUniform;
+
+						AllocatedVT->GetPackedPageTableUniform(VTPackedPageTableUniform);
+						AllocatedVT->GetPackedUniform(&VTPackedUniform, (uint32)0);
+
+						Parameters->VTPackedPageTableUniform[0] = VTPackedPageTableUniform[0];
+						Parameters->VTPackedPageTableUniform[1] = VTPackedPageTableUniform[1];
+						Parameters->VTPackedUniform = VTPackedUniform;
+
+						IsValid = true;
+					}
+				}
+			}
+		}
+	}
+	
+	if (!IsValid)
+	{
+		Parameters->TextureSize = FIntPoint(0, 0);
+		Parameters->MipLevels = 0;
+		Parameters->TextureSampler = TStaticSamplerState<SF_Point>::GetRHI();		
+		Parameters->VTPackedPageTableUniform[0] = FUintVector4(0, 0, 0, 0);
+		Parameters->VTPackedPageTableUniform[1] = FUintVector4(0, 0, 0, 0);
+		Parameters->VTPackedUniform = FUintVector4();
+		Parameters->PhysicalTexture = GBlackTextureWithSRV->ShaderResourceViewRHI;
+		Parameters->PageTableTexture0 = GBlackTexture->TextureRHI;
+		Parameters->PageTableTexture1 = GBlackTexture->TextureRHI;		
+	}
+}
+
+#undef LOCTEXT_NAMESPACE
+
