small fixes

diff --git a/Engine/Plugins/FX/Niagara/Source/Niagara/Private/NiagaraDataInterfaceVirtualTextureSample.cpp b/Engine/Plugins/FX/Niagara/Source/Niagara/Private/NiagaraDataInterfaceVirtualTextureSample.cpp
--- a/Engine/Plugins/FX/Niagara/Source/Niagara/Private/NiagaraDataInterfaceVirtualTextureSample.cpp
+++ b/Engine/Plugins/FX/Niagara/Source/Niagara/Private/NiagaraDataInterfaceVirtualTextureSample.cpp
@@ -38,8 +38,8 @@ struct FNDIVirtualTextureInstanceData_GameThread
 
 struct FNDIVirtualTextureSampleInstanceData_RenderThread
 {
-	FVirtualTexture2DResource *VirtualTextureResource;
-	uint8					VirtualTextureResource_SRGB;	
+	FVirtualTexture2DResource *VirtualTextureResource = nullptr;
+	uint8					VirtualTextureResource_SRGB = 0;	
 	FIntPoint				TextureSize = FIntPoint(0, 0);
 	int32					MipLevels = 0;	
 };
@@ -365,7 +365,7 @@ void UNiagaraDataInterfaceVirtualTextureSample::SetShaderParameters(const FNiaga
 					Parameters->PageTableTexture0 = AllocatedVT->GetPageTableTexture(0);
 					Parameters->PageTableTexture1 = AllocatedVT->GetNumPageTableTextures() > 1 ? AllocatedVT->GetPageTableTexture(1) : GBlackTexture->TextureRHI.GetReference();
 
-					if (Parameters->TextureSampler->IsValid() && Parameters->PageTableTexture0->IsValid() && Parameters->PageTableTexture0->IsValid())
+					if (Parameters->TextureSampler->IsValid() && Parameters->PageTableTexture0->IsValid() && Parameters->PageTableTexture1->IsValid())
 					{
 						FUintVector4 VTPackedPageTableUniform[2];
 						FUintVector4 VTPackedUniform;
