Niagara - Clean up some DI APIs (routing refactor prep)
+ Change PostCompile to take a const pointer to the system, rather than having the DIs use GetOuter.
+ Remove an unused option from the data interface utilities, and also make the system and system instance references const.

diff --git a/Engine/Plugins/FX/Niagara/Source/Niagara/Classes/NiagaraDataInterface.h b/Engine/Plugins/FX/Niagara/Source/Niagara/Classes/NiagaraDataInterface.h
--- a/Engine/Plugins/FX/Niagara/Source/Niagara/Classes/NiagaraDataInterface.h
+++ b/Engine/Plugins/FX/Niagara/Source/Niagara/Classes/NiagaraDataInterface.h
@@ -697,9 +697,17 @@ public:
 	virtual bool HasPreStageTick(ENiagaraScriptUsage Usage) const { return false; }
 	virtual bool HasPostStageTick(ENiagaraScriptUsage Usage) const { return false; }
 
-	/** Called after system compilation completes. Useful for caching any constant, compile dependent data. */
+	UE_DEPRECATED(5.7, "Please update to the PostCompile that takes the owning system as this api will be removed in a future version.")
 	virtual void PostCompile() { }
 
+	/** Called after system compilation completes. Useful for caching any constant, compile dependent data. */
+	virtual void PostCompile(const UNiagaraSystem& OwningSystem)
+	{
+		PRAGMA_DISABLE_DEPRECATION_WARNINGS
+		PostCompile();
+		PRAGMA_ENABLE_DEPRECATION_WARNINGS
+	}
+
 	/**
 	When set to true the simulation may not complete in the same frame it started, allowing maximum overlap with the GameThread.
 	You must override and set to false if you require the data interface to complete before PostActorTick completes.
diff --git a/Engine/Plugins/FX/Niagara/Source/Niagara/Classes/NiagaraDataInterfaceCurveBase.h b/Engine/Plugins/FX/Niagara/Source/Niagara/Classes/NiagaraDataInterfaceCurveBase.h
--- a/Engine/Plugins/FX/Niagara/Source/Niagara/Classes/NiagaraDataInterfaceCurveBase.h
+++ b/Engine/Plugins/FX/Niagara/Source/Niagara/Classes/NiagaraDataInterfaceCurveBase.h
@@ -134,7 +134,7 @@ public:
 #endif
 	NIAGARA_API virtual void BuildShaderParameters(FNiagaraShaderParametersBuilder& ShaderParametersBuilder) const override;
 	NIAGARA_API virtual void SetShaderParameters(const FNiagaraDataInterfaceSetShaderParametersContext& Context) const override;
-	NIAGARA_API virtual void PostCompile() override;
+	NIAGARA_API virtual void PostCompile(const UNiagaraSystem& OwningSystem) override;
 
 	NIAGARA_API void SetDefaultLUT();
 #if WITH_EDITORONLY_DATA
diff --git a/Engine/Plugins/FX/Niagara/Source/Niagara/Classes/NiagaraSimCacheCustomStorageInterface.h b/Engine/Plugins/FX/Niagara/Source/Niagara/Classes/NiagaraSimCacheCustomStorageInterface.h
--- a/Engine/Plugins/FX/Niagara/Source/Niagara/Classes/NiagaraSimCacheCustomStorageInterface.h
+++ b/Engine/Plugins/FX/Niagara/Source/Niagara/Classes/NiagaraSimCacheCustomStorageInterface.h
@@ -136,7 +136,7 @@ public:
 	For example, if you were to require 'Particles.MyAttribute' in order to process the cache results you would need to convert
 	this into 'MyEmitter.Particles.MyAttribute' by checking the UsageContext is a UNiagaraEmitter and then creating the variable from the unique name.
 	*/
-	virtual TArray<FNiagaraVariableBase> GetSimCacheRendererAttributes(UObject* UsageContext) const { return TArray<FNiagaraVariableBase>(); }
+	virtual TArray<FNiagaraVariableBase> GetSimCacheRendererAttributes(const UObject* UsageContext) const { return TArray<FNiagaraVariableBase>(); }
 
 	/**
 	 This converts the content of the storage object to a json representation. If another interchange format (e.g. an image format) is better, then the json this method produces should link to the secondary files.
diff --git a/Engine/Plugins/FX/Niagara/Source/Niagara/Internal/DataInterface/NiagaraDataInterfaceDataChannelCommon.h b/Engine/Plugins/FX/Niagara/Source/Niagara/Internal/DataInterface/NiagaraDataInterfaceDataChannelCommon.h
--- a/Engine/Plugins/FX/Niagara/Source/Niagara/Internal/DataInterface/NiagaraDataInterfaceDataChannelCommon.h
+++ b/Engine/Plugins/FX/Niagara/Source/Niagara/Internal/DataInterface/NiagaraDataInterfaceDataChannelCommon.h
@@ -163,7 +163,7 @@ struct FNDIDataChannelCompiledData
 	GENERATED_BODY()
 
 	/** Initializes function bindings and binds if there is a valid DataSetCompiledData. */
-	bool Init(UNiagaraSystem* System, UNiagaraDataInterface* OwnerDI);
+	bool Init(const UNiagaraSystem* System, UNiagaraDataInterface* OwnerDI);
 
 	int32 FindFunctionInfoIndex(FName Name, const TArray<FNiagaraVariableBase>& VariadicInputs, const TArray<FNiagaraVariableBase>& VariadicOutputs)const;
 
@@ -214,7 +214,7 @@ protected:
 	bool bCallsWrite = false;
 
 	/** Iterates over all scripts for the owning system and gathers all functions and parameters accessing this DI. Building the FunctionInfoTable and GPUScriptParameterInfos map.  */
-	void GatherAccessInfo(UNiagaraSystem* System, UNiagaraDataInterface* Owner);
+	void GatherAccessInfo(const UNiagaraSystem* System, UNiagaraDataInterface* Owner);
 };
 
 class FNDIDummyUAV : public FRenderResource
diff --git a/Engine/Plugins/FX/Niagara/Source/Niagara/Internal/DataInterface/NiagaraDataInterfaceDataChannelRead.h b/Engine/Plugins/FX/Niagara/Source/Niagara/Internal/DataInterface/NiagaraDataInterfaceDataChannelRead.h
--- a/Engine/Plugins/FX/Niagara/Source/Niagara/Internal/DataInterface/NiagaraDataInterfaceDataChannelRead.h
+++ b/Engine/Plugins/FX/Niagara/Source/Niagara/Internal/DataInterface/NiagaraDataInterfaceDataChannelRead.h
@@ -164,7 +164,7 @@ public:
 
 	virtual bool UpgradeFunctionCall(FNiagaraFunctionSignature& FunctionSignature)override;
 
-	NIAGARA_API virtual void PostCompile()override;
+	NIAGARA_API virtual void PostCompile(const UNiagaraSystem& OwningSystem) override;
 #endif	
 
 #if WITH_EDITOR	
diff --git a/Engine/Plugins/FX/Niagara/Source/Niagara/Internal/DataInterface/NiagaraDataInterfaceDataChannelWrite.h b/Engine/Plugins/FX/Niagara/Source/Niagara/Internal/DataInterface/NiagaraDataInterfaceDataChannelWrite.h
--- a/Engine/Plugins/FX/Niagara/Source/Niagara/Internal/DataInterface/NiagaraDataInterfaceDataChannelWrite.h
+++ b/Engine/Plugins/FX/Niagara/Source/Niagara/Internal/DataInterface/NiagaraDataInterfaceDataChannelWrite.h
@@ -35,7 +35,7 @@ struct FNDIDataChannelWriteCompiledData : public FNDIDataChannelCompiledData
 	UPROPERTY()
 	FNiagaraDataSetCompiledData DataLayout;
 
-	bool Init(UNiagaraSystem* System, UNiagaraDataInterfaceDataChannelWrite* OwnerDI);
+	bool Init(const UNiagaraSystem* System, UNiagaraDataInterfaceDataChannelWrite* OwnerDI);
 };
 
 USTRUCT()
@@ -182,7 +182,7 @@ public:
 	NIAGARA_API virtual bool GetFunctionHLSL(const FNiagaraDataInterfaceHlslGenerationContext& HlslGenContext, FString& OutHLSL) override;
 	NIAGARA_API virtual void GetParameterDefinitionHLSL(const FNiagaraDataInterfaceHlslGenerationContext& HlslGenContext, FString& OutHLSL) override;
 
-	NIAGARA_API virtual void PostCompile()override;
+	NIAGARA_API virtual void PostCompile(const UNiagaraSystem& OwningSystem) override;
 #endif
 	NIAGARA_API virtual void BuildShaderParameters(FNiagaraShaderParametersBuilder& ShaderParametersBuilder) const override;
 	NIAGARA_API virtual void SetShaderParameters(const FNiagaraDataInterfaceSetShaderParametersContext& Context) const override;
diff --git a/Engine/Plugins/FX/Niagara/Source/Niagara/Internal/DataInterface/NiagaraDataInterfaceDataTable.h b/Engine/Plugins/FX/Niagara/Source/Niagara/Internal/DataInterface/NiagaraDataInterfaceDataTable.h
--- a/Engine/Plugins/FX/Niagara/Source/Niagara/Internal/DataInterface/NiagaraDataInterfaceDataTable.h
+++ b/Engine/Plugins/FX/Niagara/Source/Niagara/Internal/DataInterface/NiagaraDataInterfaceDataTable.h
@@ -42,7 +42,7 @@ public:
 	NIAGARA_API virtual void GetParameterDefinitionHLSL(const FNiagaraDataInterfaceGPUParamInfo& ParamInfo, FString& OutHLSL) override;
 	NIAGARA_API virtual bool GetFunctionHLSL(const FNiagaraDataInterfaceHlslGenerationContext& HlslGenContext, FString& OutHLSL) override;
 	NIAGARA_API virtual bool AppendCompileHash(FNiagaraCompileHashVisitor* InVisitor) const override;
-	NIAGARA_API virtual void PostCompile() override;
+	NIAGARA_API virtual void PostCompile(const UNiagaraSystem& OwningSystem) override;
 #endif
 #if WITH_EDITOR
 	virtual bool GetGpuUseIndirectDispatch() const override { return true; }
diff --git a/Engine/Plugins/FX/Niagara/Source/Niagara/Private/DataInterface/NiagaraDataInterfaceDataChannelCommon.cpp b/Engine/Plugins/FX/Niagara/Source/Niagara/Private/DataInterface/NiagaraDataInterfaceDataChannelCommon.cpp
--- a/Engine/Plugins/FX/Niagara/Source/Niagara/Private/DataInterface/NiagaraDataInterfaceDataChannelCommon.cpp
+++ b/Engine/Plugins/FX/Niagara/Source/Niagara/Private/DataInterface/NiagaraDataInterfaceDataChannelCommon.cpp
@@ -189,7 +189,7 @@ FNDIDataChannel_FuncToDataSetBindingPtr FNDIDataChannelLayoutManager::GetLayoutI
 //////////////////////////////////////////////////////////////////////////
 //FNDIDataChannelCompiledData
 
-void FNDIDataChannelCompiledData::GatherAccessInfo(UNiagaraSystem* System, UNiagaraDataInterface* Owner)
+void FNDIDataChannelCompiledData::GatherAccessInfo(const UNiagaraSystem* System, UNiagaraDataInterface* Owner)
 {
 	//We search all VM functions called on this DI to generate an appropriate FNDIDataChannelFunctionInfo that can later be used in binding to actual dataset Data.
 	auto HandleVMFunc = [&](const UNiagaraScript* Script, const FVMExternalFunctionBindingInfo& BindingInfo)
@@ -266,7 +266,7 @@ void FNDIDataChannelCompiledData::GatherAccessInfo(UNiagaraSystem* System, UNiag
 	}
 }
 
-bool FNDIDataChannelCompiledData::Init(UNiagaraSystem* System, UNiagaraDataInterface* OwnerDI)
+bool FNDIDataChannelCompiledData::Init(const UNiagaraSystem* System, UNiagaraDataInterface* OwnerDI)
 {
 	FunctionInfo.Empty();
 
diff --git a/Engine/Plugins/FX/Niagara/Source/Niagara/Private/DataInterface/NiagaraDataInterfaceDataChannelRead.cpp b/Engine/Plugins/FX/Niagara/Source/Niagara/Private/DataInterface/NiagaraDataInterfaceDataChannelRead.cpp
--- a/Engine/Plugins/FX/Niagara/Source/Niagara/Private/DataInterface/NiagaraDataInterfaceDataChannelRead.cpp
+++ b/Engine/Plugins/FX/Niagara/Source/Niagara/Private/DataInterface/NiagaraDataInterfaceDataChannelRead.cpp
@@ -1053,10 +1053,9 @@ ETickingGroup UNiagaraDataInterfaceDataChannelRead::CalculateTickGroup(const voi
 
 #if WITH_EDITORONLY_DATA
 
-void UNiagaraDataInterfaceDataChannelRead::PostCompile()
+void UNiagaraDataInterfaceDataChannelRead::PostCompile(const UNiagaraSystem& OwningSystem)
 {
-	UNiagaraSystem* OwnerSystem = GetTypedOuter<UNiagaraSystem>();
-	CompiledData.Init(OwnerSystem, this);
+	CompiledData.Init(&OwningSystem, this);
 }
 
 #endif
diff --git a/Engine/Plugins/FX/Niagara/Source/Niagara/Private/DataInterface/NiagaraDataInterfaceDataChannelWrite.cpp b/Engine/Plugins/FX/Niagara/Source/Niagara/Private/DataInterface/NiagaraDataInterfaceDataChannelWrite.cpp
--- a/Engine/Plugins/FX/Niagara/Source/Niagara/Private/DataInterface/NiagaraDataInterfaceDataChannelWrite.cpp
+++ b/Engine/Plugins/FX/Niagara/Source/Niagara/Private/DataInterface/NiagaraDataInterfaceDataChannelWrite.cpp
@@ -223,7 +223,7 @@ struct FNDIDataChannelWriteInstanceData_RT
 //////////////////////////////////////////////////////////////////////////
 //FNDIDataChannelWriteCompiledData
 
-bool FNDIDataChannelWriteCompiledData::Init(UNiagaraSystem* System, UNiagaraDataInterfaceDataChannelWrite* OwnerDI)
+bool FNDIDataChannelWriteCompiledData::Init(const UNiagaraSystem* System, UNiagaraDataInterfaceDataChannelWrite* OwnerDI)
 {
 	FunctionInfo.Reset();
 
@@ -712,10 +712,9 @@ ETickingGroup UNiagaraDataInterfaceDataChannelWrite::CalculateFinalTickGroup(con
 
 #if WITH_EDITORONLY_DATA
 
-void UNiagaraDataInterfaceDataChannelWrite::PostCompile()
+void UNiagaraDataInterfaceDataChannelWrite::PostCompile(const UNiagaraSystem& OwningSystem)
 {
-	UNiagaraSystem* OwnerSystem = GetTypedOuter<UNiagaraSystem>();
-	CompiledData.Init(OwnerSystem, this);
+	CompiledData.Init(&OwningSystem, this);
 }
 
 #endif
diff --git a/Engine/Plugins/FX/Niagara/Source/Niagara/Private/DataInterface/NiagaraDataInterfaceDataTable.cpp b/Engine/Plugins/FX/Niagara/Source/Niagara/Private/DataInterface/NiagaraDataInterfaceDataTable.cpp
--- a/Engine/Plugins/FX/Niagara/Source/Niagara/Private/DataInterface/NiagaraDataInterfaceDataTable.cpp
+++ b/Engine/Plugins/FX/Niagara/Source/Niagara/Private/DataInterface/NiagaraDataInterfaceDataTable.cpp
@@ -775,20 +775,19 @@ bool UNiagaraDataInterfaceDataTable::AppendCompileHash(FNiagaraCompileHashVisito
 	return bSuccess;
 }
 
-void UNiagaraDataInterfaceDataTable::PostCompile()
+void UNiagaraDataInterfaceDataTable::PostCompile(const UNiagaraSystem& OwningSystem)
 {
 	using namespace NDIDataTableLocal;
 
-	Super::PostCompile();
+	Super::PostCompile(OwningSystem);
 
 	bCreateFilteredTable = true;
-	if ( UNiagaraSystem* OwnerSystem = GetTypedOuter<UNiagaraSystem>() )
 	{
 		static const TArray<FName> AllTableFunctions = { NAME_GetRow, NAME_GetNumRows };
 
 		FNiagaraDataInterfaceUtilities::ForEachVMFunction(
 			this,
-			OwnerSystem,
+			&OwningSystem,
 			[this](const UNiagaraScript* Script, const FVMExternalFunctionBindingInfo& FunctionBinding) -> bool
 			{
 				if (AllTableFunctions.Contains(FunctionBinding.Name))
@@ -803,7 +802,7 @@ void UNiagaraDataInterfaceDataTable::PostCompile()
 		{
 			FNiagaraDataInterfaceUtilities::ForEachGpuFunction(
 				this,
-				OwnerSystem,
+				&OwningSystem,
 				[this](const UNiagaraScript* Script, const FNiagaraDataInterfaceGeneratedFunction& FunctionBinding) -> bool
 				{
 					if (AllTableFunctions.Contains(FunctionBinding.DefinitionName))
diff --git a/Engine/Plugins/FX/Niagara/Source/Niagara/Private/DataInterface/NiagaraDataInterfaceUtilities.cpp b/Engine/Plugins/FX/Niagara/Source/Niagara/Private/DataInterface/NiagaraDataInterfaceUtilities.cpp
--- a/Engine/Plugins/FX/Niagara/Source/Niagara/Private/DataInterface/NiagaraDataInterfaceUtilities.cpp
+++ b/Engine/Plugins/FX/Niagara/Source/Niagara/Private/DataInterface/NiagaraDataInterfaceUtilities.cpp
@@ -14,7 +14,7 @@
 namespace FNiagaraDataInterfaceUtilities
 {
 
-void ForEachVMFunctionImpl(class UNiagaraDataInterface* RuntimeResolvedDataInterface, class UNiagaraSystem* NiagaraSystem, const FNiagaraParameterStore& ParameterStore, TFunction<bool(const UNiagaraScript* Script, const FVMExternalFunctionBindingInfo&)> Action)
+void ForEachVMFunctionImpl(class UNiagaraDataInterface* RuntimeResolvedDataInterface, const class UNiagaraSystem* NiagaraSystem, const FNiagaraParameterStore& ParameterStore, TFunction<bool(const UNiagaraScript* Script, const FVMExternalFunctionBindingInfo&)> Action)
 {
 	if (RuntimeResolvedDataInterface == nullptr || NiagaraSystem == nullptr )
 	{
@@ -75,7 +75,7 @@ void ForEachVMFunctionImpl(class UNiagaraDataInterface* RuntimeResolvedDataInter
 	);
 }
 
-void ForEachGpuFunctionImpl(class UNiagaraDataInterface* RuntimeResolvedDataInterface, class UNiagaraSystem* NiagaraSystem, const FNiagaraParameterStore& ParameterStore, TFunction<bool(const UNiagaraScript*, const FNiagaraDataInterfaceGeneratedFunction&)> Action)
+void ForEachGpuFunctionImpl(class UNiagaraDataInterface* RuntimeResolvedDataInterface, const class UNiagaraSystem* NiagaraSystem, const FNiagaraParameterStore& ParameterStore, TFunction<bool(const UNiagaraScript*, const FNiagaraDataInterfaceGeneratedFunction&)> Action)
 {
 	if (RuntimeResolvedDataInterface == nullptr || NiagaraSystem == nullptr)
 	{
@@ -133,7 +133,7 @@ void ForEachGpuFunctionImpl(class UNiagaraDataInterface* RuntimeResolvedDataInte
 	);
 }
 
-void ForEachDataInterface(const FNiagaraParameterStore& ParameterStore, TFunction<bool(const FNiagaraVariableBase Variable, UNiagaraDataInterface* DataInterface)> Action, FDataInterfaceSearchOptions SearchOptions)
+void ForEachDataInterface(const FNiagaraParameterStore& ParameterStore, TFunction<bool(const FNiagaraVariableBase Variable, UNiagaraDataInterface* DataInterface)> Action)
 {
 	const TArray<UNiagaraDataInterface*>& DataInterfaces = ParameterStore.GetDataInterfaces();
 	for (const FNiagaraVariableWithOffset& Variable : ParameterStore.ReadParameterVariables())
@@ -143,11 +143,6 @@ void ForEachDataInterface(const FNiagaraParameterStore& ParameterStore, TFunctio
 			continue;
 		}
 
-		if (SearchOptions.bIncludeInternal == false && Variable.IsInNameSpace(FNiagaraConstants::InternalNamespaceString))
-		{
-			continue;
-		}
-
 		if ( UNiagaraDataInterface* DataInterface = DataInterfaces[Variable.Offset] )
 		{
 			Action(Variable, DataInterface);
@@ -155,7 +150,7 @@ void ForEachDataInterface(const FNiagaraParameterStore& ParameterStore, TFunctio
 	}
 }
 
-void ForEachDataInterface(FDataInterfaceUsageContext& UsageContext, const FNiagaraParameterStore& ParameterStore, TFunction<bool(const FDataInterfaceUsageContext&)> Action, FDataInterfaceSearchOptions SearchOptions)
+void ForEachDataInterface(FDataInterfaceUsageContext& UsageContext, const FNiagaraParameterStore& ParameterStore, TFunction<bool(const FDataInterfaceUsageContext&)> Action)
 {
 	const TArray<UNiagaraDataInterface*>& DataInterfaces = ParameterStore.GetDataInterfaces();
 	for (const FNiagaraVariableWithOffset& Variable : ParameterStore.ReadParameterVariables())
@@ -165,11 +160,6 @@ void ForEachDataInterface(FDataInterfaceUsageContext& UsageContext, const FNiaga
 			continue;
 		}
 
-		if (SearchOptions.bIncludeInternal == false && Variable.IsInNameSpace(FNiagaraConstants::InternalNamespaceString))
-		{
-			continue;
-		}
-
 		if (UNiagaraDataInterface* DataInterface = DataInterfaces[Variable.Offset])
 		{
 			UsageContext.Variable = Variable;
@@ -181,7 +171,7 @@ void ForEachDataInterface(FDataInterfaceUsageContext& UsageContext, const FNiaga
 
 } //namespace FNiagaraDataInterfaceUtilities
 
-void FNiagaraDataInterfaceUtilities::ForEachVMFunction(UNiagaraDataInterface* RuntimeResolvedDataInterface, UNiagaraSystem* NiagaraSystem, TFunction<bool(const UNiagaraScript*, const FVMExternalFunctionBindingInfo&)> Action)
+void FNiagaraDataInterfaceUtilities::ForEachVMFunction(UNiagaraDataInterface* RuntimeResolvedDataInterface, const UNiagaraSystem* NiagaraSystem, TFunction<bool(const UNiagaraScript*, const FVMExternalFunctionBindingInfo&)> Action)
 {
 	ForEachVMFunctionImpl(RuntimeResolvedDataInterface, NiagaraSystem, NiagaraSystem->GetExposedParameters(), Action);
 }
@@ -194,7 +184,7 @@ void FNiagaraDataInterfaceUtilities::ForEachVMFunction(UNiagaraDataInterface* Ru
 	}
 }
 
-void FNiagaraDataInterfaceUtilities::ForEachVMFunction(UNiagaraDataInterface* RuntimeResolvedDataInterface, FNiagaraSystemInstance* SystemInstance, TFunction<bool(const UNiagaraScript*, const FVMExternalFunctionBindingInfo&)> Action)
+void FNiagaraDataInterfaceUtilities::ForEachVMFunction(UNiagaraDataInterface* RuntimeResolvedDataInterface, const FNiagaraSystemInstance* SystemInstance, TFunction<bool(const UNiagaraScript*, const FVMExternalFunctionBindingInfo&)> Action)
 {
 	if (SystemInstance != nullptr && SystemInstance->GetOverrideParameters() != nullptr && SystemInstance->GetSystem() != nullptr)
 	{
@@ -202,7 +192,7 @@ void FNiagaraDataInterfaceUtilities::ForEachVMFunction(UNiagaraDataInterface* Ru
 	}
 }
 
-void FNiagaraDataInterfaceUtilities::ForEachGpuFunction(UNiagaraDataInterface* RuntimeResolvedDataInterface, class UNiagaraSystem* NiagaraSystem, TFunction<bool(const UNiagaraScript*, const FNiagaraDataInterfaceGeneratedFunction&)> Action)
+void FNiagaraDataInterfaceUtilities::ForEachGpuFunction(UNiagaraDataInterface* RuntimeResolvedDataInterface, const class UNiagaraSystem* NiagaraSystem, TFunction<bool(const UNiagaraScript*, const FNiagaraDataInterfaceGeneratedFunction&)> Action)
 {
 	ForEachGpuFunctionImpl(RuntimeResolvedDataInterface, NiagaraSystem, NiagaraSystem->GetExposedParameters(), Action);
 }
@@ -215,7 +205,7 @@ void FNiagaraDataInterfaceUtilities::ForEachGpuFunction(UNiagaraDataInterface* R
 	}
 }
 
-void FNiagaraDataInterfaceUtilities::ForEachGpuFunction(UNiagaraDataInterface* RuntimeResolvedDataInterface, FNiagaraSystemInstance* SystemInstance, TFunction<bool(const UNiagaraScript*, const FNiagaraDataInterfaceGeneratedFunction&)> Action)
+void FNiagaraDataInterfaceUtilities::ForEachGpuFunction(UNiagaraDataInterface* RuntimeResolvedDataInterface, const FNiagaraSystemInstance* SystemInstance, TFunction<bool(const UNiagaraScript*, const FNiagaraDataInterfaceGeneratedFunction&)> Action)
 {
 	if (SystemInstance != nullptr && SystemInstance->GetOverrideParameters() != nullptr && SystemInstance->GetSystem() != nullptr)
 	{
@@ -223,7 +213,7 @@ void FNiagaraDataInterfaceUtilities::ForEachGpuFunction(UNiagaraDataInterface* R
 	}
 }
 
-void FNiagaraDataInterfaceUtilities::ForEachDataInterface(FNiagaraSystemInstance* SystemInstance, TFunction<bool(const FNiagaraVariableBase Variable, UNiagaraDataInterface* DataInterface)> Action, FDataInterfaceSearchOptions SearchOptions)
+void FNiagaraDataInterfaceUtilities::ForEachDataInterface(const FNiagaraSystemInstance* SystemInstance, TFunction<bool(const FNiagaraVariableBase Variable, UNiagaraDataInterface* DataInterface)> Action)
 {
 	if ( SystemInstance == nullptr )
 	{
@@ -236,10 +226,10 @@ void FNiagaraDataInterfaceUtilities::ForEachDataInterface(FNiagaraSystemInstance
 		return;
 	}
 
-	ForEachDataInterface(*SystemInstance->GetOverrideParameters(), Action, SearchOptions);
+	ForEachDataInterface(*SystemInstance->GetOverrideParameters(), Action);
 
-	ForEachDataInterface(SystemSimulation->GetSpawnExecutionContext()->Parameters, Action, SearchOptions);
-	ForEachDataInterface(SystemSimulation->GetUpdateExecutionContext()->Parameters, Action, SearchOptions);
+	ForEachDataInterface(SystemSimulation->GetSpawnExecutionContext()->Parameters, Action);
+	ForEachDataInterface(SystemSimulation->GetUpdateExecutionContext()->Parameters, Action);
 
 	for (const FNiagaraEmitterInstanceRef& Emitter : SystemInstance->GetEmitters())
 	{
@@ -249,21 +239,21 @@ void FNiagaraDataInterfaceUtilities::ForEachDataInterface(FNiagaraSystemInstance
 			continue;
 		}
 
-		ForEachDataInterface(StatefulEmitter->GetSpawnExecutionContext().Parameters, Action, SearchOptions);
-		ForEachDataInterface(StatefulEmitter->GetUpdateExecutionContext().Parameters, Action, SearchOptions);
+		ForEachDataInterface(StatefulEmitter->GetSpawnExecutionContext().Parameters, Action);
+		ForEachDataInterface(StatefulEmitter->GetUpdateExecutionContext().Parameters, Action);
 		for (const FNiagaraScriptExecutionContext& EventExecContext : StatefulEmitter->GetEventExecutionContexts())
 		{
-			ForEachDataInterface(EventExecContext.Parameters, Action, SearchOptions);
+			ForEachDataInterface(EventExecContext.Parameters, Action);
 		}
 
 		if (Emitter->GetEmitter() && Emitter->GetSimTarget() == ENiagaraSimTarget::GPUComputeSim && Emitter->GetGPUContext())
 		{
-			ForEachDataInterface(Emitter->GetGPUContext()->CombinedParamStore, Action, SearchOptions);
+			ForEachDataInterface(Emitter->GetGPUContext()->CombinedParamStore, Action);
 		}
 	}
 }
 
-void FNiagaraDataInterfaceUtilities::ForEachDataInterface(FNiagaraSystemInstance* SystemInstance, TFunction<bool(const FDataInterfaceUsageContext&)> Action, FDataInterfaceSearchOptions SearchOptions)
+void FNiagaraDataInterfaceUtilities::ForEachDataInterface(const FNiagaraSystemInstance* SystemInstance, TFunction<bool(const FDataInterfaceUsageContext&)> Action)
 {
 	if (SystemInstance == nullptr)
 	{
@@ -278,10 +268,10 @@ void FNiagaraDataInterfaceUtilities::ForEachDataInterface(FNiagaraSystemInstance
 
 	FDataInterfaceUsageContext UsageContext;
 	UsageContext.OwnerObject = SystemSimulation->GetSystem();
-	ForEachDataInterface(UsageContext, *SystemInstance->GetOverrideParameters(), Action, SearchOptions);
+	ForEachDataInterface(UsageContext, *SystemInstance->GetOverrideParameters(), Action);
 
-	ForEachDataInterface(UsageContext, SystemSimulation->GetSpawnExecutionContext()->Parameters, Action, SearchOptions);
-	ForEachDataInterface(UsageContext, SystemSimulation->GetUpdateExecutionContext()->Parameters, Action, SearchOptions);
+	ForEachDataInterface(UsageContext, SystemSimulation->GetSpawnExecutionContext()->Parameters, Action);
+	ForEachDataInterface(UsageContext, SystemSimulation->GetUpdateExecutionContext()->Parameters, Action);
 
 	for (const FNiagaraEmitterInstanceRef& Emitter : SystemInstance->GetEmitters())
 	{
@@ -292,21 +282,21 @@ void FNiagaraDataInterfaceUtilities::ForEachDataInterface(FNiagaraSystemInstance
 		}
 		UsageContext.OwnerObject = Emitter->GetEmitter();
 
-		ForEachDataInterface(UsageContext, StatefulEmitter->GetSpawnExecutionContext().Parameters, Action, SearchOptions);
-		ForEachDataInterface(UsageContext, StatefulEmitter->GetUpdateExecutionContext().Parameters, Action, SearchOptions);
+		ForEachDataInterface(UsageContext, StatefulEmitter->GetSpawnExecutionContext().Parameters, Action);
+		ForEachDataInterface(UsageContext, StatefulEmitter->GetUpdateExecutionContext().Parameters, Action);
 		for (const FNiagaraScriptExecutionContext& EventExecContext : StatefulEmitter->GetEventExecutionContexts() )
 		{
-			ForEachDataInterface(UsageContext, EventExecContext.Parameters, Action, SearchOptions);
+			ForEachDataInterface(UsageContext, EventExecContext.Parameters, Action);
 		}
 
 		if (Emitter->GetEmitter() && Emitter->GetSimTarget() == ENiagaraSimTarget::GPUComputeSim && Emitter->GetGPUContext())
 		{
-			ForEachDataInterface(UsageContext, Emitter->GetGPUContext()->CombinedParamStore, Action, SearchOptions);
+			ForEachDataInterface(UsageContext, Emitter->GetGPUContext()->CombinedParamStore, Action);
 		}
 	}
 }
 
-void FNiagaraDataInterfaceUtilities::ForEachDataInterface(UNiagaraSystem* NiagaraSystem, TFunction<bool(const FDataInterfaceUsageContext&)> Action, FDataInterfaceSearchOptions SearchOptions)
+void FNiagaraDataInterfaceUtilities::ForEachDataInterface(const UNiagaraSystem* NiagaraSystem, TFunction<bool(const FDataInterfaceUsageContext&)> Action)
 {
 	if (NiagaraSystem == nullptr)
 	{
@@ -316,16 +306,16 @@ void FNiagaraDataInterfaceUtilities::ForEachDataInterface(UNiagaraSystem* Niagar
 	FDataInterfaceUsageContext UsageContext;
 	UsageContext.OwnerObject = NiagaraSystem;
 
-	for (UNiagaraScript* SystemScript : { NiagaraSystem->GetSystemSpawnScript(), NiagaraSystem->GetSystemUpdateScript() } )
+	for (const UNiagaraScript* SystemScript : { NiagaraSystem->GetSystemSpawnScript(), NiagaraSystem->GetSystemUpdateScript() } )
 	{
 		if (SystemScript == nullptr)
 		{
 			continue;
 		}
 
-		if (const FNiagaraScriptExecutionParameterStore* ParameterStore = SystemScript->GetExecutionReadyParameterStore(ENiagaraSimTarget::CPUSim))
+		if (const FNiagaraScriptExecutionParameterStore* ParameterStore = const_cast<UNiagaraScript*>(SystemScript)->GetExecutionReadyParameterStore(ENiagaraSimTarget::CPUSim))
 		{
-			ForEachDataInterface(UsageContext, *ParameterStore, Action, SearchOptions);
+			ForEachDataInterface(UsageContext, *ParameterStore, Action);
 		}
 	}
 
@@ -348,7 +338,7 @@ void FNiagaraDataInterfaceUtilities::ForEachDataInterface(UNiagaraSystem* Niagar
 					return;
 				}
 
-				ForEachDataInterface(UsageContext, *ParameterStore, Action, SearchOptions);
+				ForEachDataInterface(UsageContext, *ParameterStore, Action);
 			}
 		);
 	}
diff --git a/Engine/Plugins/FX/Niagara/Source/Niagara/Private/NiagaraDataInterfaceCurveBase.cpp b/Engine/Plugins/FX/Niagara/Source/Niagara/Private/NiagaraDataInterfaceCurveBase.cpp
--- a/Engine/Plugins/FX/Niagara/Source/Niagara/Private/NiagaraDataInterfaceCurveBase.cpp
+++ b/Engine/Plugins/FX/Niagara/Source/Niagara/Private/NiagaraDataInterfaceCurveBase.cpp
@@ -476,7 +476,7 @@ void UNiagaraDataInterfaceCurveBase::SetShaderParameters(const FNiagaraDataInter
 	Parameters->CurveLUT = DIProxy.CurveLUT.SRV.IsValid() ? DIProxy.CurveLUT.SRV.GetReference() : FNiagaraRenderer::GetDummyFloatBuffer();
 }
 
-void UNiagaraDataInterfaceCurveBase::PostCompile()
+void UNiagaraDataInterfaceCurveBase::PostCompile(const UNiagaraSystem& OwningSystem)
 {
 #if WITH_EDITORONLY_DATA
 	SyncCurvesToAsset();
diff --git a/Engine/Plugins/FX/Niagara/Source/Niagara/Private/NiagaraSimCache.cpp b/Engine/Plugins/FX/Niagara/Source/Niagara/Private/NiagaraSimCache.cpp
--- a/Engine/Plugins/FX/Niagara/Source/Niagara/Private/NiagaraSimCache.cpp
+++ b/Engine/Plugins/FX/Niagara/Source/Niagara/Private/NiagaraSimCache.cpp
@@ -32,11 +32,6 @@
 UNiagaraSimCache::FOnCacheBeginWrite	UNiagaraSimCache::OnCacheBeginWrite;
 UNiagaraSimCache::FOnCacheEndWrite		UNiagaraSimCache::OnCacheEndWrite;
 
-namespace FNiagaraSimCacheInternal
-{
-	FNiagaraDataInterfaceUtilities::FDataInterfaceSearchOptions DISearchOptions = { true };
-}
-
 ////////////////////////////////////////////////////////////////////////////////////////////////////////////////////////////////////////////////////
 
 FNiagaraSimCacheCreateParameters FNiagaraSimCacheCreateParameters::CreateForDebugging()
@@ -428,7 +423,7 @@ bool UNiagaraSimCache::BeginWrite(FNiagaraSimCacheCreateParameters InCreateParam
 			CreateParameters.ExplicitCaptureAttributes.AddUnique(NameBuilder.ToString());
 		}
 
-		ForEachDataInterface(
+		FNiagaraDataInterfaceUtilities::ForEachDataInterface(
 			Helper.SystemInstance,
 			[&](const FNiagaraDataInterfaceUtilities::FDataInterfaceUsageContext& UsageContext)
 			{
@@ -440,8 +435,7 @@ bool UNiagaraSimCache::BeginWrite(FNiagaraSimCacheCreateParameters InCreateParam
 					}
 				}
 				return true;
-			},
-			FNiagaraSimCacheInternal::DISearchOptions
+			}
 		);
 	}
 
@@ -459,7 +453,7 @@ bool UNiagaraSimCache::BeginWrite(FNiagaraSimCacheCreateParameters InCreateParam
 	if ( CreateParameters.bAllowDataInterfaceCaching )
 	{
 		TSet<FNiagaraVariableBase> VisitedDataInterfaces;
-		ForEachDataInterface(
+		FNiagaraDataInterfaceUtilities::ForEachDataInterface(
 			Helper.SystemInstance,
 			[&](const FNiagaraVariableBase& Variable, UNiagaraDataInterface* DataInterface)
 			{
@@ -484,8 +478,7 @@ bool UNiagaraSimCache::BeginWrite(FNiagaraSimCacheCreateParameters InCreateParam
 					}
 				}
 				return true;
-			},
-			FNiagaraSimCacheInternal::DISearchOptions
+			}
 		);
 	}
 
@@ -740,7 +733,7 @@ bool UNiagaraSimCache::WriteFrame(UNiagaraComponent* NiagaraComponent, FNiagaraS
 		FString DataInterfaceName;
 		TSet<FNiagaraVariableBase> VisitedDataInterfaces;
 
-		ForEachDataInterface(
+		FNiagaraDataInterfaceUtilities::ForEachDataInterface(
 			Helper.SystemInstance,
 			[&](const FNiagaraVariableBase& Variable, UNiagaraDataInterface* DataInterface)
 			{
@@ -761,8 +754,7 @@ bool UNiagaraSimCache::WriteFrame(UNiagaraComponent* NiagaraComponent, FNiagaraS
 					}
 				}
 				return true;
-			},
-			FNiagaraSimCacheInternal::DISearchOptions
+			}
 		);
 
 		// A data interface failed to write information
@@ -1111,7 +1103,7 @@ bool UNiagaraSimCache::ReadFrame(int32 FrameIndex, float FrameFraction, FNiagara
 
 		bool bDataInterfacesSucess = true;
 
-		ForEachDataInterface(
+		FNiagaraDataInterfaceUtilities::ForEachDataInterface(
 			Helper.SystemInstance,
 			[&](const FNiagaraVariableBase& Variable, UNiagaraDataInterface* DataInterface)
 			{
@@ -1174,8 +1166,7 @@ bool UNiagaraSimCache::ReadFrame(int32 FrameIndex, float FrameFraction, FNiagara
 					}
 				}
 				return true;
-			},
-			FNiagaraSimCacheInternal::DISearchOptions
+			}
 		);
 
 
diff --git a/Engine/Plugins/FX/Niagara/Source/Niagara/Private/NiagaraSystem.cpp b/Engine/Plugins/FX/Niagara/Source/Niagara/Private/NiagaraSystem.cpp
--- a/Engine/Plugins/FX/Niagara/Source/Niagara/Private/NiagaraSystem.cpp
+++ b/Engine/Plugins/FX/Niagara/Source/Niagara/Private/NiagaraSystem.cpp
@@ -876,9 +876,7 @@ void UNiagaraSystem::ReportAnalyticsData(bool bIsCooking)
 		// gather data interface data
 		TSet<UNiagaraDataInterface*> DataInterfaces;
 		TArray<FString> DataInterfaceClasses;
-		FNiagaraDataInterfaceUtilities::FDataInterfaceSearchOptions SearchOptions;
-		SearchOptions.bIncludeInternal = true;
-		ForEachDataInterface(this, [&](const FNiagaraDataInterfaceUtilities::FDataInterfaceUsageContext& UsageContext) -> bool
+		FNiagaraDataInterfaceUtilities::ForEachDataInterface(this, [&](const FNiagaraDataInterfaceUtilities::FDataInterfaceUsageContext& UsageContext) -> bool
 		{
 			if (UsageContext.DataInterface && !DataInterfaces.Contains(UsageContext.DataInterface))
 			{
@@ -890,7 +888,7 @@ void UNiagaraSystem::ReportAnalyticsData(bool bIsCooking)
 				}
 			}
 			return true;
-		}, SearchOptions);
+		});
 		DataInterfaceClasses.Sort();
 		Attributes.Emplace(TEXT("DataInterfaceCount"), DataInterfaces.Num());
 		Attributes.Emplace(TEXT("ExposedDataInterfaceCount"), GetExposedParameters().GetDataInterfaces().Num());
@@ -2717,7 +2715,7 @@ void UNiagaraSystem::UpdatePostCompileDIInfo()
 		{
 			if(ResolvedDIInfo.ResolvedDataInterface)
 			{
-				ResolvedDIInfo.ResolvedDataInterface->PostCompile();
+				ResolvedDIInfo.ResolvedDataInterface->PostCompile(*this);
 			}
 		}		
 	});
diff --git a/Engine/Plugins/FX/Niagara/Source/Niagara/Public/NiagaraDataInterfaceUtilities.h b/Engine/Plugins/FX/Niagara/Source/Niagara/Public/NiagaraDataInterfaceUtilities.h
--- a/Engine/Plugins/FX/Niagara/Source/Niagara/Public/NiagaraDataInterfaceUtilities.h
+++ b/Engine/Plugins/FX/Niagara/Source/Niagara/Public/NiagaraDataInterfaceUtilities.h
@@ -13,39 +13,34 @@ namespace FNiagaraDataInterfaceUtilities
 {
 	struct FDataInterfaceUsageContext
 	{
-		UObject*				OwnerObject = nullptr;
+		const UObject*				OwnerObject = nullptr;
 		FNiagaraVariableBase	Variable;
 		UNiagaraDataInterface*	DataInterface = nullptr;
 	};
 
-	struct FDataInterfaceSearchOptions
-	{
-		bool bIncludeInternal = false;
-	};
-
 	// Finds all VM function calls made using the runtime resolved data interface.
 	// The same function call made be made multiple times across different scripts so you may see the same function multiple times
 	// The action should return True to continue iteration or False to stop
-	NIAGARA_API void ForEachVMFunction(UNiagaraDataInterface* ResolvedRuntimeDataInterface, UNiagaraSystem* NiagaraSystem, TFunction<bool(const UNiagaraScript*, const FVMExternalFunctionBindingInfo&)> Action);
+	NIAGARA_API void ForEachVMFunction(UNiagaraDataInterface* ResolvedRuntimeDataInterface, const UNiagaraSystem* NiagaraSystem, TFunction<bool(const UNiagaraScript*, const FVMExternalFunctionBindingInfo&)> Action);
 	NIAGARA_API void ForEachVMFunction(UNiagaraDataInterface* ResolvedRuntimeDataInterface, UNiagaraComponent* Component, TFunction<bool(const UNiagaraScript*, const FVMExternalFunctionBindingInfo&)> Action);
-	NIAGARA_API void ForEachVMFunction(UNiagaraDataInterface* ResolvedRuntimeDataInterface, FNiagaraSystemInstance* SystemInstance, TFunction<bool(const UNiagaraScript*, const FVMExternalFunctionBindingInfo&)> Action);
+	NIAGARA_API void ForEachVMFunction(UNiagaraDataInterface* ResolvedRuntimeDataInterface, const FNiagaraSystemInstance* SystemInstance, TFunction<bool(const UNiagaraScript*, const FVMExternalFunctionBindingInfo&)> Action);
 
 	// Finds all Gpu function calls made using the runtime resolved data interface.
 	// The same function call made be made multiple times across different scripts so you may see the same function multiple times
 	// The action should return True to continue iteration or False to stop
-	NIAGARA_API void ForEachGpuFunction(UNiagaraDataInterface* ResolvedRuntimeDataInterface, UNiagaraSystem* NiagaraSystem, TFunction<bool(const UNiagaraScript*, const FNiagaraDataInterfaceGeneratedFunction&)> Action);
+	NIAGARA_API void ForEachGpuFunction(UNiagaraDataInterface* ResolvedRuntimeDataInterface, const UNiagaraSystem* NiagaraSystem, TFunction<bool(const UNiagaraScript*, const FNiagaraDataInterfaceGeneratedFunction&)> Action);
 	NIAGARA_API void ForEachGpuFunction(UNiagaraDataInterface* ResolvedRuntimeDataInterface, UNiagaraComponent* Component, TFunction<bool(const UNiagaraScript*, const FNiagaraDataInterfaceGeneratedFunction&)> Action);
-	NIAGARA_API void ForEachGpuFunction(UNiagaraDataInterface* ResolvedRuntimeDataInterface, FNiagaraSystemInstance* SystemInstance, TFunction<bool(const UNiagaraScript*, const FNiagaraDataInterfaceGeneratedFunction&)> Action);
+	NIAGARA_API void ForEachGpuFunction(UNiagaraDataInterface* ResolvedRuntimeDataInterface, const FNiagaraSystemInstance* SystemInstance, TFunction<bool(const UNiagaraScript*, const FNiagaraDataInterfaceGeneratedFunction&)> Action);
 
 	// Loops over all data interfaces inside the SystemInstance
 	// The action should return True to continue iteration or False to stop
-	NIAGARA_API void ForEachDataInterface(FNiagaraSystemInstance* SystemInstance, TFunction<bool(const FNiagaraVariableBase Variable, UNiagaraDataInterface* DataInterface)> Action, FDataInterfaceSearchOptions SearchOptions = FDataInterfaceSearchOptions());
+	NIAGARA_API void ForEachDataInterface(const FNiagaraSystemInstance* SystemInstance, TFunction<bool(const FNiagaraVariableBase Variable, UNiagaraDataInterface* DataInterface)> Action);
 	// Loops over all data interfaces inside the SystemInstance
 	// The action should return True to continue iteration or False to stop
-	NIAGARA_API void ForEachDataInterface(FNiagaraSystemInstance* SystemInstance, TFunction<bool(const FDataInterfaceUsageContext&)> Action, FDataInterfaceSearchOptions SearchOptions = FDataInterfaceSearchOptions());
+	NIAGARA_API void ForEachDataInterface(const FNiagaraSystemInstance* SystemInstance, TFunction<bool(const FDataInterfaceUsageContext&)> Action);
 	// Loops over all data interfaces inside the NiagaraSystem
 	// The action should return True to continue iteration or False to stop
-	NIAGARA_API void ForEachDataInterface(UNiagaraSystem* NiagaraSystem, TFunction<bool(const FDataInterfaceUsageContext&)> Action, FDataInterfaceSearchOptions SearchOptions = FDataInterfaceSearchOptions());
+	NIAGARA_API void ForEachDataInterface(const UNiagaraSystem* NiagaraSystem, TFunction<bool(const FDataInterfaceUsageContext&)> Action);
 }
 
 #if WITH_EDITOR
diff --git a/Engine/Plugins/FX/Niagara/Source/NiagaraEditor/Private/NiagaraValidationRules.cpp b/Engine/Plugins/FX/Niagara/Source/NiagaraEditor/Private/NiagaraValidationRules.cpp
--- a/Engine/Plugins/FX/Niagara/Source/NiagaraEditor/Private/NiagaraValidationRules.cpp
+++ b/Engine/Plugins/FX/Niagara/Source/NiagaraEditor/Private/NiagaraValidationRules.cpp
@@ -787,9 +787,6 @@ void UNiagaraValidationRule_BannedDataInterfaces::CheckValidity(const FNiagaraVa
 {
 	UNiagaraSystem* NiagaraSystem = &Context.ViewModel->GetSystem();
 
-	FNiagaraDataInterfaceUtilities::FDataInterfaceSearchOptions SearchOptions;
-	SearchOptions.bIncludeInternal = true;
-
 	FNiagaraDataInterfaceUtilities::ForEachDataInterface(
 		NiagaraSystem,
 		[&](const FNiagaraDataInterfaceUtilities::FDataInterfaceUsageContext& UsageContext) -> bool
@@ -805,7 +802,7 @@ void UNiagaraValidationRule_BannedDataInterfaces::CheckValidity(const FNiagaraVa
 			static const FText EmitterDescFormat(LOCTEXT("BannedDataInteraceEmitterFormatDesc", "DataInterface '{0} - {1}' is banned on currently enabled platforms '{2}'"));
 
 			UObject* WarningObject = nullptr;
-			if (UNiagaraEmitter* NiagaraEmitter = Cast<UNiagaraEmitter>(UsageContext.OwnerObject))
+			if (UNiagaraEmitter* NiagaraEmitter = Cast<UNiagaraEmitter>(const_cast<UObject*>(UsageContext.OwnerObject)))
 			{
 				const TSharedPtr<FNiagaraEmitterHandleViewModel> EmitterViewModel = NiagaraValidation::GetEmitterViewModel(Context, NiagaraEmitter);
 				if (EmitterViewModel == nullptr)
@@ -837,7 +834,7 @@ void UNiagaraValidationRule_BannedDataInterfaces::CheckValidity(const FNiagaraVa
 					}
 				}
 			}
-			else if (UNiagaraSystem* NiagaraSystem = Cast<UNiagaraSystem>(UsageContext.OwnerObject))
+			else if (const UNiagaraSystem* NiagaraSystem = Cast<const UNiagaraSystem>(UsageContext.OwnerObject))
 			{
 				if (bBanOnCpu == true)
 				{
@@ -851,8 +848,7 @@ void UNiagaraValidationRule_BannedDataInterfaces::CheckValidity(const FNiagaraVa
 			}
 
 			return true;
-		},
-		SearchOptions
+		}
 	);
 }
 
diff --git a/Engine/Plugins/Runtime/HairStrands/Source/HairStrandsCore/Private/Niagara/NiagaraDataInterfaceHairStrands.cpp b/Engine/Plugins/Runtime/HairStrands/Source/HairStrandsCore/Private/Niagara/NiagaraDataInterfaceHairStrands.cpp
--- a/Engine/Plugins/Runtime/HairStrands/Source/HairStrandsCore/Private/Niagara/NiagaraDataInterfaceHairStrands.cpp
+++ b/Engine/Plugins/Runtime/HairStrands/Source/HairStrandsCore/Private/Niagara/NiagaraDataInterfaceHairStrands.cpp
@@ -1132,7 +1132,7 @@ void UNiagaraDataInterfaceHairStrands::SimCachePostReadFrame(void* OptionalPerIn
 	}
 }
 
-TArray<FNiagaraVariableBase> UNiagaraDataInterfaceHairStrands::GetSimCacheRendererAttributes(UObject* UsageContext) const
+TArray<FNiagaraVariableBase> UNiagaraDataInterfaceHairStrands::GetSimCacheRendererAttributes(const UObject* UsageContext) const
 {
 	TArray<FNiagaraVariableBase> HairStrandsCachedVariables;
 	if (const UNiagaraEmitter* UsageEmitter = Cast<UNiagaraEmitter>(UsageContext))
diff --git a/Engine/Plugins/Runtime/HairStrands/Source/HairStrandsCore/Public/Niagara/NiagaraDataInterfaceHairStrands.h b/Engine/Plugins/Runtime/HairStrands/Source/HairStrandsCore/Public/Niagara/NiagaraDataInterfaceHairStrands.h
--- a/Engine/Plugins/Runtime/HairStrands/Source/HairStrandsCore/Public/Niagara/NiagaraDataInterfaceHairStrands.h
+++ b/Engine/Plugins/Runtime/HairStrands/Source/HairStrandsCore/Public/Niagara/NiagaraDataInterfaceHairStrands.h
@@ -425,7 +425,7 @@ public:
 
 	/** Begin INiagaraSimCacheCustomStorageInterface Interface */
 	UE_API virtual void SimCachePostReadFrame(void* OptionalPerInstanceData, FNiagaraSystemInstance* SystemInstance) override;
-	UE_API virtual TArray<FNiagaraVariableBase> GetSimCacheRendererAttributes(UObject* UsageContext) const override;
+	UE_API virtual TArray<FNiagaraVariableBase> GetSimCacheRendererAttributes(const UObject* UsageContext) const override;
 	/** End INiagaraSimCacheCustomStorageInterface Interface */
 
 	/** Update the source component */
