Fix for broken static switches generating ensures/checks instead of just normal warning messages and a failed compile with the new compilation mode

#tests Generating broken graphs (invalid static variables) and validating how the compiler responds

diff --git a/Engine/Plugins/FX/Niagara/Source/NiagaraEditor/Private/NiagaraGraphDigest.cpp b/Engine/Plugins/FX/Niagara/Source/NiagaraEditor/Private/NiagaraGraphDigest.cpp
--- a/Engine/Plugins/FX/Niagara/Source/NiagaraEditor/Private/NiagaraGraphDigest.cpp
+++ b/Engine/Plugins/FX/Niagara/Source/NiagaraEditor/Private/NiagaraGraphDigest.cpp
@@ -1055,8 +1055,16 @@ void FNiagaraCompilationGraphInstanced::ValidateRefinement() const
 		switch (Node->NodeType)
 		{
 			case FNiagaraCompilationNode::ENodeType::StaticSwitch:
+			{
 				// only switches handled by the compiler are allowed to make it to this stage
-				ensure(Node->AsType<FNiagaraCompilationNodeStaticSwitch>()->bSetByCompiler);
+				const bool bSetByCompiler = Node->AsType<FNiagaraCompilationNodeStaticSwitch>()->bSetByCompiler;
+				
+				if (!bSetByCompiler)
+				{
+					UE_LOG(LogNiagaraEditor, Warning, TEXT("Encountered StaticSwitch that should have been resolved.  Graph will likely fail to compile.  Task: %s Node %s"),
+						*SourceScriptFullName, *Node->NodeName);
+				}
+			}
 			break;
 
 			// it's possible that we could eliminate these during instantiation as well, but currently
@@ -4260,13 +4268,25 @@ void FNiagaraCompilationNodeStaticSwitch::BuildParameterMapHistory(FParameterMap
 }
 void FNiagaraCompilationNodeStaticSwitch::Compile(FTranslator* Translator, TArray<int32>& Outputs) const
 {
-	check(bSetByCompiler);
+	int32 SelectorValue = INDEX_NONE;
+
+	if (!bSetByCompiler)
+	{
+		// all static switches should have already been resolved.  If they haven't it's likely because the graph
+		// has been setup incorrectly.  Possibly a selector pin is being used but it's associated static variable
+		// isn't being defined.
+		if (InputPins.IsValidIndex(SelectorPinIndex))
+		{
+			Translator->Error(LOCTEXT("CouldNotResolveStaticVarByPin", "Could not resolve static variable through pin."), this, &InputPins[SelectorPinIndex]);
+		}
+
+		// default to the first entry
+		SelectorValue = 0;
+	}
 
 	const int32 OutputPinCount = OutputPins.Num();
 	Outputs.Init(INDEX_NONE, OutputPinCount);
 
-	int32 SelectorValue = INDEX_NONE;
-
 	if (bSetByCompiler)
 	{
 		FNiagaraFixedConstantResolver ConstantResolver(Translator);
