- Multiply the debug draw color by alpha so we can have translucent lines

#tests Draw debug lines using the debug draw DI and modify the alpha

diff --git a/Engine/Plugins/FX/Niagara/Shaders/Private/NiagaraDebugDraw.usf b/Engine/Plugins/FX/Niagara/Shaders/Private/NiagaraDebugDraw.usf
--- a/Engine/Plugins/FX/Niagara/Shaders/Private/NiagaraDebugDraw.usf
+++ b/Engine/Plugins/FX/Niagara/Shaders/Private/NiagaraDebugDraw.usf
@@ -72,6 +72,7 @@ void MainPS (
 	const bool bIsHidden	= Position.z < Depth; // Reverse-Z
 
 	const float ColorScale = bIsHidden ? OccludedColorScale : 1;
-	OutColor = Color * ColorScale;
+	OutColor.a = Color.a * ColorScale;
+	OutColor.rgb = Color.rgb * OutColor.a;
 }
 #endif
