- Fix emitter renderer tick ordering issue when combined with LWE

#tests Create a system with a LWE, create another standard emitter and use add the component renderer to it display static meshes, note that it is not visible until you disable the LWE

diff --git a/Engine/Plugins/FX/Niagara/Source/Niagara/Classes/NiagaraEmitter.h b/Engine/Plugins/FX/Niagara/Source/Niagara/Classes/NiagaraEmitter.h
--- a/Engine/Plugins/FX/Niagara/Source/Niagara/Classes/NiagaraEmitter.h
+++ b/Engine/Plugins/FX/Niagara/Source/Niagara/Classes/NiagaraEmitter.h
@@ -423,6 +423,9 @@ struct FVersionedNiagaraEmitterData
 	template<typename TAction>
 	void ForEachEnabledRenderer(TAction Func) const;
 
+	template<typename TAction>
+	void ForEachEnabledRendererWithIndex(TAction Func) const;
+
 	template<typename TAction>
 	void ForEachScript(TAction Func) const;
 
@@ -1093,6 +1096,20 @@ void FVersionedNiagaraEmitterData::ForEachEnabledRenderer(TAction Func) const
 		}
 	}
 }
+
+template<typename TAction>
+void FVersionedNiagaraEmitterData::ForEachEnabledRendererWithIndex(TAction Func) const
+{
+	for (int32 i=0; i < RendererProperties.Num(); ++i)
+	{
+		UNiagaraRendererProperties* Renderer = RendererProperties[i];
+		if (Renderer && Renderer->GetIsEnabled() && Renderer->IsSimTargetSupported(this->SimTarget))
+		{
+			Func(Renderer, i);
+		}
+	}
+}
+
 template<typename TAction>
 void FVersionedNiagaraEmitterData::ForEachScript(TAction Func) const
 {
diff --git a/Engine/Plugins/FX/Niagara/Source/Niagara/Classes/NiagaraEmitterHandle.h b/Engine/Plugins/FX/Niagara/Source/Niagara/Classes/NiagaraEmitterHandle.h
--- a/Engine/Plugins/FX/Niagara/Source/Niagara/Classes/NiagaraEmitterHandle.h
+++ b/Engine/Plugins/FX/Niagara/Source/Niagara/Classes/NiagaraEmitterHandle.h
@@ -98,6 +98,9 @@ public:
 	FSimpleMulticastDelegate& OnEmitterModeChanged() { return OnEmitterModeChangedDelegate; }
 #endif
 
+	/** Iterate over all the enabled renderers on the emitter */
+	NIAGARA_API void ForEachEnabledRendererWithIndex(TFunction<void(const class UNiagaraRendererProperties*, int32)> Delegate) const;
+
 	/** Gets a unique name for this emitter instance for use in scripts and parameter stores etc.*/
 	NIAGARA_API FString GetUniqueInstanceName()const;
 
diff --git a/Engine/Plugins/FX/Niagara/Source/Niagara/Internal/Stateless/NiagaraStatelessEmitter.h b/Engine/Plugins/FX/Niagara/Source/Niagara/Internal/Stateless/NiagaraStatelessEmitter.h
--- a/Engine/Plugins/FX/Niagara/Source/Niagara/Internal/Stateless/NiagaraStatelessEmitter.h
+++ b/Engine/Plugins/FX/Niagara/Source/Niagara/Internal/Stateless/NiagaraStatelessEmitter.h
@@ -97,6 +97,24 @@ public:
 			}
 		}
 	}
+	template<typename TAction>
+	void ForEachEnabledRendererWithIndex(TAction Func) const
+	{
+		if (RendererProperties.Num() == 0)
+		{
+			return;
+		}
+
+		const ENiagaraSimTarget SimTarget = ComputeSimTarget();
+		for (int32 i=0; i < RendererProperties.Num(); ++i)
+		{
+			UNiagaraRendererProperties* Renderer = RendererProperties[i];
+			if (Renderer && Renderer->GetIsEnabled() && Renderer->IsSimTargetSupported(SimTarget))
+			{
+				Func(Renderer, i);
+			}
+		}
+	}
 
 	template<typename TAction>
 	void ForEachRenderer(TAction Func) const
diff --git a/Engine/Plugins/FX/Niagara/Source/Niagara/Private/NiagaraEmitterHandle.cpp b/Engine/Plugins/FX/Niagara/Source/Niagara/Private/NiagaraEmitterHandle.cpp
--- a/Engine/Plugins/FX/Niagara/Source/Niagara/Private/NiagaraEmitterHandle.cpp
+++ b/Engine/Plugins/FX/Niagara/Source/Niagara/Private/NiagaraEmitterHandle.cpp
@@ -266,6 +266,21 @@ FString FNiagaraEmitterHandle::GetUniqueInstanceName() const
 	return FString();
 }
 
+void FNiagaraEmitterHandle::ForEachEnabledRendererWithIndex(TFunction<void(const class UNiagaraRendererProperties*, int32)> Delegate) const
+{
+	if (EmitterMode == ENiagaraEmitterMode::Standard)
+	{
+		if (FVersionedNiagaraEmitterData* EmitterData = VersionedInstance.GetEmitterData())
+		{
+			EmitterData->ForEachEnabledRendererWithIndex(Delegate);
+		}
+	}
+	else if (StatelessEmitter)
+	{
+		StatelessEmitter->ForEachEnabledRendererWithIndex(Delegate);
+	}
+}
+
 #if WITH_EDITORONLY_DATA
 
 bool FNiagaraEmitterHandle::NeedsRecompile() const
diff --git a/Engine/Plugins/FX/Niagara/Source/Niagara/Private/NiagaraSystem.cpp b/Engine/Plugins/FX/Niagara/Source/Niagara/Private/NiagaraSystem.cpp
--- a/Engine/Plugins/FX/Niagara/Source/Niagara/Private/NiagaraSystem.cpp
+++ b/Engine/Plugins/FX/Niagara/Source/Niagara/Private/NiagaraSystem.cpp
@@ -1505,16 +1505,17 @@ void UNiagaraSystem::PrecachePSOs()
 	FMaterialInterfacePSOPrecacheParamsList MaterialInterfacePSOPrecacheParamsList;
 	for (const FNiagaraEmitterHandle& EmitterHandle : GetEmitterHandles())
 	{
-		FVersionedNiagaraEmitterData* EmitterData = EmitterHandle.GetEmitterData();
-		if (EmitterData && EmitterHandle.GetIsEnabled())
+		if (!EmitterHandle.GetIsEnabled())
 		{
-			EmitterData->ForEachEnabledRenderer(
-				[&](UNiagaraRendererProperties* Properties)
-				{
-					Properties->CollectPSOPrecacheData(EmitterInstance, MaterialInterfacePSOPrecacheParamsList);
-				}
-			);
+			continue;
 		}
+
+		EmitterHandle.ForEachEnabledRendererWithIndex(
+			[&](const UNiagaraRendererProperties* Properties, int32 EmitterIndex)
+			{
+				Properties->CollectPSOPrecacheData(EmitterInstance, MaterialInterfacePSOPrecacheParamsList);
+			}
+		);
 	}
 
 	LaunchPSOPrecaching(MaterialInterfacePSOPrecacheParamsList);
@@ -2235,17 +2236,11 @@ void UNiagaraSystem::ComputeEmittersExecutionOrder()
 	for (int32 EmitterIdx = 0; EmitterIdx < NumEmitters; ++EmitterIdx)
 	{
 		const FNiagaraEmitterHandle& EmitterHandle = EmitterHandles[EmitterIdx];
-		FVersionedNiagaraEmitterData* EmitterData = EmitterHandle.GetEmitterData();
 
 		EmitterExecutionOrder[EmitterIdx].bStartNewOverlapGroup = false;
 		EmitterExecutionOrder[EmitterIdx].EmitterIndex = EmitterIdx;
 		EmitterPriorities[EmitterIdx] = -1;
 
-		if (EmitterData == nullptr)
-		{
-			continue;
-		}
-
 		if (!EmitterHandle.GetIsEnabled())
 		{
 			continue;
@@ -2253,83 +2248,85 @@ void UNiagaraSystem::ComputeEmittersExecutionOrder()
 
 		EmitterDependencies.SetNum(0, EAllowShrinking::No);
 
-		for (const FNiagaraDataInterfaceEmitterBinding& EmitterBinding : EmitterData->EmitterDependencies)
+		FVersionedNiagaraEmitterData* EmitterData = EmitterHandle.GetEmitterData();
+		if (EmitterData)
 		{
-			const FNiagaraEmitterHandle* ResolvedHandle = EmitterBinding.ResolveHandle(this, &EmitterHandle);
-			if (ResolvedHandle)
+			for (const FNiagaraDataInterfaceEmitterBinding& EmitterBinding : EmitterData->EmitterDependencies)
 			{
-				EmitterDependencies.Add(ResolvedHandle->GetInstance());
+				const FNiagaraEmitterHandle* ResolvedHandle = EmitterBinding.ResolveHandle(this, &EmitterHandle);
+				if (ResolvedHandle)
+				{
+					EmitterDependencies.Add(ResolvedHandle->GetInstance());
+				}
 			}
-		}
 
-		if (EmitterData->SimTarget == ENiagaraSimTarget::GPUComputeSim && EmitterData->GetGPUComputeScript())
-		{
-			// GPU emitters have a combined execution context for spawn and update.
-			FindDataInterfaceDependencies(EmitterData, EmitterData->GetGPUComputeScript(), EmitterDependencies);
-		}
-		else
-		{
-			// CPU emitters have separate contexts for spawn and update, so we need to gather DIs from both. They also support events,
-			// so we need to look at the event sources for extra dependencies.
-			FindDataInterfaceDependencies(EmitterData, EmitterData->SpawnScriptProps.Script, EmitterDependencies);
-			FindDataInterfaceDependencies(EmitterData, EmitterData->UpdateScriptProps.Script, EmitterDependencies);
-			FindEventDependencies(EmitterData, EmitterDependencies);
-		}
+			if (EmitterData->SimTarget == ENiagaraSimTarget::GPUComputeSim && EmitterData->GetGPUComputeScript())
+			{
+				// GPU emitters have a combined execution context for spawn and update.
+				FindDataInterfaceDependencies(EmitterData, EmitterData->GetGPUComputeScript(), EmitterDependencies);
+			}
+			else
+			{
+				// CPU emitters have separate contexts for spawn and update, so we need to gather DIs from both. They also support events,
+				// so we need to look at the event sources for extra dependencies.
+				FindDataInterfaceDependencies(EmitterData, EmitterData->SpawnScriptProps.Script, EmitterDependencies);
+				FindDataInterfaceDependencies(EmitterData, EmitterData->UpdateScriptProps.Script, EmitterDependencies);
+				FindEventDependencies(EmitterData, EmitterDependencies);
+			}
 
-		// Map the pointers returned by the emitter to indices inside the Emitters array. This is O(N^2), but we expect
-		// to have few dependencies, so in practice it should be faster than a TMap. If it gets out of hand, we can also
-		// ask the DIs to give us indices directly, since they probably got the pointers by scanning the array we gave them
-		// through GetEmitters() anyway.
-		for (int32 DepIdx = 0; DepIdx < EmitterDependencies.Num(); ++DepIdx)
-		{
-			for (int32 OtherEmitterIdx = 0; OtherEmitterIdx < NumEmitters; ++OtherEmitterIdx)
+			// Map the pointers returned by the emitter to indices inside the Emitters array. This is O(N^2), but we expect
+			// to have few dependencies, so in practice it should be faster than a TMap. If it gets out of hand, we can also
+			// ask the DIs to give us indices directly, since they probably got the pointers by scanning the array we gave them
+			// through GetEmitters() anyway.
+			for (int32 DepIdx = 0; DepIdx < EmitterDependencies.Num(); ++DepIdx)
 			{
-				if (EmitterDependencies[DepIdx] == EmitterHandles[OtherEmitterIdx].GetInstance())
+				for (int32 OtherEmitterIdx = 0; OtherEmitterIdx < NumEmitters; ++OtherEmitterIdx)
 				{
-					const bool HasSourceEmitter = EmitterHandles[EmitterIdx].GetInstance().Emitter != nullptr;
-					const bool HasDependentEmitter = EmitterHandles[OtherEmitterIdx].GetInstance().Emitter != nullptr;
-
-					// check to see if the emitter we're dependent on may have been culled during the cook
-					if (HasSourceEmitter && !HasDependentEmitter)
+					if (EmitterDependencies[DepIdx] == EmitterHandles[OtherEmitterIdx].GetInstance())
 					{
-						UE_LOG(LogNiagara, Error, TEXT("Emitter[%s] depends on Emitter[%s] which is not available (has scalability removed it during a cook?)."),
-							*EmitterHandles[EmitterIdx].GetName().ToString(),
-							*EmitterHandles[OtherEmitterIdx].GetName().ToString());
-					}
+						const bool HasSourceEmitter = EmitterHandles[EmitterIdx].GetInstance().Emitter != nullptr;
+						const bool HasDependentEmitter = EmitterHandles[OtherEmitterIdx].GetInstance().Emitter != nullptr;
 
-					// Some DIs might read from the same emitter they're applied to. We don't care about dependencies on self.
-					if (EmitterIdx != OtherEmitterIdx)
-					{
-						EmitterDependencyGraph.SetRange(EmitterIdx * NumEmitters + OtherEmitterIdx, 1, true);
-						bHasEmitterDependencies = true;
+						// check to see if the emitter we're dependent on may have been culled during the cook
+						if (HasSourceEmitter && !HasDependentEmitter)
+						{
+							UE_LOG(LogNiagara, Error, TEXT("Emitter[%s] depends on Emitter[%s] which is not available (has scalability removed it during a cook?)."),
+								*EmitterHandles[EmitterIdx].GetName().ToString(),
+								*EmitterHandles[OtherEmitterIdx].GetName().ToString());
+						}
+
+						// Some DIs might read from the same emitter they're applied to. We don't care about dependencies on self.
+						if (EmitterIdx != OtherEmitterIdx)
+						{
+							EmitterDependencyGraph.SetRange(EmitterIdx * NumEmitters + OtherEmitterIdx, 1, true);
+							bHasEmitterDependencies = true;
+						}
+						break;
 					}
-					break;
 				}
 			}
 		}
 
 		// Determine renderer execution order for PostTick and Completion for any renderers that opt into it
-		for (int32 RendererIndex = 0; RendererIndex < EmitterData->GetRenderers().Num(); ++RendererIndex)
-		{
-			const UNiagaraRendererProperties* Renderer = EmitterData->GetRenderers()[RendererIndex];
-			if (Renderer && Renderer->GetIsEnabled() && Renderer->IsSimTargetSupported(EmitterData->SimTarget))
+		EmitterHandle.ForEachEnabledRendererWithIndex(
+			[&](const UNiagaraRendererProperties* RendererProperties, int32 RendererIndex)
 			{
 				FNiagaraRendererExecutionIndex ExecutionIndex;
 				ExecutionIndex.EmitterIndex = EmitterIdx;
 				ExecutionIndex.EmitterRendererIndex = RendererIndex;
 				ExecutionIndex.SystemRendererIndex = SystemRendererIndex;
 
-				if (Renderer->NeedsSystemPostTick())
+				if (RendererProperties->NeedsSystemPostTick())
 				{
 					RendererPostTickOrder.Add(ExecutionIndex);
 				}
-				if (Renderer->NeedsSystemCompletion())
+				if (RendererProperties->NeedsSystemCompletion())
 				{
 					RendererCompletionOrder.Add(ExecutionIndex);
 				}
 				++SystemRendererIndex;
 			}
-		}
+		);
 	}
 
 	if (bHasEmitterDependencies)
@@ -2409,34 +2406,12 @@ void UNiagaraSystem::ComputeRenderersDrawOrder()
 			continue;
 		}
 
-		if (EmitterHandle.GetEmitterMode() == ENiagaraEmitterMode::Standard)
-		{
-			if (FVersionedNiagaraEmitterData* EmitterData = EmitterHandle.GetEmitterData())
+		EmitterHandle.ForEachEnabledRendererWithIndex(
+			[&](const UNiagaraRendererProperties* Properties, int32 EmitterIndex)
 			{
-				EmitterData->ForEachEnabledRenderer(
-					[&](UNiagaraRendererProperties* Properties)
-					{
-						RendererSortInfo.Emplace(Properties->SortOrderHint, RendererSortInfo.Num());
-					}
-				);
-			}
-		}
-		//-TODO:Stateless:Abstract this away
-		else
-		{
-			ensure(EmitterHandle.GetEmitterMode() == ENiagaraEmitterMode::Stateless);
-			const UNiagaraStatelessEmitter* StatelessEmitter = EmitterHandle.GetStatelessEmitter();
-			if (StatelessEmitter)
-			{
-				StatelessEmitter->ForEachEnabledRenderer(
-					[&](UNiagaraRendererProperties* Properties)
-					{
-						RendererSortInfo.Emplace(Properties->SortOrderHint, RendererSortInfo.Num());
-					}
-				);
+				RendererSortInfo.Emplace(Properties->SortOrderHint, RendererSortInfo.Num());
 			}
-		}
-		//-TODO:Stateless:Abstract this away
+		);
 	}
 
 	// We sort by the sort hint in order to guarantee that we submit according to the preferred sort order..
