- Fix reading component list that can be modified during iteration

#rnx

diff --git a/Engine/Plugins/FX/Niagara/Source/NiagaraEditor/Private/Customizations/NiagaraComponentDetails.cpp b/Engine/Plugins/FX/Niagara/Source/NiagaraEditor/Private/Customizations/NiagaraComponentDetails.cpp
--- a/Engine/Plugins/FX/Niagara/Source/NiagaraEditor/Private/Customizations/NiagaraComponentDetails.cpp
+++ b/Engine/Plugins/FX/Niagara/Source/NiagaraEditor/Private/Customizations/NiagaraComponentDetails.cpp
@@ -259,15 +259,21 @@ FReply FNiagaraComponentDetails::OnResetSelectedSystem()
 		{
 			if (AActor* Actor = Cast<AActor>(SelectedObjects[Idx].Get()))
 			{
+				// Note: Must gather then reset as the array we are iterating will be modified with the ReregisterComponent call
+				TArray<UNiagaraComponent*> ComponentsToReset;
 				for (UActorComponent* AC : Actor->GetComponents())
 				{
-					UNiagaraComponent* NiagaraComponent = Cast<UNiagaraComponent>(AC);
-					if (NiagaraComponent)
+					if (UNiagaraComponent* NiagaraComponent = Cast<UNiagaraComponent>(AC))
 					{
-						NiagaraComponent->Activate(true);
-						NiagaraComponent->ReregisterComponent();
+						ComponentsToReset.Add(NiagaraComponent);
 					}
 				}
+
+				for (UNiagaraComponent* NiagaraComponent : ComponentsToReset)
+				{
+					NiagaraComponent->Activate(true);
+					NiagaraComponent->ReregisterComponent();
+				}
 			}
 			else if (UNiagaraComponent* NiagaraComponent = Cast<UNiagaraComponent>(SelectedObjects[Idx].Get()))
 			{
diff --git a/Engine/Plugins/FX/Niagara/Source/NiagaraEditor/Private/NiagaraEditorModule.cpp b/Engine/Plugins/FX/Niagara/Source/NiagaraEditor/Private/NiagaraEditorModule.cpp
--- a/Engine/Plugins/FX/Niagara/Source/NiagaraEditor/Private/NiagaraEditorModule.cpp
+++ b/Engine/Plugins/FX/Niagara/Source/NiagaraEditor/Private/NiagaraEditorModule.cpp
@@ -2397,15 +2397,21 @@ void FNiagaraEditorModule::OnExecParticleInvoked(const TCHAR* Str)
 			// Reset ALL emitters in the level
 			for (TObjectIterator<AActor> It; It; ++It)
 			{
+				// Note: Must gather then reset as the array we are iterating will be modified with the ReregisterComponent call
+				TArray<UNiagaraComponent*> ComponentsToReset;
 				for (UActorComponent* AC : It->GetComponents())
 				{
-					UNiagaraComponent* NiagaraComponent = Cast<UNiagaraComponent>(AC);
-					if (NiagaraComponent)
+					if (UNiagaraComponent* NiagaraComponent = Cast<UNiagaraComponent>(AC))
 					{
-						NiagaraComponent->Activate(true);
-						NiagaraComponent->ReregisterComponent();
+						ComponentsToReset.Add(NiagaraComponent);
 					}
 				}
+
+				for (UNiagaraComponent* NiagaraComponent : ComponentsToReset)
+				{
+					NiagaraComponent->Activate(true);
+					NiagaraComponent->ReregisterComponent();
+				}
 			}
 		}
 	}
