- Adjust customization for binding so that the parameter can be clicked on to change what it's bound to
- This is to reduce confusion as it wasn't obvious to click the drop down to the right

#rnx
#tests Add mesh renderer / nanite renderer, swap meshes binding / mesh binding / LOD level binding / per instance data bindings

diff --git a/Engine/Plugins/FX/Niagara/Source/NiagaraEditor/Internal/Customizations/NiagaraParameterBindingAdapterCustomization.cpp b/Engine/Plugins/FX/Niagara/Source/NiagaraEditor/Internal/Customizations/NiagaraParameterBindingAdapterCustomization.cpp
--- a/Engine/Plugins/FX/Niagara/Source/NiagaraEditor/Internal/Customizations/NiagaraParameterBindingAdapterCustomization.cpp
+++ b/Engine/Plugins/FX/Niagara/Source/NiagaraEditor/Internal/Customizations/NiagaraParameterBindingAdapterCustomization.cpp
@@ -275,8 +275,6 @@ void FNiagaraParameterBindingAdapterCustomization::CustomizeHeader(TSharedRef<IP
 			ConstantValueParameterEditor->SetOnValueChanged(SNiagaraParameterEditor::FOnValueChange::CreateSP(this, &FNiagaraParameterBindingAdapterCustomization::OnConstantValueChanged));
 
 			ParameterPanel->AddSlot()
-			.AutoWidth()
-			.Padding(0, 0, 0, 3)
 			[
 				SNew(SBox)
 				.HAlign(ConstantValueParameterEditor->GetHorizontalAlignment())
@@ -284,7 +282,22 @@ void FNiagaraParameterBindingAdapterCustomization::CustomizeHeader(TSharedRef<IP
 				.IsEnabled(this, &FNiagaraParameterBindingAdapterCustomization::IsConstantEnabled)
 				.Visibility(this, &FNiagaraParameterBindingAdapterCustomization::IsConstantVisibile)
 				[
-					ConstantValueParameterEditor.ToSharedRef()
+					SNew(SHorizontalBox)
+					+SHorizontalBox::Slot()
+					[
+						ConstantValueParameterEditor.ToSharedRef()
+					]
+					+SHorizontalBox::Slot()
+					.HAlign(HAlign_Right)
+					[
+						SNew(SComboButton)
+						.ButtonStyle(FAppStyle::Get(), "HoverHintOnly")
+						.ForegroundColor(FSlateColor::UseForeground())
+						.OnGetMenuContent(this, &FNiagaraParameterBindingAdapterCustomization::OnGetMenuContent)
+						.ContentPadding(2)
+						.MenuPlacement(MenuPlacement_BelowRightAnchor)
+						.ToolTipText(this, &FNiagaraParameterBindingAdapterCustomization::GetTooltipText)
+					]
 				]
 			];
 		}
@@ -297,8 +310,7 @@ void FNiagaraParameterBindingAdapterCustomization::CustomizeHeader(TSharedRef<IP
 		.ContentPadding(2)
 		.IsEnabled(this, &FNiagaraParameterBindingAdapterCustomization::IsBindingEnabled)
 		.Visibility(this, &FNiagaraParameterBindingAdapterCustomization::IsBindingVisibile)
-		.HasDownArrow(false)
-		.IsFocusable(false)
+		.OnGetMenuContent(this, &FNiagaraParameterBindingAdapterCustomization::OnGetMenuContent)
 		.ButtonContent()
 		[
 			SNew(SNiagaraParameterName)
@@ -317,23 +329,11 @@ void FNiagaraParameterBindingAdapterCustomization::CustomizeHeader(TSharedRef<IP
 	[
 		SNew(SHorizontalBox)
 		+ SHorizontalBox::Slot()
+		.AutoWidth()
 		[
 			ParameterPanel.ToSharedRef()
 		]
 		+ SHorizontalBox::Slot()
-		.FillWidth(1.0f)
-		.HAlign(HAlign_Right)
-		.Padding(3, 0, 0, 0)
-		[
-			SNew(SComboButton)
-			.ButtonStyle(FAppStyle::Get(), "HoverHintOnly")
-			.ForegroundColor(FSlateColor::UseForeground())
-			.OnGetMenuContent(this, &FNiagaraParameterBindingAdapterCustomization::OnGetMenuContent)
-			.ContentPadding(2)
-			.MenuPlacement(MenuPlacement_BelowRightAnchor)
-			.ToolTipText(this, &FNiagaraParameterBindingAdapterCustomization::GetTooltipText)
-		]
-		+ SHorizontalBox::Slot()
 		.AutoWidth()
 		.VAlign(VAlign_Center)
 		.Padding(2, 1)
