- LWE - Set the sim target on the compiled data set to match the mode
- This fixes a mismatch that can cause attributes to not upload to the GPU when we are in CPU mode due to using the light / decal renderer

#tests Create a LWE with Mesh+Sprite+Decal renderer and observe that not all renders will be visible

diff --git a/Engine/Plugins/FX/Niagara/Source/Niagara/Private/Stateless/NiagaraStatelessEmitter.cpp b/Engine/Plugins/FX/Niagara/Source/Niagara/Private/Stateless/NiagaraStatelessEmitter.cpp
--- a/Engine/Plugins/FX/Niagara/Source/Niagara/Private/Stateless/NiagaraStatelessEmitter.cpp
+++ b/Engine/Plugins/FX/Niagara/Source/Niagara/Private/Stateless/NiagaraStatelessEmitter.cpp
@@ -382,6 +382,10 @@ void UNiagaraStatelessEmitter::CacheFromCompiledData()
 		UE_LOG(LogNiagara, Log, TEXT("Stateless Emitter (%s) can not execute on any available path and will be disabled."), *GetFullName());
 	}
 
+	// Compute the simulation target
+	// This must always be done to ensure we filter renderers correctly and that the data set uses the correct target
+	StatelessEmitterData->SimTarget = ComputeSimTarget();
+
 	// Build data set
 	BuildCompiledDataSet();
 
@@ -427,9 +431,6 @@ void UNiagaraStatelessEmitter::CacheFromCompiledData()
 	// Resolve scalability settings
 	ResolveScalabilitySettings();
 
-	// Compute the simulation target, this must always be done to ensure we filter renderers correctly
-	StatelessEmitterData->SimTarget = ComputeSimTarget();
-
 	// Build buffers that are shared across all instances
 	//-OPT: We should be able to build and serialize this data as part of the UNiagaraStatelessEmitter, potentially all of this data even since it's immutable and does not change at runtime
 	if (StatelessEmitterData->bCanEverExecute)
@@ -551,7 +552,7 @@ void UNiagaraStatelessEmitter::BuildCompiledDataSet()
 {
 #if WITH_EDITORONLY_DATA
 	ParticleDataSetCompiledData.Empty();
-	ParticleDataSetCompiledData.SimTarget = ENiagaraSimTarget::GPUComputeSim;
+	ParticleDataSetCompiledData.SimTarget = StatelessEmitterData->SimTarget;
 
 	if (EmitterTemplate)
 	{
