- Disable GPU Scene for Niagara Mesh Renderer on mobile
- The lack of order preserving draws can cause flickering in some cases

diff --git a/Engine/Plugins/FX/Niagara/Source/Niagara/Private/NiagaraRendererMeshes.cpp b/Engine/Plugins/FX/Niagara/Source/Niagara/Private/NiagaraRendererMeshes.cpp
--- a/Engine/Plugins/FX/Niagara/Source/Niagara/Private/NiagaraRendererMeshes.cpp
+++ b/Engine/Plugins/FX/Niagara/Source/Niagara/Private/NiagaraRendererMeshes.cpp
@@ -317,11 +317,7 @@ void FNiagaraRendererMeshes::PrepareParticleMeshRenderData(FRHICommandListBase&
 {
 	ParticleMeshRenderData.Collector = &Collector;
 
-#if NIAGARA_ENABLE_GPU_SCENE_MESHES
-	ParticleMeshRenderData.bUseGPUScene = UseGPUScene(SceneProxy->GetScene().GetShaderPlatform(), FeatureLevel);
-#else
-	ParticleMeshRenderData.bUseGPUScene = false;
-#endif
+	ParticleMeshRenderData.bUseGPUScene = FNiagaraMeshVertexFactory::IsGPUSceneEnabled(SceneProxy->GetScene().GetShaderPlatform(), FeatureLevel);
 
 	// Anything to render?
 	ParticleMeshRenderData.DynamicDataMesh = static_cast<FNiagaraDynamicDataMesh*>(InDynamicData);
diff --git a/Engine/Plugins/FX/Niagara/Source/NiagaraVertexFactories/Private/NiagaraMeshVertexFactory.cpp b/Engine/Plugins/FX/Niagara/Source/NiagaraVertexFactories/Private/NiagaraMeshVertexFactory.cpp
--- a/Engine/Plugins/FX/Niagara/Source/NiagaraVertexFactories/Private/NiagaraMeshVertexFactory.cpp
+++ b/Engine/Plugins/FX/Niagara/Source/NiagaraVertexFactories/Private/NiagaraMeshVertexFactory.cpp
@@ -14,6 +14,9 @@
 #include "GlobalRenderResources.h"
 #include "DataDrivenShaderPlatformInfo.h"
 
+// Master control switch for GPU scene on all platforms
+#define NIAGARA_ENABLE_GPU_SCENE_MESHES 1
+
 IMPLEMENT_GLOBAL_SHADER_PARAMETER_STRUCT(FNiagaraMeshUniformParameters, "NiagaraMeshVF");
 
 class FNiagaraMeshVertexFactoryShaderParametersVS : public FNiagaraVertexFactoryShaderParametersBase
@@ -68,6 +71,19 @@ public:
 
 };
 
+bool FNiagaraMeshVertexFactory::IsGPUSceneEnabled(const FStaticShaderPlatform Platform, const FStaticFeatureLevel FeatureLevel)
+{
+#if NIAGARA_ENABLE_GPU_SCENE_MESHES
+	if (FeatureLevel == ERHIFeatureLevel::ES3_1)
+	{
+		return false;
+	}
+	return UseGPUScene(Platform, FeatureLevel);
+#else
+	return false;
+#endif
+}
+
 IMPLEMENT_TYPE_LAYOUT(FNiagaraMeshVertexFactoryShaderParametersPS);
 
 void FNiagaraMeshVertexFactory::GetVertexElements(ERHIFeatureLevel::Type FeatureLevel, bool bSupportsManualVertexFetch, FStaticMeshDataType& Data, FVertexDeclarationElementList& Elements, FVertexStreamList& InOutStreams)
@@ -165,14 +181,16 @@ void FNiagaraMeshVertexFactory::ModifyCompilationEnvironment(const FVertexFactor
 {
 	FNiagaraVertexFactoryBase::ModifyCompilationEnvironment(Parameters, OutEnvironment);
 
+	const ERHIFeatureLevel::Type MaxSupportedFeatureLevel = GetMaxSupportedFeatureLevel(Parameters.Platform);
+	const bool bUseGPUScene = IsGPUSceneEnabled(Parameters.Platform, MaxSupportedFeatureLevel);
+
 	// Set a define so we can tell in MaterialTemplate.usf when we are compiling a mesh particle vertex factory
 	OutEnvironment.SetDefine(TEXT("NIAGARA_MESH_FACTORY"), TEXT("1"));
 	OutEnvironment.SetDefine(TEXT("NIAGARA_MESH_INSTANCED"), TEXT("1"));
 	OutEnvironment.SetDefine(TEXT("NiagaraVFLooseParameters"), TEXT("NiagaraMeshVF"));
+	OutEnvironment.SetDefine(TEXT("NIAGARA_MESH_WITH_GPU_SCENE"), bUseGPUScene);
 
 #if NIAGARA_ENABLE_GPU_SCENE_MESHES
-	const ERHIFeatureLevel::Type MaxSupportedFeatureLevel = GetMaxSupportedFeatureLevel(Parameters.Platform);
-	const bool bUseGPUScene = UseGPUScene(Parameters.Platform, MaxSupportedFeatureLevel);
 	const bool bSupportsPrimitiveIdStream = Parameters.VertexFactoryType->SupportsPrimitiveIdStream();
 	
 	OutEnvironment.SetDefine(TEXT("VF_SUPPORTS_PRIMITIVE_SCENE_DATA"), bSupportsPrimitiveIdStream && bUseGPUScene);
@@ -195,13 +213,10 @@ void FNiagaraMeshVertexFactory::GetPSOPrecacheVertexFetchElements(EVertexInputSt
 {
 	Elements.Add(FVertexElement(0, 0, VET_Float3, 0, sizeof(float)*3u, false));
 
-#if NIAGARA_ENABLE_GPU_SCENE_MESHES
-	if (UseGPUScene(GMaxRHIShaderPlatform, GMaxRHIFeatureLevel)
-		&& !PlatformGPUSceneUsesUniformBufferView(GMaxRHIShaderPlatform))
+	if (IsGPUSceneEnabled(GMaxRHIShaderPlatform, GMaxRHIFeatureLevel) && !PlatformGPUSceneUsesUniformBufferView(GMaxRHIShaderPlatform))
 	{
 		Elements.Add(FVertexElement(1, 0, VET_UInt, 13, sizeof(uint32), true));
 	}
-#endif
 }
 
 void FNiagaraMeshVertexFactory::GetVertexElements(ERHIFeatureLevel::Type FeatureLevel, bool bSupportsManualVertexFetch, FStaticMeshDataType& Data, FVertexDeclarationElementList& Elements)
@@ -209,13 +224,10 @@ void FNiagaraMeshVertexFactory::GetVertexElements(ERHIFeatureLevel::Type Feature
 	FVertexStreamList InOutStreams;
 	GetVertexElements(FeatureLevel, bSupportsManualVertexFetch, Data, Elements, InOutStreams);
 
-#if NIAGARA_ENABLE_GPU_SCENE_MESHES
-	if (UseGPUScene(GMaxRHIShaderPlatform, GMaxRHIFeatureLevel)
-		&& !PlatformGPUSceneUsesUniformBufferView(GMaxRHIShaderPlatform))
+	if (IsGPUSceneEnabled(GMaxRHIShaderPlatform, GMaxRHIFeatureLevel) && !PlatformGPUSceneUsesUniformBufferView(GMaxRHIShaderPlatform))
 	{
 		Elements.Add(FVertexElement(InOutStreams.Num(), 0, VET_UInt, 13, sizeof(uint32), true));
 	}
-#endif
 }
 
 void FNiagaraMeshVertexFactory::SetData(FRHICommandListBase& RHICmdList, const FStaticMeshDataType& InData)
diff --git a/Engine/Plugins/FX/Niagara/Source/NiagaraVertexFactories/Public/NiagaraMeshVertexFactory.h b/Engine/Plugins/FX/Niagara/Source/NiagaraVertexFactories/Public/NiagaraMeshVertexFactory.h
--- a/Engine/Plugins/FX/Niagara/Source/NiagaraVertexFactories/Public/NiagaraMeshVertexFactory.h
+++ b/Engine/Plugins/FX/Niagara/Source/NiagaraVertexFactories/Public/NiagaraMeshVertexFactory.h
@@ -16,10 +16,6 @@ ParticleVertexFactory.h: Particle vertex factory definitions.
 #include "UniformBuffer.h"
 #include "VertexFactory.h"
 
-// Disable this define to test disabling the use of GPU Scene with Niagara mesh renderer
-// NOTE: Changing this will also require you to make a trivial change to the mesh factory shader, or it may use cached shaders
-#define NIAGARA_ENABLE_GPU_SCENE_MESHES 1
-
 class FMaterial;
 class FVertexBuffer;
 struct FDynamicReadBuffer;
@@ -211,6 +207,9 @@ public:
 	{
 		FNiagaraVertexFactoryBase::SetFeatureLevel(InFeatureLevel);
 	}
+
+	NIAGARAVERTEXFACTORIES_API static bool IsGPUSceneEnabled(const FStaticShaderPlatform Platform, const FStaticFeatureLevel FeatureLevel);
+
 protected:
 	static NIAGARAVERTEXFACTORIES_API void GetVertexElements(ERHIFeatureLevel::Type FeatureLevel, bool bSupportsManualVertexFetch, FStaticMeshDataType& Data, FVertexDeclarationElementList& Elements, FVertexStreamList& InOutStreams);
 
