- CIS fix for unacceptable word

#rnx

diff --git a/Engine/Plugins/FX/Niagara/Source/NiagaraVertexFactories/Private/NiagaraMeshVertexFactory.cpp b/Engine/Plugins/FX/Niagara/Source/NiagaraVertexFactories/Private/NiagaraMeshVertexFactory.cpp
--- a/Engine/Plugins/FX/Niagara/Source/NiagaraVertexFactories/Private/NiagaraMeshVertexFactory.cpp
+++ b/Engine/Plugins/FX/Niagara/Source/NiagaraVertexFactories/Private/NiagaraMeshVertexFactory.cpp
@@ -14,7 +14,7 @@
 #include "GlobalRenderResources.h"
 #include "DataDrivenShaderPlatformInfo.h"
 
-// Master control switch for GPU scene on all platforms
+// Main control for GPU scene on all platforms
 #define NIAGARA_ENABLE_GPU_SCENE_MESHES 1
 
 IMPLEMENT_GLOBAL_SHADER_PARAMETER_STRUCT(FNiagaraMeshUniformParameters, "NiagaraMeshVF");
