- LWE: Add support for debug HUD particle attribute display
- LWE: Add support to data set readback

#tests Use debug HUD to display particle attributes

diff --git a/Engine/Plugins/FX/Niagara/Source/Niagara/Internal/Stateless/NiagaraStatelessEmitterInstance.h b/Engine/Plugins/FX/Niagara/Source/Niagara/Internal/Stateless/NiagaraStatelessEmitterInstance.h
--- a/Engine/Plugins/FX/Niagara/Source/Niagara/Internal/Stateless/NiagaraStatelessEmitterInstance.h
+++ b/Engine/Plugins/FX/Niagara/Source/Niagara/Internal/Stateless/NiagaraStatelessEmitterInstance.h
@@ -105,6 +105,7 @@ private:
 	void SetExecutionStateInternal(ENiagaraExecutionState InExecutionState);
 
 public:
+	void EnqueueParticleDataBufferReadback(FNiagaraDataBuffer* DataBuffer, TFunction<void()> OnReadbackComplete) const;
 	void CaptureForDebugging(FNiagaraDataBuffer* DataBuffer) const;
 
 private:
diff --git a/Engine/Plugins/FX/Niagara/Source/Niagara/Private/NiagaraDataSetReadback.cpp b/Engine/Plugins/FX/Niagara/Source/Niagara/Private/NiagaraDataSetReadback.cpp
--- a/Engine/Plugins/FX/Niagara/Source/Niagara/Private/NiagaraDataSetReadback.cpp
+++ b/Engine/Plugins/FX/Niagara/Source/Niagara/Private/NiagaraDataSetReadback.cpp
@@ -9,6 +9,7 @@
 #include "NiagaraGpuComputeDispatchInterface.h"
 #include "NiagaraGPUInstanceCountManager.h"
 #include "NiagaraGpuReadbackManager.h"
+#include "Stateless/NiagaraStatelessEmitterInstance.h"
 
 #include "Async/Async.h"
 
@@ -42,7 +43,7 @@ void FNiagaraDataSetReadback::EnqueueReadback(FNiagaraEmitterInstance* EmitterIn
 			}
 		);
 	}
-	else
+	else if (FNiagaraEmitterInstanceImpl* StatefulEmitterInstance = EmitterInstance->AsStateful())
 	{
 		const FNiagaraDataSet& SourceDataSet = EmitterInstance->GetParticleData();
 		if (FNiagaraDataBuffer* SourceDataBuffer = SourceDataSet.GetCurrentData())
@@ -53,11 +54,23 @@ void FNiagaraDataSetReadback::EnqueueReadback(FNiagaraEmitterInstance* EmitterIn
 		{
 			DataSet.CopyFromGPUReadback(nullptr, nullptr, nullptr, 0, 0, 0, 0, 0);
 		}
-		//-TODO:Stateless
-		if (FNiagaraEmitterInstanceImpl* StatefulEmitterInstance = EmitterInstance->AsStateful())
-		{
-			ParameterStore = StatefulEmitterInstance->GetUpdateExecutionContext().Parameters;
-		}
+		ParameterStore = StatefulEmitterInstance->GetUpdateExecutionContext().Parameters;
+	}
+	else if (FNiagaraStatelessEmitterInstance* StatelessEmitterInstance = EmitterInstance->AsStateless())
+	{
+		++PendingReadbacks;
+
+		DataSet.BeginSimulate();
+		DataSet.EndSimulate();
+
+		StatelessEmitterInstance->EnqueueParticleDataBufferReadback(
+			DataSet.GetCurrentData(),
+			[RT_DataSetReadback=AsShared()]()
+			{
+				RT_DataSetReadback->ReadbackCompleteInternal();
+			}
+		);
+		ParameterStore = FNiagaraParameterStore();
 	}
 }
 
diff --git a/Engine/Plugins/FX/Niagara/Source/Niagara/Private/NiagaraDebugHud.cpp b/Engine/Plugins/FX/Niagara/Source/Niagara/Private/NiagaraDebugHud.cpp
--- a/Engine/Plugins/FX/Niagara/Source/Niagara/Private/NiagaraDebugHud.cpp
+++ b/Engine/Plugins/FX/Niagara/Source/Niagara/Private/NiagaraDebugHud.cpp
@@ -21,6 +21,8 @@
 #include "NiagaraSystemGpuComputeProxy.h"
 #include "NiagaraSystemInstanceController.h"
 #include "NiagaraWorldManager.h"
+#include "Stateless/NiagaraStatelessEmitter.h"
+#include "Stateless/NiagaraStatelessEmitterData.h"
 
 #include "Components/LineBatchComponent.h"
 #include "Debug/DebugDrawService.h"
@@ -509,13 +511,14 @@ namespace NiagaraDebugLocal
 
 			if (Settings.bShowParticleVariables && Settings.ParticlesVariables.Num() > 0)
 			{
-				const TArray<TSharedRef<const FNiagaraEmitterCompiledData>>& AllEmittersCompiledData = NiagaraSystem->GetEmitterCompiledData();
+				const TConstArrayView<FNiagaraEmitterHandle> EmitterHandles = NiagaraSystem->GetEmitterHandles();
+				const TArray<TSharedRef<const FNiagaraEmitterCompiledData>>& AllStatefulEmittersCompiledData = NiagaraSystem->GetEmitterCompiledData();
 
-				CachedVariables->ParticleVariables.AddDefaulted(AllEmittersCompiledData.Num());
-				CachedVariables->ParticlePositionAccessors.AddDefaulted(AllEmittersCompiledData.Num());
-				for (int32 iEmitter = 0; iEmitter < AllEmittersCompiledData.Num(); ++iEmitter)
+				CachedVariables->ParticleVariables.AddDefaulted(EmitterHandles.Num());
+				CachedVariables->ParticlePositionAccessors.AddDefaulted(EmitterHandles.Num());
+				for (int32 iEmitter=0; iEmitter < EmitterHandles.Num(); ++iEmitter)
 				{
-					const FNiagaraEmitterHandle& EmitterHandle = NiagaraSystem->GetEmitterHandle(iEmitter);
+					const FNiagaraEmitterHandle& EmitterHandle = EmitterHandles[iEmitter];
 					if (!EmitterHandle.IsValid() || !EmitterHandle.GetIsEnabled())
 					{
 						continue;
@@ -526,18 +529,33 @@ namespace NiagaraDebugLocal
 						continue;
 					}
 
-					const FNiagaraDataSetCompiledData& EmitterCompiledData = AllEmittersCompiledData[iEmitter]->DataSetCompiledData;
+					const FNiagaraDataSetCompiledData* EmitterCompiledData = nullptr;
+					if (EmitterHandle.GetEmitterMode() == ENiagaraEmitterMode::Stateless)
+					{
+						UNiagaraStatelessEmitter* StatelessEmitter = EmitterHandle.GetStatelessEmitter();
+						FNiagaraStatelessEmitterDataPtr EmitterData = StatelessEmitter ? StatelessEmitter->GetEmitterData() : nullptr;
+						EmitterCompiledData = EmitterData ? EmitterData->ParticleDataSetCompiledData.Get() : nullptr;
+					}
+					else
+					{
+						EmitterCompiledData = AllStatefulEmittersCompiledData.IsValidIndex(iEmitter) ? &AllStatefulEmittersCompiledData[iEmitter]->DataSetCompiledData : nullptr;
+					}
+
+					if (EmitterCompiledData == nullptr)
+					{
+						continue;
+					}
 
 					FindVariablesByWildcard(
-						EmitterCompiledData.Variables,
+						EmitterCompiledData->Variables,
 						Settings.ParticlesVariables,
-						[&](const FNiagaraVariable& Variable) { CachedVariables->ParticleVariables[iEmitter].AddDefaulted_GetRef().Init(EmitterCompiledData, Variable.GetName()); }
+						[&](const FNiagaraVariable& Variable) { CachedVariables->ParticleVariables[iEmitter].AddDefaulted_GetRef().Init(*EmitterCompiledData, Variable.GetName()); }
 					);
 
 					if (CachedVariables->ParticleVariables[iEmitter].Num() > 0)
 					{
 						static const FName PositionName(TEXT("Position"));
-						CachedVariables->ParticlePositionAccessors[iEmitter].Init(EmitterCompiledData, PositionName);
+						CachedVariables->ParticlePositionAccessors[iEmitter].Init(*EmitterCompiledData, PositionName);
 					}
 				}
 			}
@@ -1321,14 +1339,13 @@ const FNiagaraDataSet* FNiagaraDebugHud::GetParticleDataSet(FNiagaraSystemInstan
 	using namespace NiagaraDebugLocal;
 
 	// For GPU context we need to readback and cache the data
-	if (EmitterInstance->GetGPUContext())
+	if (EmitterInstance->GetSimTarget() == ENiagaraSimTarget::GPUComputeSim)
 	{
 		if (!Settings.bEnableGpuParticleReadback)
 		{
 			return nullptr;
 		}
 
-		FNiagaraComputeExecutionContext* GPUExecContext = EmitterInstance->GetGPUContext();
 		FGpuEmitterCache* GpuCachedData = GpuEmitterData.Find(SystemInstance->GetId());
 		if ( GpuCachedData == nullptr )
 		{
@@ -3127,7 +3144,7 @@ void FNiagaraDebugHud::DrawComponents(FNiagaraWorldManager* WorldManager, UCanva
 						if (Settings.SystemEmitterVerbosity == ENiagaraDebugHudVerbosity::Verbose)
 						{
 							const FString EmitterName = EmitterInstance->GetEmitterHandle().GetUniqueInstanceName();
-							if ( EmitterInstance->GetGPUContext() )
+							if ( EmitterInstance->GetSimTarget() == ENiagaraSimTarget::GPUComputeSim )
 							{
 								StringBuilder.Appendf(TEXT("Emitter(GPU) %s - State %s - Particles %d\n"), *EmitterName, *ExecutionStateEnum->GetNameStringByIndex((int32)EmitterInstance->GetExecutionState()), EmitterInstance->GetNumParticles());
 							}
diff --git a/Engine/Plugins/FX/Niagara/Source/Niagara/Private/Stateless/NiagaraStatelessEmitterInstance.cpp b/Engine/Plugins/FX/Niagara/Source/Niagara/Private/Stateless/NiagaraStatelessEmitterInstance.cpp
--- a/Engine/Plugins/FX/Niagara/Source/Niagara/Private/Stateless/NiagaraStatelessEmitterInstance.cpp
+++ b/Engine/Plugins/FX/Niagara/Source/Niagara/Private/Stateless/NiagaraStatelessEmitterInstance.cpp
@@ -964,40 +964,53 @@ void FNiagaraStatelessEmitterInstance::InitEmitterData()
 	}
 }
 
-void FNiagaraStatelessEmitterInstance::CaptureForDebugging(FNiagaraDataBuffer* DataBuffer) const
+void FNiagaraStatelessEmitterInstance::EnqueueParticleDataBufferReadback(FNiagaraDataBuffer* DataBuffer, TFunction<void()> OnReadbackComplete) const
 {
+	check(IsInGameThread());
 	check(DataBuffer);
 
 	// Set instances to zero to handle any early outs
 	DataBuffer->SetNumInstances(0);
 	if (!bCanEverExecute || IsComplete())
 	{
+		OnReadbackComplete();
 		return;
 	}
 
-	// Capture the data based on Sim Target
+	// If we target the CPU we can copy the data
 	if (SimTarget == ENiagaraSimTarget::CPUSim)
 	{
 		if (FNiagaraDataBuffer* CurrentDataBuffer = ParticleDataSet->GetCurrentData())
 		{
 			CurrentDataBuffer->CopyTo(*DataBuffer, 0, 0, INDEX_NONE);
 		}
+		OnReadbackComplete();
 	}
+	// If we target the GPU then we need to generate the data
 	else
 	{
 		ENQUEUE_RENDER_COMMAND(CaptureStatelessForDebugging)(
-			[RenderThreadData=RenderThreadDataPtr.Get(), DataBuffer](FRHICommandListImmediate& RHICmdList)
+			[RenderThreadData=RenderThreadDataPtr.Get(), DataBuffer, OnReadbackComplete_RT=MoveTemp(OnReadbackComplete)](FRHICommandListImmediate& RHICmdList)
 			{
-				// No compute manager then we can not do anything
-				if ( !RenderThreadData->ComputeManager )
+				if ( RenderThreadData->ComputeManager )
 				{
-					return;
+					RenderThreadData->ComputeManager->GenerateDataBufferForDebugging(RHICmdList, DataBuffer, RenderThreadData);
 				}
-
-				RenderThreadData->ComputeManager->GenerateDataBufferForDebugging(RHICmdList, DataBuffer, RenderThreadData);
+				OnReadbackComplete_RT();
 			}
 		);
+	}
+}
 
+void FNiagaraStatelessEmitterInstance::CaptureForDebugging(FNiagaraDataBuffer* DataBuffer) const
+{
+	bool bIsComplete = false;
+
+	EnqueueParticleDataBufferReadback(DataBuffer, [&bIsComplete]() { bIsComplete = true; } );
+
+	if (bIsComplete == false)
+	{
 		FlushRenderingCommands();
+		check(bIsComplete == true);
 	}
 }
