- Add console commands (fx.Niagara.Component.ReinitializeActive & fx.Niagara.Component.DeactivateActive) to allow a quick way to kill / restart Niagara components
- Cleaned up various console variables while I was here

#tests drop systems in a level, use the console command

diff --git a/Engine/Plugins/FX/Niagara/Source/Niagara/Private/NiagaraComponent.cpp b/Engine/Plugins/FX/Niagara/Source/Niagara/Private/NiagaraComponent.cpp
--- a/Engine/Plugins/FX/Niagara/Source/Niagara/Private/NiagaraComponent.cpp
+++ b/Engine/Plugins/FX/Niagara/Source/Niagara/Private/NiagaraComponent.cpp
@@ -54,164 +54,182 @@ DEFINE_LOG_CATEGORY(LogNiagara);
 
 DEFINE_RENDER_COMMAND_PIPE(NiagaraDynamicData, ERenderCommandPipeFlags::None);
 
-static bool GNiagaraParallelGDME = true;
-static FAutoConsoleVariableRef CVarNiagaraParallelGDME(
-	TEXT("fx.Niagara.ParallelGDME"),
-	GNiagaraParallelGDME,
-	TEXT("Allow Niagara to run parallel in GDME."),
-	ECVF_Default
-);
-
-static int GNiagaraSoloTickEarly = 1;
-static FAutoConsoleVariableRef CVarNiagaraSoloTickEarly(
-	TEXT("fx.Niagara.Solo.TickEarly"),
-	GNiagaraSoloTickEarly,
-	TEXT("When enabled will tick kin the first available tick group."),
-	ECVF_Default
-);
-
-static int GNiagaraSoloAllowAsyncWorkToEndOfFrame = 1;
-static FAutoConsoleVariableRef CVarNiagaraSoloAllowAsyncWorkToEndOfFrame(
-	TEXT("fx.Niagara.Solo.AllowAsyncWorkToEndOfFrame"),
-	GNiagaraSoloAllowAsyncWorkToEndOfFrame,
-	TEXT("Allow async work to continue until the end of the frame for solo Niagara instances, if false it will complete within the tick group it started in."),
-	ECVF_Default
-);
-
-static bool GNiagaraAllowDesiredAgeAsync = true;
-static FAutoConsoleVariableRef CVarNiagaraAllowDesiredAgeAsync(
-	TEXT("fx.Niagara.AllowDesiredAgeAsync"),
-	GNiagaraAllowDesiredAgeAsync,
-	TEXT("When enabled Desired Age / Desired Age No Seek tick modes will be allows to run async if we single tick in a frame, multi tick will still run on the game thread."),
-	ECVF_Default
-);
-
-static bool GbNiagaraComponentsEnabled = true;
-static FAutoConsoleVariableRef CVarNiagaraComponentsEnabled(
-	TEXT("fx.NiagaraComponentsEnabled"),
-	GbNiagaraComponentsEnabled,
-	TEXT("If false Niagara particle components will not be activated."),
-	ECVF_Default
-);
-
-static int32 GNiagaraComponentWarnNullAsset = 0;
-static FAutoConsoleVariableRef CVarNiagaraComponentWarnNullAsset(
-	TEXT("fx.Niagara.ComponentWarnNullAsset"),
-	GNiagaraComponentWarnNullAsset,
-	TEXT("When enabled we will warn if a NiagaraComponent is activate with a null asset.  This is sometimes useful for tracking down components that can be removed."),
-	ECVF_Default
-);
-
-static int32 GNiagaraComponentWarnAsleepCullReaction = 1;
-static FAutoConsoleVariableRef CVarNiagaraComponentWarnAsleepCullReaction(
-	TEXT("fx.Niagara.ComponentWarnAsleepCullReaction"),
-	GNiagaraComponentWarnAsleepCullReaction,
-	TEXT("When enabled we will warn if a NiagaraComponent completes naturally but has Asleep mode set for cullreaction."),
-	ECVF_Default
-);
-
-static int32 GNiagaraUseFastSetUserParametersToDefaultValues = 1;
-static FAutoConsoleVariableRef CVarNiagaraUseFastSetUserParametersToDefaultValues(
-	TEXT("fx.Niagara.UseFastSetUserParametersToDefaultValues"),
-	GNiagaraUseFastSetUserParametersToDefaultValues,
-	TEXT("When a component is activated we will check the surpession list."),
-	ECVF_Default
-);
-
-static int GNiagaraForceWaitForCompilationOnActivate = 0;
-static FAutoConsoleVariableRef CVarNiagaraForceWaitForCompilationOnActivate(
-	TEXT("fx.Niagara.ForceWaitForCompilationOnActivate"),
-	GNiagaraForceWaitForCompilationOnActivate,
-	TEXT("When a component is activated it will stall waiting for any pending shader compilation."),
-	ECVF_Default
-);
-
-static bool GNiagaraComponentBoostPSOPrecacheHighest = false;
-static FAutoConsoleVariableRef CVarNiagaraComponentBoostPSOPrecacheHighest(
-	TEXT("r.PSOPrecache.NiagaraComponentPSOPrecachePriority"),
-	GNiagaraComponentBoostPSOPrecacheHighest,
-	TEXT("Niagara component PSO precache priority level.\n")
-	TEXT(" 0. Niagara component's PSO precache requests are set to high priority (default)\n")
-	TEXT(" 1. Niagara component's PSO precache requests are set to highest priority"),
-	ECVF_Default
-);
-
-FAutoConsoleCommandWithWorldAndArgs DumpNiagaraComponentsCommand(
-	TEXT("fx.Niagara.DumpComponents"),
-	TEXT("Dump Information about all Niagara Components"),
-	FConsoleCommandWithWorldAndArgsDelegate::CreateLambda(
-		[](const TArray<FString>& Args, UWorld* World)
-		{
-			const bool bFullDump = Args.Contains(TEXT("full"));
-			const bool bAllWorlds = (World == nullptr) || Args.Contains(TEXT("allworlds"));
-			const FString SystemFilter = 
-				[&]() -> FString
-				{
-					static const FString FilterPrefix(TEXT("filter="));
-					for ( const FString& Arg : Args )
+namespace NiagaraComponentPrivate
+{
+	static bool GParallelGDMEEnabled = true;
+	static FAutoConsoleVariableRef CVarNiagaraParallelGDME(
+		TEXT("fx.Niagara.ParallelGDME"),
+		GParallelGDMEEnabled,
+		TEXT("Allow Niagara to run parallel in GDME."),
+		ECVF_Default
+	);
+
+	static bool GbComponentsEnabled = true;
+	static FAutoConsoleVariableRef CVarNiagaraComponentsEnabled(
+		TEXT("fx.NiagaraComponentsEnabled"),
+		GbComponentsEnabled,
+		TEXT("If false Niagara particle components will not be activated."),
+		ECVF_Default
+	);
+
+	static int32 GWarnNullAsset = 0;
+	static FAutoConsoleVariableRef CVarNiagaraComponentWarnNullAsset(
+		TEXT("fx.Niagara.ComponentWarnNullAsset"),
+		GWarnNullAsset,
+		TEXT("When enabled we will warn if a NiagaraComponent is activate with a null asset.  This is sometimes useful for tracking down components that can be removed."),
+		ECVF_Default
+	);
+
+	static int32 GWarnAsleepCullReaction = 1;
+	static FAutoConsoleVariableRef CVarNiagaraComponentWarnAsleepCullReaction(
+		TEXT("fx.Niagara.ComponentWarnAsleepCullReaction"),
+		GWarnAsleepCullReaction,
+		TEXT("When enabled we will warn if a NiagaraComponent completes naturally but has Asleep mode set for cullreaction."),
+		ECVF_Default
+	);
+
+	static int GForceWaitForCompilationOnActivate = 0;
+	static FAutoConsoleVariableRef CVarNiagaraForceWaitForCompilationOnActivate(
+		TEXT("fx.Niagara.ForceWaitForCompilationOnActivate"),
+		GForceWaitForCompilationOnActivate,
+		TEXT("When a component is activated it will stall waiting for any pending shader compilation."),
+		ECVF_Default
+	);
+
+	static bool GBoostPSOPrecacheHighest = false;
+	static FAutoConsoleVariableRef CVarNiagaraComponentBoostPSOPrecacheHighest(
+		TEXT("r.PSOPrecache.NiagaraComponentPSOPrecachePriority"),
+		GBoostPSOPrecacheHighest,
+		TEXT("Niagara component PSO precache priority level.\n")
+		TEXT(" 0. Niagara component's PSO precache requests are set to high priority (default)\n")
+		TEXT(" 1. Niagara component's PSO precache requests are set to highest priority"),
+		ECVF_Default
+	);
+
+	static int GCanPreventScalabilityCullingOnPlayerFX = 1;
+	static FAutoConsoleVariableRef CVarNiagaraCanPreventScalabilityCullingOnPlayerFX(
+		TEXT("fx.Niagara.Scalability.CanPreventCullingOnPlayerFX"),
+		GCanPreventScalabilityCullingOnPlayerFX,
+		TEXT("When enabled Niagara can optionally prevent scalability culling on FX linked to the player."),
+		FConsoleVariableDelegate::CreateLambda([](IConsoleVariable* CVar)
+			{
+				FNiagaraWorldManager::OnRefreshOwnerAllowsScalability();
+			}),
+		ECVF_Default
+	);
+
+	FAutoConsoleCommandWithWorldAndArgs DumpNiagaraComponentsCommand(
+		TEXT("fx.Niagara.DumpComponents"),
+		TEXT("Dump Information about all Niagara Components"),
+		FConsoleCommandWithWorldAndArgsDelegate::CreateLambda(
+			[](const TArray<FString>& Args, UWorld* World)
+			{
+				const bool bFullDump = Args.Contains(TEXT("full"));
+				const bool bAllWorlds = (World == nullptr) || Args.Contains(TEXT("allworlds"));
+				const FString SystemFilter = 
+					[&]() -> FString
 					{
-						if ( Arg.StartsWith(FilterPrefix) )
+						static const FString FilterPrefix(TEXT("filter="));
+						for ( const FString& Arg : Args )
 						{
-							return Arg.Mid(FilterPrefix.Len());
+							if ( Arg.StartsWith(FilterPrefix) )
+							{
+								return Arg.Mid(FilterPrefix.Len());
+							}
 						}
-					}
-					return FString();
-				}();
-
-			UE_LOG(LogNiagara, Log, TEXT("=========================== Begin Niagara Dump ==========================="));
-			for (TObjectIterator<UNiagaraComponent> It; It; ++It)
-			{
-				UNiagaraComponent* Component = *It;
-				if ( !IsValid(Component) )
-				{
-					continue;
-				}
+						return FString();
+					}();
 
-				// Filter by world
-				if (!bAllWorlds && (Component->GetWorld() != World) )
+				UE_LOG(LogNiagara, Log, TEXT("=========================== Begin Niagara Dump ==========================="));
+				for (TObjectIterator<UNiagaraComponent> It; It; ++It)
 				{
-					continue;
-				}
+					UNiagaraComponent* Component = *It;
+					if ( !IsValid(Component) )
+					{
+						continue;
+					}
 
-				// Filter by asset name, we allow null assets to continue so we can identify components who are attached but have no valid asset associated with them
-				UNiagaraSystem* NiagaraSystem = Component->GetAsset();
-				if ( !SystemFilter.IsEmpty() )
-				{
-					if (!NiagaraSystem || !NiagaraSystem->GetName().Contains(SystemFilter) )
+					// Filter by world
+					if (!bAllWorlds && (Component->GetWorld() != World) )
 					{
 						continue;
 					}
-				}
 
-				if (Component->PoolingMethod == ENCPoolMethod::FreeInPool)
-				{
-					UE_LOG(LogNiagara, Log, TEXT("Component '%s' Asset '%s' Free In Component Pool"), *GetNameSafe(Component), *GetNameSafe(NiagaraSystem));
-				}
-				else
-				{
-					UE_LOG(LogNiagara, Log, TEXT("Component '%s' Asset '%s' Actor '%s' is %s"), *GetNameSafe(Component), *GetNameSafe(NiagaraSystem), *GetNameSafe(Component->GetTypedOuter<AActor>()), Component->IsActive() ? TEXT("Active") : TEXT("Inactive"));
-					if (FNiagaraSystemInstanceControllerPtr SystemInstanceController = Component->GetSystemInstanceController())
+					// Filter by asset name, we allow null assets to continue so we can identify components who are attached but have no valid asset associated with them
+					UNiagaraSystem* NiagaraSystem = Component->GetAsset();
+					if ( !SystemFilter.IsEmpty() )
 					{
-						SystemInstanceController->DebugDump(bFullDump);
+						if (!NiagaraSystem || !NiagaraSystem->GetName().Contains(SystemFilter) )
+						{
+							continue;
+						}
+					}
+
+					if (Component->PoolingMethod == ENCPoolMethod::FreeInPool)
+					{
+						UE_LOG(LogNiagara, Log, TEXT("Component '%s' Asset '%s' Free In Component Pool"), *GetNameSafe(Component), *GetNameSafe(NiagaraSystem));
+					}
+					else
+					{
+						UE_LOG(LogNiagara, Log, TEXT("Component '%s' Asset '%s' Actor '%s' is %s"), *GetNameSafe(Component), *GetNameSafe(NiagaraSystem), *GetNameSafe(Component->GetTypedOuter<AActor>()), Component->IsActive() ? TEXT("Active") : TEXT("Inactive"));
+						if (FNiagaraSystemInstanceControllerPtr SystemInstanceController = Component->GetSystemInstanceController())
+						{
+							SystemInstanceController->DebugDump(bFullDump);
+						}
 					}
 				}
+				UE_LOG(LogNiagara, Log, TEXT("=========================== End Niagara Dump ==========================="));
 			}
-			UE_LOG(LogNiagara, Log, TEXT("=========================== End Niagara Dump ==========================="));
-		}
-	)
-);
+		)
+	);
 
-static int GNiagaraCanPreventScalabilityCullingOnPlayerFX = 1;
-static FAutoConsoleVariableRef CVarNiagaraCanPreventScalabilityCullingOnPlayerFX(
-	TEXT("fx.Niagara.Scalability.CanPreventCullingOnPlayerFX"),
-	GNiagaraCanPreventScalabilityCullingOnPlayerFX,
-	TEXT("When enabled Niagara can optionally prevent scalability culling on FX linked to the player."),
-	FConsoleVariableDelegate::CreateLambda([](IConsoleVariable* CVar)
+	static void ForEachComponentInWorld(const TArray<FString>& InArgs, UWorld* World, TFunction<void(UNiagaraComponent*)> Callback)
 	{
-		FNiagaraWorldManager::OnRefreshOwnerAllowsScalability();
-	}),
-	ECVF_Default
-);
+		const FString Command = FString::Join(InArgs, TEXT(" "));
+
+		FString SystemFilter;
+		FParse::Value(*Command, TEXT("filter="), SystemFilter);
+
+		for (TObjectIterator<UNiagaraComponent> It; It; ++It)
+		{
+			UNiagaraComponent* Component = *It;
+			if (!IsValid(Component) || Component->GetWorld() != World || !Component->IsActive() || !Component->GetAsset())
+			{
+				continue;
+			}
+
+			UNiagaraSystem* NiagaraSystem = Component->GetAsset();
+			if (!SystemFilter.IsEmpty() && !NiagaraSystem->GetName().MatchesWildcard(SystemFilter))
+			{
+				continue;
+			}
+
+			Callback(Component);
+		}
+	}
+
+	FAutoConsoleCommandWithWorldAndArgs ReinitializeActiveCommand(
+		TEXT("fx.Niagara.Component.ReinitializeActive"),
+		TEXT("Iterates all activate components and reinitializes them."),
+		FConsoleCommandWithWorldAndArgsDelegate::CreateLambda(
+			[](const TArray<FString>& InArgs, UWorld* World)
+			{
+				ForEachComponentInWorld(InArgs, World, [](UNiagaraComponent* Component) { Component->ReinitializeSystem(); });
+			}
+		)
+	);
+
+	FAutoConsoleCommandWithWorldAndArgs DeactivateActiveCommand(
+		TEXT("fx.Niagara.Component.DeactivateActive"),
+		TEXT("Iterates all activate components and deactivates them."),
+		FConsoleCommandWithWorldAndArgsDelegate::CreateLambda(
+			[](const TArray<FString>& InArgs, UWorld* World)
+			{
+				ForEachComponentInWorld(InArgs, World, [](UNiagaraComponent* Component) { Component->DeactivateImmediate(); });
+			}
+		)
+	);
+}
 
 bool FitsIntoFloat(FVector InValue)
 {
@@ -247,7 +265,7 @@ FNiagaraSceneProxy::FNiagaraSceneProxy(const FNiagaraSceneProxyDesc& Desc)
 	bHasDeformableMesh = false;
 
 	// Optionally enable GDME for Niagara, any caveats are covered below
-	bSupportsParallelGDME = GNiagaraParallelGDME;
+	bSupportsParallelGDME = NiagaraComponentPrivate::GParallelGDMEEnabled;
 
 	// Both GDME and GDRTI call methods such as InitializeVertexBuffersResources(...)
 	// which means GDME and GDRTI can't be called in parallel.
@@ -644,8 +662,8 @@ UNiagaraComponent::UNiagaraComponent(const FObjectInitializer& ObjectInitializer
 	OverrideParameters.SetOwner(this);
 
 	PrimaryComponentTick.bCanEverTick = true;
-	PrimaryComponentTick.TickGroup = GNiagaraSoloTickEarly ? TG_PrePhysics : TG_DuringPhysics;
-	PrimaryComponentTick.EndTickGroup = GNiagaraSoloAllowAsyncWorkToEndOfFrame ? TG_LastDemotable : ETickingGroup(PrimaryComponentTick.TickGroup);
+	PrimaryComponentTick.TickGroup = TG_PrePhysics;
+	PrimaryComponentTick.EndTickGroup = TG_LastDemotable;
 	PrimaryComponentTick.bStartWithTickEnabled = false;
 	PrimaryComponentTick.bAllowTickOnDedicatedServer = false;
 	PrimaryComponentTick.SetTickFunctionEnable(false);
@@ -973,7 +991,7 @@ void UNiagaraComponent::TickComponent(float DeltaSeconds, enum ELevelTick TickTy
 
 				if (SystemInstanceController->GetSoloSystemSimulation().IsValid())
 				{
-					FGraphEventRef CompletionEvent = (GNiagaraAllowDesiredAgeAsync && ThisTickFunction && ThisTickFunction->IsCompletionHandleValid()) ? ThisTickFunction->GetCompletionHandle() : nullptr;
+					FGraphEventRef CompletionEvent = (ThisTickFunction && ThisTickFunction->IsCompletionHandleValid()) ? ThisTickFunction->GetCompletionHandle() : nullptr;
 
 					const bool bFixedTickDelta = Asset->HasFixedTickDelta();
 					const float TickDelta = bFixedTickDelta ? Asset->GetFixedTickDeltaTime() : SeekDelta;
@@ -1003,7 +1021,7 @@ void UNiagaraComponent::TickComponent(float DeltaSeconds, enum ELevelTick TickTy
 		}
 		else if (AgeUpdateMode == ENiagaraAgeUpdateMode::DesiredAgeNoSeek)
 		{
-			FGraphEventRef CompletionEvent = (GNiagaraAllowDesiredAgeAsync && ThisTickFunction && ThisTickFunction->IsCompletionHandleValid()) ? ThisTickFunction->GetCompletionHandle() : nullptr;
+			FGraphEventRef CompletionEvent = (ThisTickFunction && ThisTickFunction->IsCompletionHandleValid()) ? ThisTickFunction->GetCompletionHandle() : nullptr;
 
 			int32 MaxForwardFrames = 5; // HACK - for some reason sequencer jumps forwards by multiple frames on pause, so this is being added to allow for FX to stay alive when being controlled by sequencer in the editor.  This should be lowered once that issue is fixed.
 			float AgeDiff = DesiredAge - LastHandledDesiredAge;
@@ -1262,7 +1280,7 @@ void UNiagaraComponent::ActivateInternal(bool bReset /* = false */, bool bIsScal
 		CurrLocalBounds.Init();
 	}
 
-	if (!GbNiagaraComponentsEnabled)
+	if (!NiagaraComponentPrivate::GbComponentsEnabled)
 	{
 		UnregisterWithScalabilityManager();
 		OnSystemComplete(true);
@@ -1275,7 +1293,7 @@ void UNiagaraComponent::ActivateInternal(bool bReset /* = false */, bool bIsScal
 	if (Asset == nullptr)
 	{
 		DestroyInstance();
-		if (GNiagaraComponentWarnNullAsset && !HasAnyFlags(RF_DefaultSubObject | RF_ArchetypeObject | RF_ClassDefaultObject))
+		if (NiagaraComponentPrivate::GWarnNullAsset && !HasAnyFlags(RF_DefaultSubObject | RF_ArchetypeObject | RF_ClassDefaultObject))
 		{
 			UE_LOG(LogNiagara, Warning, TEXT("Failed to activate Niagara Component due to missing or invalid asset! (%s)"), *GetFullName());
 		}
@@ -1316,7 +1334,7 @@ void UNiagaraComponent::ActivateInternal(bool bReset /* = false */, bool bIsScal
 	// In case we're not yet ready to run due to compilation requests, go ahead and keep polling there..
 	if (Asset->HasOutstandingCompilationRequests(true))
 	{
-		if (bWaitForCompilationOnActivate || GNiagaraForceWaitForCompilationOnActivate || GIsAutomationTesting)
+		if (bWaitForCompilationOnActivate || NiagaraComponentPrivate::GForceWaitForCompilationOnActivate || GIsAutomationTesting)
 		{
 			Asset->WaitForCompilationComplete(true);
 		}
@@ -1489,8 +1507,8 @@ void UNiagaraComponent::ActivateInternal(bool bReset /* = false */, bool bIsScal
 		else if (SystemInstanceController->IsSolo())
 		{
 			const ETickingGroup SoloTickGroup = SystemInstanceController->CalculateTickGroup();
-			PrimaryComponentTick.TickGroup = FMath::Max(GNiagaraSoloTickEarly ? TG_PrePhysics : TG_DuringPhysics, SoloTickGroup);
-			PrimaryComponentTick.EndTickGroup = GNiagaraSoloAllowAsyncWorkToEndOfFrame ? TG_LastDemotable : ETickingGroup(PrimaryComponentTick.TickGroup);
+			PrimaryComponentTick.TickGroup = FMath::Max(TG_PrePhysics, SoloTickGroup);
+			PrimaryComponentTick.EndTickGroup = TG_LastDemotable;
 
 			// Solo instances require the component tick to be enabled
 			ScopedComponentTickEnabled.bTickEnabled = true;
@@ -1837,7 +1855,7 @@ void UNiagaraComponent::OnSystemComplete(bool bExternalCompletion)
 					//Only trigger warning if we're not being deactivated/completed from the outside and this is a natural completion by the system itself.
 					if (EffectType->CullReaction == ENiagaraCullReaction::DeactivateImmediateResume || EffectType->CullReaction == ENiagaraCullReaction::DeactivateResume)
 					{
-						if (GNiagaraComponentWarnAsleepCullReaction == 1 && GetAsset())
+						if (NiagaraComponentPrivate::GWarnAsleepCullReaction == 1 && GetAsset())
 						{
 							//If we're completing naturally, i.e. we're a burst/non-looping system then we shouldn't be using a mode reactivates the effect.
 							UE_LOG(LogNiagara, Warning, TEXT("Niagara Effect has completed naturally but has an effect type with the \"Asleep\" cull reaction. If an effect like this is culled before it can complete then it could leak into the scalability manager and be reactivated incorrectly. Please verify this is using the correct EffectType.\nComponent:%s\nSystem:%s")
@@ -2412,7 +2430,7 @@ FPrimitiveSceneProxy* UNiagaraComponent::CreateSceneProxy()
 			PrecacheAssetPSOs(Asset);
 		}
 
-		EPSOPrecachePriority PSOPrecachePriority = GNiagaraComponentBoostPSOPrecacheHighest ? EPSOPrecachePriority::Highest : EPSOPrecachePriority::High;
+		EPSOPrecachePriority PSOPrecachePriority = NiagaraComponentPrivate::GBoostPSOPrecacheHighest ? EPSOPrecachePriority::Highest : EPSOPrecachePriority::High;
 		if (CheckPSOPrecachingAndBoostPriority(PSOPrecachePriority) && GetPSOPrecacheProxyCreationStrategy() != EPSOPrecacheProxyCreationStrategy::AlwaysCreate)
 		{
 			UE_LOG(LogNiagara, Verbose, TEXT("Skipping CreateSceneProxy for UNiagaraComponent %s (UNiagaraSystem PSOs are still compiling)"), *GetFullName());
@@ -2491,7 +2509,7 @@ bool UNiagaraComponent::ResolveOwnerAllowsScalability(bool bRegister)
 {
 	bOwnerAllowsScalabiltiy = true;
 	
-	if(GNiagaraCanPreventScalabilityCullingOnPlayerFX && Asset && Asset->AllowScalabilityForLocalPlayerFX() == false)
+	if(NiagaraComponentPrivate::GCanPreventScalabilityCullingOnPlayerFX && Asset && Asset->AllowScalabilityForLocalPlayerFX() == false)
 	{
 		bOwnerAllowsScalabiltiy = IsLocalPlayerEffect() == false;
 	}	
@@ -3299,56 +3317,48 @@ void UNiagaraComponent::SetUserParametersToDefaultValues()
 		return;
 	}
 
-	if (GNiagaraUseFastSetUserParametersToDefaultValues)
-	{
-		const FNiagaraUserRedirectionParameterStore& SourceUserParameterStore = Asset->GetExposedParameters();
-		TArrayView<const FNiagaraVariableWithOffset> DestParameters = OverrideParameters.ReadParameterVariables();
+	const FNiagaraUserRedirectionParameterStore& SourceUserParameterStore = Asset->GetExposedParameters();
+	TArrayView<const FNiagaraVariableWithOffset> DestParameters = OverrideParameters.ReadParameterVariables();
 
-		TArray<FNiagaraVariableBase, TInlineAllocator<8>> ParametersToRemove;
-		bool bInterfacesChanged = false;
+	TArray<FNiagaraVariableBase, TInlineAllocator<8>> ParametersToRemove;
+	bool bInterfacesChanged = false;
 
-		for (int32 i=0; i < DestParameters.Num(); ++i)
+	for (int32 i=0; i < DestParameters.Num(); ++i)
+	{
+		const FNiagaraVariableWithOffset& DestParameter = DestParameters[i];
+		const int32 DestIndex = DestParameter.Offset;
+		const int32 SourceIndex = SourceUserParameterStore.IndexOf(DestParameter);
+		if (SourceIndex != INDEX_NONE)
 		{
-			const FNiagaraVariableWithOffset& DestParameter = DestParameters[i];
-			const int32 DestIndex = DestParameter.Offset;
-			const int32 SourceIndex = SourceUserParameterStore.IndexOf(DestParameter);
-			if (SourceIndex != INDEX_NONE)
+			if (DestParameter.IsDataInterface())
 			{
-				if (DestParameter.IsDataInterface())
-				{
-					check(OverrideParameters.GetDataInterface(DestIndex) != nullptr);
-					SourceUserParameterStore.GetDataInterface(SourceIndex)->CopyTo(OverrideParameters.GetDataInterface(DestIndex));
-					bInterfacesChanged = true;
-				}
-				else if (DestParameter.IsUObject())
-				{
-					OverrideParameters.SetUObject(SourceUserParameterStore.GetUObject(SourceIndex), DestIndex);
-				}
-				else
-				{
-					OverrideParameters.SetParameterData(SourceUserParameterStore.GetParameterData(SourceIndex, DestParameter.GetType()), DestIndex, DestParameter.GetSizeInBytes());
-				}
+				check(OverrideParameters.GetDataInterface(DestIndex) != nullptr);
+				SourceUserParameterStore.GetDataInterface(SourceIndex)->CopyTo(OverrideParameters.GetDataInterface(DestIndex));
+				bInterfacesChanged = true;
+			}
+			else if (DestParameter.IsUObject())
+			{
+				OverrideParameters.SetUObject(SourceUserParameterStore.GetUObject(SourceIndex), DestIndex);
 			}
 			else
 			{
-				ParametersToRemove.Add(DestParameter);
+				OverrideParameters.SetParameterData(SourceUserParameterStore.GetParameterData(SourceIndex, DestParameter.GetType()), DestIndex, DestParameter.GetSizeInBytes());
 			}
 		}
-
-		for (const FNiagaraVariableBase& ParameterToRemove : ParametersToRemove)
+		else
 		{
-			OverrideParameters.RemoveParameter(ParameterToRemove);
+			ParametersToRemove.Add(DestParameter);
 		}
+	}
 
-		if (bInterfacesChanged)
-		{
-			OverrideParameters.OnInterfaceChange();
-		}
+	for (const FNiagaraVariableBase& ParameterToRemove : ParametersToRemove)
+	{
+		OverrideParameters.RemoveParameter(ParameterToRemove);
 	}
-	else
+
+	if (bInterfacesChanged)
 	{
-		OverrideParameters.Empty(false);
-		CopyParametersFromAsset();
+		OverrideParameters.OnInterfaceChange();
 	}
 
 	OverrideParameters.Rebind();
