- Fix for single segment GPU ribbon UV calculation being incorrect

#tests Create a GPU ribbon with 2 segments and observe the UVs are incorrect compared to CPU

diff --git a/Engine/Plugins/FX/Niagara/Shaders/Private/Ribbons/NiagaraRibbonRibbonUVParamCalculation.usf b/Engine/Plugins/FX/Niagara/Shaders/Private/Ribbons/NiagaraRibbonRibbonUVParamCalculation.usf
--- a/Engine/Plugins/FX/Niagara/Shaders/Private/Ribbons/NiagaraRibbonRibbonUVParamCalculation.usf
+++ b/Engine/Plugins/FX/Niagara/Shaders/Private/Ribbons/NiagaraRibbonRibbonUVParamCalculation.usf
@@ -230,7 +230,7 @@ void GenerateRibbonUVParams(uint3 DispatchThreadId : SV_DispatchThreadID)
 		const FNiagaraRibbonUVSettings UV1Settings = GetUV1Params();
 
 		const int PackedPerRibbonDataOffsetUVs = PackedPerRibbonDataOffset + PACKED_PER_RIBBON_UVDATAOFFSET;
-		if ( NumSegments > 1 )
+		if ( NumSegments > 0 )
 		{
 			// UV0
 			BRANCH
