- LWE: Fix for color grad  icon not showing in editor

#rnx

diff --git a/Engine/Plugins/FX/Niagara/Source/NiagaraEditor/Private/NiagaraEditorStyle.cpp b/Engine/Plugins/FX/Niagara/Source/NiagaraEditor/Private/NiagaraEditorStyle.cpp
--- a/Engine/Plugins/FX/Niagara/Source/NiagaraEditor/Private/NiagaraEditorStyle.cpp
+++ b/Engine/Plugins/FX/Niagara/Source/NiagaraEditor/Private/NiagaraEditorStyle.cpp
@@ -452,8 +452,8 @@ void FNiagaraEditorStyle::InitIcons()
 
 	Set("NiagaraEditor.DistributionEditor.UniformCurve", new IMAGE_BRUSH_SVG("Icons/DistributionEditor/UniformCurve", Icon20x20));
 	Set("NiagaraEditor.DistributionEditor.NonUniformCurve", new IMAGE_BRUSH_SVG("Icons/DistributionEditor/NonUniformCurve", Icon20x20));
-	Set("NiagaraEditor.DistributionEditor.ColorGradient", new IMAGE_BRUSH_SVG("Icons/DistributionEditor/ColorGradient", Icon20x20));
-	
+	Set("NiagaraEditor.DistributionEditor.ColorGradient", new IMAGE_BRUSH("Icons/DistributionEditor/ColorGradient", Icon20x20));
+
 	Set("NiagaraEditor.DistributionEditor.CurveKeyHandle", new IMAGE_BRUSH_SVG("Icons/DistributionEditor/CurveKeyHandle", Icon20x20));
 }
 
