Niagara: Expose property getters and setters to blueprints for UMovieSceneNiagaraSystemSpawnSection

Added Set/GetSectionStartBehavior, Set/GetSectionEvaluateBehavior, Set/GetSectionEndBehavior, Set/GetAgeUpdateMode, Set/GetAllowScalability

diff --git a/Engine/Plugins/FX/Niagara/Source/Niagara/Private/MovieScene/MovieSceneNiagaraSystemSpawnSection.cpp b/Engine/Plugins/FX/Niagara/Source/Niagara/Private/MovieScene/MovieSceneNiagaraSystemSpawnSection.cpp
--- a/Engine/Plugins/FX/Niagara/Source/Niagara/Private/MovieScene/MovieSceneNiagaraSystemSpawnSection.cpp
+++ b/Engine/Plugins/FX/Niagara/Source/Niagara/Private/MovieScene/MovieSceneNiagaraSystemSpawnSection.cpp
@@ -20,27 +20,48 @@ ENiagaraSystemSpawnSectionStartBehavior UMovieSceneNiagaraSystemSpawnSection::Ge
 	return SectionStartBehavior;
 }
 
+void UMovieSceneNiagaraSystemSpawnSection::SetSectionStartBehavior(ENiagaraSystemSpawnSectionStartBehavior InBehavior)
+{
+	SectionStartBehavior = InBehavior;
+}
+
 ENiagaraSystemSpawnSectionEvaluateBehavior UMovieSceneNiagaraSystemSpawnSection::GetSectionEvaluateBehavior() const
 {
 	return SectionEvaluateBehavior;
 }
 
+void UMovieSceneNiagaraSystemSpawnSection::SetSectionEvaluateBehavior(ENiagaraSystemSpawnSectionEvaluateBehavior InBehavior)
+{
+	SectionEvaluateBehavior = InBehavior;
+}
+
 ENiagaraSystemSpawnSectionEndBehavior UMovieSceneNiagaraSystemSpawnSection::GetSectionEndBehavior() const
 {
 	return SectionEndBehavior;
 }
 
+void UMovieSceneNiagaraSystemSpawnSection::SetSectionEndBehavior(ENiagaraSystemSpawnSectionEndBehavior InBehavior)
+{
+	SectionEndBehavior = InBehavior;
+}
+
+
 ENiagaraAgeUpdateMode UMovieSceneNiagaraSystemSpawnSection::GetAgeUpdateMode() const
 {
 	return AgeUpdateMode;
 }
 
-void UMovieSceneNiagaraSystemSpawnSection::SetAgeUpdateMode(ENiagaraAgeUpdateMode NewMode)
+void UMovieSceneNiagaraSystemSpawnSection::SetAgeUpdateMode(ENiagaraAgeUpdateMode InMode)
 {
-	AgeUpdateMode = NewMode;
+	AgeUpdateMode = InMode;
 }
 
-bool UMovieSceneNiagaraSystemSpawnSection::GetAllowScalability()const
+bool UMovieSceneNiagaraSystemSpawnSection::GetAllowScalability() const
 {
 	return bAllowScalability;
 }
+
+void UMovieSceneNiagaraSystemSpawnSection::SetAllowScalability(bool bInAllowScalability)
+{
+	bAllowScalability = bInAllowScalability;
+}
diff --git a/Engine/Plugins/FX/Niagara/Source/Niagara/Private/MovieScene/MovieSceneNiagaraSystemTrackTemplate.h b/Engine/Plugins/FX/Niagara/Source/Niagara/Private/MovieScene/MovieSceneNiagaraSystemTrackTemplate.h
--- a/Engine/Plugins/FX/Niagara/Source/Niagara/Private/MovieScene/MovieSceneNiagaraSystemTrackTemplate.h
+++ b/Engine/Plugins/FX/Niagara/Source/Niagara/Private/MovieScene/MovieSceneNiagaraSystemTrackTemplate.h
@@ -7,9 +7,9 @@
 #include "MovieSceneNiagaraSystemTrackTemplate.generated.h"
 
 enum class ENiagaraAgeUpdateMode : uint8;
-enum class ENiagaraSystemSpawnSectionEndBehavior;
-enum class ENiagaraSystemSpawnSectionEvaluateBehavior;
-enum class ENiagaraSystemSpawnSectionStartBehavior;
+enum class ENiagaraSystemSpawnSectionEndBehavior : uint8;
+enum class ENiagaraSystemSpawnSectionEvaluateBehavior : uint8;
+enum class ENiagaraSystemSpawnSectionStartBehavior : uint8;
 
 USTRUCT()
 struct FMovieSceneNiagaraSystemTrackTemplate : public FMovieSceneEvalTemplate
diff --git a/Engine/Plugins/FX/Niagara/Source/Niagara/Public/MovieScene/MovieSceneNiagaraSystemSpawnSection.h b/Engine/Plugins/FX/Niagara/Source/Niagara/Public/MovieScene/MovieSceneNiagaraSystemSpawnSection.h
--- a/Engine/Plugins/FX/Niagara/Source/Niagara/Public/MovieScene/MovieSceneNiagaraSystemSpawnSection.h
+++ b/Engine/Plugins/FX/Niagara/Source/Niagara/Public/MovieScene/MovieSceneNiagaraSystemSpawnSection.h
@@ -7,8 +7,8 @@
 #include "MovieSceneNiagaraSystemSpawnSection.generated.h"
 
 /** Defines options for system life cycle for before the section is evaluating up to the first frame the section evaluates. */
-UENUM()
-enum class ENiagaraSystemSpawnSectionStartBehavior
+UENUM(BlueprintType)
+enum class ENiagaraSystemSpawnSectionStartBehavior : uint8
 {
 	/** When the time before the section evaluates the particle system's component will be deactivated and on the first frame of the section the
 	 system's component will be activated. */
@@ -16,8 +16,8 @@ enum class ENiagaraSystemSpawnSectionStartBehavior
 };
 
 /** Defines options for system life cycle for when the section is evaluating from the 2nd frame until the last frame of the section. */
-UENUM()
-enum class ENiagaraSystemSpawnSectionEvaluateBehavior
+UENUM(BlueprintType)
+enum class ENiagaraSystemSpawnSectionEvaluateBehavior : uint8
 {
 	/** The system's component will be activated on any frame where it is inactive.  This is useful for continuous emitters, especially if the sequencer will start in the middle of the section. */
 	ActivateIfInactive,
@@ -26,8 +26,8 @@ enum class ENiagaraSystemSpawnSectionEvaluateBehavior
 };
 
 /** Defines options for system life cycle for the time after the section. */
-UENUM()
-enum class ENiagaraSystemSpawnSectionEndBehavior
+UENUM(BlueprintType)
+enum class ENiagaraSystemSpawnSectionEndBehavior : uint8
 {
 	//** When the section ends the system is set to inactive which stops spawning but lets existing particles simulate until death.
 	SetSystemInactive,
@@ -45,16 +45,30 @@ class UMovieSceneNiagaraSystemSpawnSection : public UMovieSceneSection
 public:
 	UMovieSceneNiagaraSystemSpawnSection();
 
+	UFUNCTION(BlueprintPure, Category = "Niagara|Life Cycle")
 	ENiagaraSystemSpawnSectionStartBehavior GetSectionStartBehavior() const;
+	UFUNCTION(BlueprintCallable, Category = "Niagara|Life Cycle")
+	NIAGARA_API void SetSectionStartBehavior(ENiagaraSystemSpawnSectionStartBehavior InBehavior);
 
+	UFUNCTION(BlueprintPure, Category = "Niagara|Life Cycle")
 	ENiagaraSystemSpawnSectionEvaluateBehavior GetSectionEvaluateBehavior() const;
+	UFUNCTION(BlueprintCallable, Category = "Niagara|Life Cycle")
+	NIAGARA_API void SetSectionEvaluateBehavior(ENiagaraSystemSpawnSectionEvaluateBehavior InBehavior);
 
+	UFUNCTION(BlueprintPure, Category = "Niagara|Life Cycle")
 	ENiagaraSystemSpawnSectionEndBehavior GetSectionEndBehavior() const;
+	UFUNCTION(BlueprintCallable, Category = "Niagara|Life Cycle")
+	NIAGARA_API void SetSectionEndBehavior(ENiagaraSystemSpawnSectionEndBehavior InBehavior);
 
+	UFUNCTION(BlueprintPure, Category = "Niagara|Life Cycle")
 	NIAGARA_API ENiagaraAgeUpdateMode GetAgeUpdateMode() const;
-	NIAGARA_API void SetAgeUpdateMode(ENiagaraAgeUpdateMode NewMode);
+	UFUNCTION(BlueprintCallable, Category = "Niagara|Life Cycle")
+	NIAGARA_API void SetAgeUpdateMode(ENiagaraAgeUpdateMode InMode);
 
+	UFUNCTION(BlueprintPure, Category = "Niagara|Life Cycle")
 	bool GetAllowScalability() const;
+	UFUNCTION(BlueprintCallable, Category = "Niagara|Life Cycle")
+	void SetAllowScalability(bool bInAllowScalability);
 
 private:
 	/** Specifies what should happen to the niagara system from before the section evaluates up until the first frame of the section. */
