- Don't call the OnSystemFinished delegate if we are construction blueprints
- The delegate might intefer with the actor / components and can cause crashes

diff --git a/Engine/Plugins/FX/Niagara/Source/Niagara/Private/NiagaraComponent.cpp b/Engine/Plugins/FX/Niagara/Source/Niagara/Private/NiagaraComponent.cpp
--- a/Engine/Plugins/FX/Niagara/Source/Niagara/Private/NiagaraComponent.cpp
+++ b/Engine/Plugins/FX/Niagara/Source/Niagara/Private/NiagaraComponent.cpp
@@ -1806,7 +1806,14 @@ void UNiagaraComponent::OnSystemComplete(bool bExternalCompletion)
 		DestroyCullProxy();
 
 		//UE_LOG(LogNiagara, Log, TEXT("OnSystemFinished.Broadcast(this);: { %p -  %p - %s"), this, SystemInstance.Get(), *Asset->GetName());
-		OnSystemFinished.Broadcast(this);
+#if WITH_EDITOR
+		// When in the editor do not invoke the system finished delegate if we are reconstrucing blueprint instances as they can have side effects.
+		// For example the finished delegate could destroy the actor, attach / remove components / etc
+		if (!GIsReconstructingBlueprintInstances)
+#endif
+		{
+			OnSystemFinished.Broadcast(this);
+		}
 		//UE_LOG(LogNiagara, Log, TEXT("OnSystemFinished.Broadcast(this);: } %p - %s"), SystemInstance.Get(), *Asset->GetName());
 
 		if (PoolingMethod == ENCPoolMethod::AutoRelease)//Don't release back to the pool if we're completing due to scalability culling.
