Comment correction.

diff --git a/Engine/Plugins/FX/Niagara/Source/Niagara/Internal/DataInterface/NiagaraDataInterfaceDataChannelWrite.h b/Engine/Plugins/FX/Niagara/Source/Niagara/Internal/DataInterface/NiagaraDataInterfaceDataChannelWrite.h
--- a/Engine/Plugins/FX/Niagara/Source/Niagara/Internal/DataInterface/NiagaraDataInterfaceDataChannelWrite.h
+++ b/Engine/Plugins/FX/Niagara/Source/Niagara/Internal/DataInterface/NiagaraDataInterfaceDataChannelWrite.h
@@ -160,7 +160,7 @@ public:
 	UPROPERTY(EditAnywhere, Category = "Data Channel")
 	TObjectPtr<UNiagaraDataChannelAsset> Channel;
 
-	/** Access Context used when accessing NDC data. Controls which of the NDCs internal data it will return for this Interface to write into. Can be empty for Systems using bAutoLinkToSpawningNDC. */
+	/** Access Context used when accessing NDC data. Controls which of the NDCs internal data it will return for this Interface to write into. */
 	UPROPERTY(EditAnywhere, Category = "Data Channel")
 	FNDCAccessContextInst AccessContext;
 
