Impose edit conditions for some usage of NDC Access Context to better guide usage.

diff --git a/Engine/Plugins/FX/Niagara/Source/Niagara/Internal/DataInterface/NiagaraDataInterfaceDataChannelRead.h b/Engine/Plugins/FX/Niagara/Source/Niagara/Internal/DataInterface/NiagaraDataInterfaceDataChannelRead.h
--- a/Engine/Plugins/FX/Niagara/Source/Niagara/Internal/DataInterface/NiagaraDataInterfaceDataChannelRead.h
+++ b/Engine/Plugins/FX/Niagara/Source/Niagara/Internal/DataInterface/NiagaraDataInterfaceDataChannelRead.h
@@ -90,7 +90,7 @@ public:
 	TObjectPtr<UNiagaraDataChannelAsset> Channel;
 
 	/** Access Context used when accessing NDC data. Controls which of the NDCs internal data it will return for this Interface to read. Can be empty for Systems using bAutoLinkToSpawningNDC. */
-	UPROPERTY(EditAnywhere, Category = "Parameters", meta = (ShowOnlyInnerProperties))
+	UPROPERTY(EditAnywhere, Category = "Parameters", meta = (ShowOnlyInnerProperties, EditCondition="!bAutoLinkToSpawningNDC"))
 	FNDCAccessContextInst AccessContext;
 
 	/** If true this DI will attempt to automatically link itself to the NDC it's being spawned by. */
diff --git a/Engine/Plugins/FX/Niagara/Source/Niagara/Public/NiagaraDataChannelPublic.h b/Engine/Plugins/FX/Niagara/Source/Niagara/Public/NiagaraDataChannelPublic.h
--- a/Engine/Plugins/FX/Niagara/Source/Niagara/Public/NiagaraDataChannelPublic.h
+++ b/Engine/Plugins/FX/Niagara/Source/Niagara/Public/NiagaraDataChannelPublic.h
@@ -603,11 +603,11 @@ struct FNiagaraDataChannelReference
 	UPROPERTY(EditAnywhere, BlueprintReadWrite, Category = DataChannel)
 	TObjectPtr<UNiagaraDataChannelAsset> DataChannel;
 
-	UPROPERTY(EditAnywhere, BlueprintReadWrite, Category = DataChannel, meta = (ShowOnlyInnerProperties, ShowNDCAccessContextInput))
+	UPROPERTY(EditAnywhere, BlueprintReadWrite, Category = DataChannel, meta = (EditCondition = "bCustomAccessContext" , ShowOnlyInnerProperties, ShowNDCAccessContextInput))
 	mutable FNDCAccessContextInst AccessContext;
 
 	/** If true this reference will provide a custom access context. */
-	UPROPERTY(EditAnywhere, BlueprintReadOnly, Category = DataChannel)
+	UPROPERTY(EditAnywhere, BlueprintReadOnly, Category = DataChannel, meta = (InlineEditCondition))
 	bool bCustomAccessContext = false;
 
 private:
