Extending RigidMeshCollisionQuery NDI: Exposing first element of each type & extents
#tests Validated functions return valid data in a test map

diff --git a/Engine/Plugins/FX/Niagara/Shaders/Private/NiagaraDataInterfaceRigidMeshCollisionQuery.ush b/Engine/Plugins/FX/Niagara/Shaders/Private/NiagaraDataInterfaceRigidMeshCollisionQuery.ush
--- a/Engine/Plugins/FX/Niagara/Shaders/Private/NiagaraDataInterfaceRigidMeshCollisionQuery.ush
+++ b/Engine/Plugins/FX/Niagara/Shaders/Private/NiagaraDataInterfaceRigidMeshCollisionQuery.ush
@@ -54,6 +54,36 @@ void GetNumElements_{ParameterName}(out int OutNumElements)
 	OutNumElements = {ParameterName}_ElementOffsets[NUM_ELEMENTS_INDEX];
 }
 
+void GetBoxElementsStartIndex_{ParameterName}(out int OutStartIndex)
+{
+	OutStartIndex = {ParameterName}_ElementOffsets[BOX_INDEX];
+}
+
+void GetSphereElementsStartIndex_{ParameterName}(out int OutStartIndex)
+{
+	OutStartIndex = {ParameterName}_ElementOffsets[SPHERE_INDEX];
+}
+
+void GetCapsuleElementsStartIndex_{ParameterName}(out int OutStartIndex)
+{
+	OutStartIndex = {ParameterName}_ElementOffsets[CAPSULE_INDEX];
+}
+
+void GetSphereRadius_{ParameterName}(in int ElementIndex, out float OutRadius)
+{
+	OutRadius = {ParameterName}_ElementExtentBuffer[ElementIndex].x;
+}
+
+void GetCapsuleSize_{ParameterName}(in int ElementIndex, out float2 OutSize)
+{
+	OutSize = {ParameterName}_ElementExtentBuffer[ElementIndex].xy;
+}
+
+void GetBoxSize_{ParameterName}(in int ElementIndex, out float3 OutSize)
+{
+	OutSize = {ParameterName}_ElementExtentBuffer[ElementIndex].xyz;
+}
+
 /* -----------------------------------------------------------------
  * Position / Transform utils
  * -----------------------------------------------------------------
diff --git a/Engine/Plugins/FX/Niagara/Source/Niagara/Private/NiagaraDataInterfaceRigidMeshCollisionQuery.cpp b/Engine/Plugins/FX/Niagara/Source/Niagara/Private/NiagaraDataInterfaceRigidMeshCollisionQuery.cpp
--- a/Engine/Plugins/FX/Niagara/Source/Niagara/Private/NiagaraDataInterfaceRigidMeshCollisionQuery.cpp
+++ b/Engine/Plugins/FX/Niagara/Source/Niagara/Private/NiagaraDataInterfaceRigidMeshCollisionQuery.cpp
@@ -92,6 +92,12 @@ static const FName GetNumBoxesName(TEXT("GetNumBoxes"));
 static const FName GetNumSpheresName(TEXT("GetNumSpheres"));
 static const FName GetNumCapsulesName(TEXT("GetNumCapsules"));
 static const FName GetNumElementsName(TEXT("GetNumElements"));
+static const FName GetBoxElementsStartIndexName(TEXT("GetBoxElementsStartIndex"));
+static const FName GetSphereElementsStartIndexName(TEXT("GetSphereElementsStartIndex"));
+static const FName GetCapsuleElementsStartIndexName(TEXT("GetCapsuleElementsStartIndex"));
+static const FName GetSphereRadiusName(TEXT("GetSphereRadius"));
+static const FName GetCapsuleSizeName(TEXT("GetCapsuleSize"));
+static const FName GetBoxSizeName(TEXT("GetBoxSize"));
 
 //------------------------------------------------------------------------------------------------------------
 static const FName IsWorldPositionInsideCombinedBoundsName(TEXT("IsWorldPositionInsideCombinedBounds"));
@@ -1515,6 +1521,87 @@ void UNiagaraDataInterfaceRigidMeshCollisionQuery::GetFunctionsInternal(TArray<F
 
 		OutFunctions.Add(Sig);
 	}
+	{
+		FNiagaraFunctionSignature Sig;
+		Sig.Name = GetBoxElementsStartIndexName;
+		Sig.SetDescription(LOCTEXT("GetBoxElementsStartIndexNameDescription", "Returns the Index of the first of box primitive for the collection of static meshes the DI represents."));
+		Sig.SetFunctionVersion(FNiagaraRigidMeshCollisionDIFunctionVersion::LatestVersion);
+		Sig.bSupportsGPU = true;
+		Sig.bSupportsCPU = false;
+		Sig.bMemberFunction = true;
+		Sig.Inputs.Add(FNiagaraVariable(FNiagaraTypeDefinition(GetClass()), TEXT("Collision DI")));
+		Sig.Outputs.Add(FNiagaraVariable(FNiagaraTypeDefinition::GetIntDef(), TEXT("Index")));
+
+		OutFunctions.Add(Sig);
+	}
+	{
+		FNiagaraFunctionSignature Sig;
+		Sig.Name = GetSphereElementsStartIndexName;
+		Sig.SetDescription(LOCTEXT("GetSphereElementsStartIndexNameDescription", "Returns the Index of the first of sphere primitive for the collection of static meshes the DI represents."));
+		Sig.SetFunctionVersion(FNiagaraRigidMeshCollisionDIFunctionVersion::LatestVersion);
+		Sig.bSupportsGPU = true;
+		Sig.bSupportsCPU = false;
+		Sig.bMemberFunction = true;
+		Sig.Inputs.Add(FNiagaraVariable(FNiagaraTypeDefinition(GetClass()), TEXT("Collision DI")));
+		Sig.Outputs.Add(FNiagaraVariable(FNiagaraTypeDefinition::GetIntDef(), TEXT("Index")));
+
+		OutFunctions.Add(Sig);
+	}
+	{
+		FNiagaraFunctionSignature Sig;
+		Sig.Name = GetCapsuleElementsStartIndexName;
+		Sig.SetDescription(LOCTEXT("GetCapsuleElementsStartIndexNameDescription", "Returns the Index of the first of capsule primitive for the collection of static meshes the DI represents."));
+		Sig.SetFunctionVersion(FNiagaraRigidMeshCollisionDIFunctionVersion::LatestVersion);
+		Sig.bSupportsGPU = true;
+		Sig.bSupportsCPU = false;
+		Sig.bMemberFunction = true;
+		Sig.Inputs.Add(FNiagaraVariable(FNiagaraTypeDefinition(GetClass()), TEXT("Collision DI")));
+		Sig.Outputs.Add(FNiagaraVariable(FNiagaraTypeDefinition::GetIntDef(), TEXT("Index")));
+
+		OutFunctions.Add(Sig);
+	}
+	{
+		FNiagaraFunctionSignature Sig;
+		Sig.Name = GetSphereRadiusName;
+		Sig.SetDescription(LOCTEXT("GetSphereRadiusNameDescription", "Returns the radius of a sphere primitive in the collection of static meshes the DI represents."));
+		Sig.SetFunctionVersion(FNiagaraRigidMeshCollisionDIFunctionVersion::LatestVersion);
+		Sig.bSupportsGPU = true;
+		Sig.bSupportsCPU = false;
+		Sig.bMemberFunction = true;
+		Sig.Inputs.Add(FNiagaraVariable(FNiagaraTypeDefinition(GetClass()), TEXT("Collision DI")));
+		Sig.Inputs.Add(FNiagaraVariable(FNiagaraTypeDefinition::GetIntDef(), TEXT("Element Index")));
+		Sig.Outputs.Add(FNiagaraVariable(FNiagaraTypeDefinition::GetFloatDef(), TEXT("Radius")));
+
+		OutFunctions.Add(Sig);
+	}
+	{
+		FNiagaraFunctionSignature Sig;
+		Sig.Name = GetCapsuleSizeName;
+		Sig.SetDescription(LOCTEXT("GetCapsuleSizeNameDescription", "Returns the capsule size of a capsule primitive in the collection of static meshes the DI represents."));
+		Sig.SetFunctionVersion(FNiagaraRigidMeshCollisionDIFunctionVersion::LatestVersion);
+		Sig.bSupportsGPU = true;
+		Sig.bSupportsCPU = false;
+		Sig.bMemberFunction = true;
+		Sig.Inputs.Add(FNiagaraVariable(FNiagaraTypeDefinition(GetClass()), TEXT("Collision DI")));
+		Sig.Inputs.Add(FNiagaraVariable(FNiagaraTypeDefinition::GetIntDef(), TEXT("Element Index")));
+		Sig.Outputs.Add(FNiagaraVariable(FNiagaraTypeDefinition::GetVec2Def(), TEXT("Size")));
+
+		OutFunctions.Add(Sig);
+	}
+	{
+		FNiagaraFunctionSignature Sig;
+		Sig.Name = GetBoxSizeName;
+		Sig.SetDescription(LOCTEXT("GetBoxSizeNameNameDescription", "Returns the Extents of a box primitive in the collection of static meshes the DI represents."));
+		Sig.SetFunctionVersion(FNiagaraRigidMeshCollisionDIFunctionVersion::LatestVersion);
+		Sig.bSupportsGPU = true;
+		Sig.bSupportsCPU = false;
+		Sig.bMemberFunction = true;
+		Sig.Inputs.Add(FNiagaraVariable(FNiagaraTypeDefinition(GetClass()), TEXT("Collision DI")));
+		Sig.Inputs.Add(FNiagaraVariable(FNiagaraTypeDefinition::GetIntDef(), TEXT("Element Index")));
+		Sig.Outputs.Add(FNiagaraVariable(FNiagaraTypeDefinition::GetVec3Def(), TEXT("Size")));
+
+		OutFunctions.Add(Sig);
+	}
 	{
 		FNiagaraFunctionSignature Sig;
 		Sig.Name = IsWorldPositionInsideCombinedBoundsName;
@@ -1770,6 +1857,12 @@ bool UNiagaraDataInterfaceRigidMeshCollisionQuery::GetFunctionHLSL(const FNiagar
 		(FunctionInfo.DefinitionName == GetNumCapsulesName) ||
 		(FunctionInfo.DefinitionName == GetNumSpheresName) ||
 		(FunctionInfo.DefinitionName == GetNumElementsName) ||
+		(FunctionInfo.DefinitionName == GetBoxElementsStartIndexName) ||
+		(FunctionInfo.DefinitionName == GetSphereElementsStartIndexName) ||
+		(FunctionInfo.DefinitionName == GetCapsuleElementsStartIndexName) ||
+		(FunctionInfo.DefinitionName == GetSphereRadiusName) ||
+		(FunctionInfo.DefinitionName == GetCapsuleSizeName) ||
+		(FunctionInfo.DefinitionName == GetBoxSizeName) ||
 		(FunctionInfo.DefinitionName == IsWorldPositionInsideCombinedBoundsName) ||
 		(FunctionInfo.DefinitionName == GetClosestPointName) ||
 		(FunctionInfo.DefinitionName == GetClosestPointSimpleName) ||		
