- Add support for particle reset commands to entities

diff --git a/Engine/Plugins/FX/Niagara/Source/NiagaraEditor/Private/NiagaraEditorModule.cpp b/Engine/Plugins/FX/Niagara/Source/NiagaraEditor/Private/NiagaraEditorModule.cpp
--- a/Engine/Plugins/FX/Niagara/Source/NiagaraEditor/Private/NiagaraEditorModule.cpp
+++ b/Engine/Plugins/FX/Niagara/Source/NiagaraEditor/Private/NiagaraEditorModule.cpp
@@ -2372,24 +2372,31 @@ void FNiagaraEditorModule::OnExecParticleInvoked(const TCHAR* Str)
 	// Very similar logic to UEditorEngine::Exec_Particle
 	if (FParse::Command(&Str, TEXT("RESET")))
 	{
-		TArray<AEmitter*> EmittersToReset;
-		if (FParse::Command(&Str, TEXT("SELECTED")))
-		{
-			// Reset any selected emitters in the level
-			for (FSelectionIterator It(GEditor->GetSelectedActorIterator()); It; ++It)
+		auto ReinitActorComponents =
+			[](AActor* Actor)
 			{
-				AActor* Actor = static_cast<AActor*>(*It);
-				checkSlow(Actor->IsA(AActor::StaticClass()));
-
+				TArray<UNiagaraComponent*, TInlineAllocator<1>> ComponentsToReset;
 				for (UActorComponent* AC : Actor->GetComponents())
 				{
-					UNiagaraComponent* NiagaraComponent = Cast<UNiagaraComponent>(AC);
-					if (NiagaraComponent)
+					if (UNiagaraComponent* NiagaraComponent = Cast<UNiagaraComponent>(AC))
 					{
-						NiagaraComponent->Activate(true);
-						NiagaraComponent->ReregisterComponent();
+						ComponentsToReset.Add(NiagaraComponent);
 					}
 				}
+
+				for (UNiagaraComponent* NiagaraComponent : ComponentsToReset)
+				{
+					NiagaraComponent->ReinitializeSystem();
+				}
+			};
+
+		if (FParse::Command(&Str, TEXT("SELECTED")))
+		{
+			// Reset any selected emitters in the level
+			for (FSelectionIterator It(GEditor->GetSelectedActorIterator()); It; ++It)
+			{
+				AActor* Actor = CastChecked<AActor>(*It);
+				ReinitActorComponents(Actor);
 			}
 		}
 		else if (FParse::Command(&Str, TEXT("ALL")))
@@ -2397,21 +2404,7 @@ void FNiagaraEditorModule::OnExecParticleInvoked(const TCHAR* Str)
 			// Reset ALL emitters in the level
 			for (TObjectIterator<AActor> It; It; ++It)
 			{
-				// Note: Must gather then reset as the array we are iterating will be modified with the ReregisterComponent call
-				TArray<UNiagaraComponent*> ComponentsToReset;
-				for (UActorComponent* AC : It->GetComponents())
-				{
-					if (UNiagaraComponent* NiagaraComponent = Cast<UNiagaraComponent>(AC))
-					{
-						ComponentsToReset.Add(NiagaraComponent);
-					}
-				}
-
-				for (UNiagaraComponent* NiagaraComponent : ComponentsToReset)
-				{
-					NiagaraComponent->Activate(true);
-					NiagaraComponent->ReregisterComponent();
-				}
+				ReinitActorComponents(*It);
 			}
 		}
 	}
