Add log error when using legacy access path with new NDC type that does not support it.

diff --git a/Engine/Plugins/FX/Niagara/Source/Niagara/Private/NiagaraDataChannelHandler.cpp b/Engine/Plugins/FX/Niagara/Source/Niagara/Private/NiagaraDataChannelHandler.cpp
--- a/Engine/Plugins/FX/Niagara/Source/Niagara/Private/NiagaraDataChannelHandler.cpp
+++ b/Engine/Plugins/FX/Niagara/Source/Niagara/Private/NiagaraDataChannelHandler.cpp
@@ -215,6 +215,9 @@ void UNiagaraDataChannelHandler::VerifyAccessContext(FNDCAccessContextInst& Acce
 
 FNiagaraDataChannelDataPtr UNiagaraDataChannelHandler::FindData(FNiagaraDataChannelSearchParameters SearchParams, ENiagaraResourceAccess AccessType)
 {
+	UE_LOGFMT(LogNiagara, Error, 
+	"Using legacy SearchParams API when accessing Niagara Data Channel that does not support the legacy path.\nPlease update access code to use the newer Access Context based access path.\nData Channel: {DataChannel}",
+	DataChannel->GetAsset() ? DataChannel->GetAsset()->GetPathName() : TEXT(""));
 	return nullptr;
 }
 
