Streamline NDC BP flow slightly.
- New PrepareAccessContext function that wraps up a couple of calls into one.
- Fixed issue with node expansion with multiple Then connections.

diff --git a/Engine/Plugins/FX/Niagara/Source/Niagara/Private/NiagaraDataChannel.cpp b/Engine/Plugins/FX/Niagara/Source/Niagara/Private/NiagaraDataChannel.cpp
--- a/Engine/Plugins/FX/Niagara/Source/Niagara/Private/NiagaraDataChannel.cpp
+++ b/Engine/Plugins/FX/Niagara/Source/Niagara/Private/NiagaraDataChannel.cpp
@@ -1603,7 +1603,7 @@ FNDCAccessContextInst UNiagaraDataChannelLibrary::MakeNDCAccessContextInstance(U
 	return FNDCAccessContextInst();
 }
 
-void UNiagaraDataChannelLibrary::SetMembersInNDCAccessContextInstance(UPARAM(Ref) const FNDCAccessContextInst& AccessContext, UScriptStruct* ContextStruct)
+void UNiagaraDataChannelLibrary::SetMembersInNDCAccessContextInstance(FNDCAccessContextInst& AccessContext, UScriptStruct* ContextStruct)
 {
 	// We should never hit this! stubs to avoid NoExport on the class.
 	checkNoEntry();
@@ -1621,7 +1621,7 @@ void UNiagaraDataChannelLibrary::GetSinglePropertyInNDCAccessContextInstance(con
 	checkNoEntry();
 }
 
-void UNiagaraDataChannelLibrary::SetSinglePropertyInNDCAccessContextInstance(UPARAM(Ref) const FNDCAccessContextInst& AccessContext, UScriptStruct* ContextStruct, FName PropertyName, const int32& Value)
+void UNiagaraDataChannelLibrary::SetSinglePropertyInNDCAccessContextInstance(FNDCAccessContextInst& AccessContext, UScriptStruct* ContextStruct, FName PropertyName, const int32& Value)
 {
 	// We should never hit this! stubs to avoid NoExport on the class.
 	checkNoEntry();
@@ -1881,6 +1881,13 @@ FNDCAccessContextInst& UNiagaraDataChannelLibrary::GetUsableAccessContextFromNDC
 	return NDCRef.GetUsableAccessContext();
 }
 
+FNDCAccessContextInst& UNiagaraDataChannelLibrary::PrepareAccessContextFromNDCRef(FNiagaraDataChannelReference& NDCRef)
+{
+	check(0);//Dummy function replaced by BP compiler. We should not get here.
+	static FNDCAccessContextInst Dummy;
+	return Dummy;
+}
+
 //////////////////////////////////////////////////////////////////////////
 // Functions for modifying dynamic objects.
 // 
diff --git a/Engine/Plugins/FX/Niagara/Source/Niagara/Public/NiagaraDataChannel.h b/Engine/Plugins/FX/Niagara/Source/Niagara/Public/NiagaraDataChannel.h
--- a/Engine/Plugins/FX/Niagara/Source/Niagara/Public/NiagaraDataChannel.h
+++ b/Engine/Plugins/FX/Niagara/Source/Niagara/Public/NiagaraDataChannel.h
@@ -493,13 +493,21 @@ class UNiagaraDataChannelLibrary : public UBlueprintFunctionLibrary
 
 
 	/** Returns the access context for use when accessing a Niagara Data channel. */
-	UFUNCTION(BlueprintCallable, Category = "Niagara Data Channel")
+	UFUNCTION(BlueprintCallable, BlueprintInternalUseOnly, Category = "Niagara Data Channel")
 	static UE_API FNDCAccessContextInst& GetUsableAccessContextFromNDC(const UNiagaraDataChannelAsset* DataChannel);
 
 	/** Returns the access context for use when accessing a Niagara Data channel. */
-	UFUNCTION(BlueprintCallable, Category = "Niagara Data Channel")
+	UFUNCTION(BlueprintCallable, BlueprintInternalUseOnly, Category = "Niagara Data Channel")
 	static UE_API FNDCAccessContextInst& GetUsableAccessContextFromNDCRef(const FNiagaraDataChannelReference& NDCRef);
 
+	/** 
+	Prepares an Access Context ready for accessing a Niagara Data Channel. 
+	If the passed Data Channel Reference has a valid custom AccessContext then it's values will be applied. Otherwise the defaults will be used.
+	Allows optional setting for the members of the Access Context. Transient properties are shown by default but properties to modify can be selected in the details panel.
+	*/
+	UFUNCTION(BlueprintInternalUseOnly, BlueprintCallable, Category = "Niagara Data Channel")
+	static UE_API FNDCAccessContextInst& PrepareAccessContextFromNDCRef(UPARAM(Ref) FNiagaraDataChannelReference& NDCRef);
+
 	/** Creates a new NDCAccessContext with a specified internal struct type. Struct must be a child of FNDCAccessContextBase. */
 	UFUNCTION(BlueprintInternalUseOnly, BlueprintCallable, CustomThunk, Category = "Niagara Data Channel", meta = (HidePin = "ContextStruct"))
 	static UE_API FNDCAccessContextInst MakeNDCAccessContextInstance(UScriptStruct* ContextStruct);
@@ -514,7 +522,7 @@ class UNiagaraDataChannelLibrary : public UBlueprintFunctionLibrary
 	Sets the members in the internal Struct of an FNDCAccessContextIns. Expanded on compile into individual per property calls to SetSinglePropertyInNDCAccessContextInstance.
 	*/
 	UFUNCTION(BlueprintInternalUseOnly, BlueprintCallable, Category = "Niagara Data Channel", meta = (HidePin = "ContextStruct"))
-	static UE_API void SetMembersInNDCAccessContextInstance(UPARAM(Ref) const FNDCAccessContextInst& AccessContext, UScriptStruct* ContextStruct);
+	static UE_API void SetMembersInNDCAccessContextInstance(UPARAM(Ref) FNDCAccessContextInst& AccessContext, UScriptStruct* ContextStruct);
 	
 	/** 
 	Gets a specific property within the internal data of an FNDCAccessContextInst.
@@ -526,7 +534,7 @@ class UNiagaraDataChannelLibrary : public UBlueprintFunctionLibrary
 	Sets a specific property within the internal data of an FNDCAccessContextInst.
 	*/
 	UFUNCTION(BlueprintInternalUseOnly, BlueprintCallable, CustomThunk, Category = "Niagara Data Channel", meta = (CustomStructureParam = "Value", AutoCreateRefTerm="Value"))
-	static UE_API void SetSinglePropertyInNDCAccessContextInstance(const FNDCAccessContextInst& AccessContext, UScriptStruct* ContextStruct, FName PropertyName, const int32& Value);
+	static UE_API void SetSinglePropertyInNDCAccessContextInstance(UPARAM(Ref) FNDCAccessContextInst& AccessContext, UScriptStruct* ContextStruct, FName PropertyName, const int32& Value);
 
 private:
 	DECLARE_FUNCTION(execMakeNDCAccessContextInstance);
diff --git a/Engine/Plugins/FX/Niagara/Source/NiagaraBlueprintNodes/Internal/K2Node_DataChannel_WithContext.h b/Engine/Plugins/FX/Niagara/Source/NiagaraBlueprintNodes/Internal/K2Node_DataChannel_WithContext.h
--- a/Engine/Plugins/FX/Niagara/Source/NiagaraBlueprintNodes/Internal/K2Node_DataChannel_WithContext.h
+++ b/Engine/Plugins/FX/Niagara/Source/NiagaraBlueprintNodes/Internal/K2Node_DataChannel_WithContext.h
@@ -204,6 +204,27 @@ protected:
 	[[nodiscard]] virtual FText GetTooltipTextForContext(UScriptStruct* Struct) const override;
 };
 
+/** Prepares a usable access context ready for accessing an NDC. */
+UCLASS()
+class UK2Node_DataChannelAccessContext_Prepare : public UK2Node_DataChannelAccessContextOperation
+{
+	GENERATED_BODY()
+public:
+
+	UK2Node_DataChannelAccessContext_Prepare(const FObjectInitializer& ObjectInitializer);
+
+	//~ Begin UEdGraphNode Interface
+	virtual void AllocateDefaultPins() override;
+	//End UEdGraphNode Interface.
 
+	//Begin UK2Node Interface
+	virtual void ExpandNode(class FKismetCompilerContext& CompilerContext, UEdGraph* SourceGraph) override;
+	//END UK2Node Interface
+
+protected:
+
+	[[nodiscard]] virtual FText GetNodeTitleForContext(UScriptStruct* Struct) const override;
+	[[nodiscard]] virtual FText GetTooltipTextForContext(UScriptStruct* Struct) const override;
+};
 
 //////////////////////////////////////////////////////////////////////////
\ No newline at end of file
diff --git a/Engine/Plugins/FX/Niagara/Source/NiagaraBlueprintNodes/Private/K2Node_DataChannel_WithContext.cpp b/Engine/Plugins/FX/Niagara/Source/NiagaraBlueprintNodes/Private/K2Node_DataChannel_WithContext.cpp
--- a/Engine/Plugins/FX/Niagara/Source/NiagaraBlueprintNodes/Private/K2Node_DataChannel_WithContext.cpp
+++ b/Engine/Plugins/FX/Niagara/Source/NiagaraBlueprintNodes/Private/K2Node_DataChannel_WithContext.cpp
@@ -249,7 +249,6 @@ void UK2Node_DataChannelAccessContextOperation::PreloadRequiredAssets()
 
 void UK2Node_DataChannelAccessContextOperation::ExpandNodeOptionalPins(class FKismetCompilerContext& CompilerContext, UEdGraph* SourceGraph, UEdGraphPin* AccessContextSourcePin, UEdGraphPin* CurrThen, FName PerPropertyFuncName)
 {
-	UEdGraphPin* NextExecPin = CurrThen->LinkedTo.Num() > 0 ? CurrThen->LinkedTo[0] : nullptr;
 	check(CurrThen);
 
 	//For all of our optional property pins, create shunt the pin to a new function call to perform the operation on that property only.
@@ -286,9 +285,9 @@ void UK2Node_DataChannelAccessContextOperation::ExpandNodeOptionalPins(class FKi
 				//AutoCreateRefTerm on the function meta will ensure unconnected pins are auto generated a reference property.
 				PropertyValuePin->PinType.bIsReference = true;
 
-				//Connect exec and then pins
+				//Connect exec and then pins				
+				CompilerContext.MovePinLinksToIntermediate(*CurrThen, *PropertyFuncCallNode->GetThenPin());
 				Schema->TryCreateConnection(CurrThen, PropertyFuncCallNode->GetExecPin());
-
 				CurrThen = PropertyFuncCallNode->GetThenPin();
 
 				Schema->TryCreateConnection(PropertyAccessContextPin, AccessContextSourcePin);
@@ -312,12 +311,6 @@ void UK2Node_DataChannelAccessContextOperation::ExpandNodeOptionalPins(class FKi
 		}
 	}
 
-	//Once we're done, connect our final then to the next exec pin from the original connection.
-	if (NextExecPin)
-	{
-		Schema->TryCreateConnection(CurrThen, NextExecPin);
-	}
-
 	//Remove our optional property pins from the main function call.
 	for(UEdGraphPin* SrcPin : ToRemove)
 	{
@@ -502,4 +495,66 @@ FText UK2Node_DataChannelAccessContext_SetMembers::GetTooltipTextForContext(UScr
 	return FText::Format(LOCTEXT("UK2Node_DataChannelAccessContext_SetMembers_TooltipFmt", "Sets the internal members of an Access Contecxt Instance with an internal type of {0}. "), Struct ? Struct->GetDisplayNameText() : FText::GetEmpty());
 }
 
+//////////////////////////////////////////////////////////////////////////
+
+UK2Node_DataChannelAccessContext_Prepare::UK2Node_DataChannelAccessContext_Prepare(const FObjectInitializer& ObjectInitializer)
+	: Super(ObjectInitializer)
+{
+	FunctionReference.SetExternalMember(GET_FUNCTION_NAME_CHECKED(UNiagaraDataChannelLibrary, PrepareAccessContextFromNDCRef), UNiagaraDataChannelLibrary::StaticClass());
+}
+
+void UK2Node_DataChannelAccessContext_Prepare::AllocateDefaultPins()
+{
+	Super::AllocateDefaultPins();
+	if (ContextStruct)
+	{
+		FStructOnScope StructOnScope(ContextStruct);
+		FK2Node_SetMembersInNDCAccessContext_OptionalPinManager OptionalPinManager(StructOnScope.GetStructMemory(), GetBlueprint());
+		OptionalPinManager.RebuildPropertyList(ShowPinForProperties, ContextStruct);
+		OptionalPinManager.CreateVisiblePins(ShowPinForProperties, ContextStruct, EGPD_Input, this);
+	}
+}
+
+void UK2Node_DataChannelAccessContext_Prepare::ExpandNode(class FKismetCompilerContext& CompilerContext, UEdGraph* SourceGraph)
+{
+	Super::ExpandNode(CompilerContext, SourceGraph);
+
+	if (!ContextStruct)
+	{
+		CompilerContext.MessageLog.Error(*LOCTEXT("NoValidAccessContextStruct", "Access Context Struct is null! - @@").ToString(), this);
+		return;
+	}
+
+	UEdGraphPin* RefInputPin = FindPinChecked(TEXT("NDCRef"));
+	UEdGraphPin* ResultPin = FindPinChecked(TEXT("ReturnValue"));
+
+	//First get a usable access context from the input ref.
+	UK2Node_CallFunction* GetAccessContextFuncCallNode = CompilerContext.SpawnIntermediateNode<UK2Node_CallFunction>(this, SourceGraph);
+	GetAccessContextFuncCallNode->SetFromFunction(UNiagaraDataChannelLibrary::StaticClass()->FindFunctionByName(TEXT("GetUsableAccessContextFromNDCRef")));
+	GetAccessContextFuncCallNode->AllocateDefaultPins();
+	CompilerContext.MovePinLinksToIntermediate(*GetExecPin(), *GetAccessContextFuncCallNode->GetExecPin());
+	CompilerContext.MovePinLinksToIntermediate(*GetThenPin(), *GetAccessContextFuncCallNode->GetThenPin());
+
+	UEdGraphPin* GetAccessContextFuncCallNode_RefInputPin = GetAccessContextFuncCallNode->FindPinChecked(TEXT("NDCRef"));
+	UEdGraphPin* GetAccessContextFuncCallNode_ResultPin = GetAccessContextFuncCallNode->FindPinChecked(TEXT("ReturnValue"));
+
+	CompilerContext.MovePinLinksToIntermediate(*RefInputPin, *GetAccessContextFuncCallNode_RefInputPin);
+	CompilerContext.MovePinLinksToIntermediate(*ResultPin, *GetAccessContextFuncCallNode_ResultPin);
+
+	ExpandNodeOptionalPins(CompilerContext, SourceGraph, GetAccessContextFuncCallNode_ResultPin, GetAccessContextFuncCallNode->GetThenPin(), GET_FUNCTION_NAME_CHECKED(UNiagaraDataChannelLibrary, SetSinglePropertyInNDCAccessContextInstance));
+
+	//This node is fully expanded way so clear out any remaining links
+	BreakAllNodeLinks();
+}
+
+FText UK2Node_DataChannelAccessContext_Prepare::GetNodeTitleForContext(UScriptStruct* Struct)const
+{
+	return FText::Format(LOCTEXT("UK2Node_DataChannelAccessContext_Prepare_TitleFmt", "Prepare ({0}) Instance"), Struct ? Struct->GetDisplayNameText() : FText::GetEmpty());
+}
+
+FText UK2Node_DataChannelAccessContext_Prepare::GetTooltipTextForContext(UScriptStruct* Struct)const
+{
+	return FText::Format(LOCTEXT("UK2Node_DataChannelAccessContext_Prepare_TooltipFmt", "Prepares an NDC Access Contecxt Instance with an internal type of {0}. Allows setting of its members before use. By default exposes transient inputs that will typically be set before each NDC access. "), Struct ? Struct->GetDisplayNameText() : FText::GetEmpty());
+}
+
 #undef LOCTEXT_NAMESPACE
