- Lazy instantiate parts of the baker to reduce memory usage
- Add a call to tick the baker world as paused, this ensure the world manager is called to process clean up functions

diff --git a/Engine/Plugins/FX/Niagara/Source/NiagaraEditor/Private/NiagaraBakerRenderer.cpp b/Engine/Plugins/FX/Niagara/Source/NiagaraEditor/Private/NiagaraBakerRenderer.cpp
--- a/Engine/Plugins/FX/Niagara/Source/NiagaraEditor/Private/NiagaraBakerRenderer.cpp
+++ b/Engine/Plugins/FX/Niagara/Source/NiagaraEditor/Private/NiagaraBakerRenderer.cpp
@@ -327,16 +327,6 @@ void FNiagaraBakerOutputBindingHelper::GetParticleAttributeBindings(TArray<FNiag
 FNiagaraBakerRenderer::FNiagaraBakerRenderer(UNiagaraSystem* InNiagaraSystem)
 	: NiagaraSystem(InNiagaraSystem)
 {
-	using namespace NiagaraBakerRendererPrivate;
-
-	SceneCaptureComponent = NewObject<USceneCaptureComponent2D>(GetTransientPackage(), NAME_None, RF_Transient);
-	SceneCaptureComponent->bTickInEditor = false;
-	SceneCaptureComponent->SetComponentTickEnabled(false);
-	SceneCaptureComponent->SetVisibility(true);
-	SceneCaptureComponent->bCaptureEveryFrame = false;
-	SceneCaptureComponent->bCaptureOnMovement = false;
-
-	CreatePreviewScene(NiagaraSystem, PreviewComponent, AdvancedPreviewScene);
 }
 
 FNiagaraBakerRenderer::~FNiagaraBakerRenderer()
@@ -351,7 +341,7 @@ FNiagaraBakerRenderer::~FNiagaraBakerRenderer()
 void FNiagaraBakerRenderer::SetAbsoluteTime(float AbsoluteTime, bool bShouldTickComponent)
 {
 	UNiagaraBakerSettings* BakerSettings = GetBakerSettings();
-	if ( !ensure(BakerSettings) )
+	if ( !ensure(BakerSettings) || !PreviewComponent )
 	{
 		return;
 	}
@@ -372,6 +362,7 @@ void FNiagaraBakerRenderer::SetAbsoluteTime(float AbsoluteTime, bool bShouldTick
 		World->RealTimeSeconds		= AbsoluteTime;
 		World->DeltaRealTimeSeconds	= BakerSettings->GetSeekDelta();
 		World->DeltaTimeSeconds		= BakerSettings->GetSeekDelta();
+		World->Tick(ELevelTick::LEVELTICK_PauseTick, 0.0f);
 	}
 
 	if (bShouldTickComponent)
@@ -394,7 +385,7 @@ void FNiagaraBakerRenderer::SetAbsoluteTime(float AbsoluteTime, bool bShouldTick
 
 void FNiagaraBakerRenderer::RenderSceneCapture(UTextureRenderTarget2D* RenderTarget, ESceneCaptureSource CaptureSource) const
 {
-	RenderSceneCapture(RenderTarget, PreviewComponent, CaptureSource);
+	RenderSceneCapture(RenderTarget, GetPreviewComponent(), CaptureSource);
 }
 
 void FNiagaraBakerRenderer::RenderSceneCapture(UTextureRenderTarget2D* RenderTarget, UPrimitiveComponent* BakedDataComponent, ESceneCaptureSource CaptureSource) const
@@ -411,6 +402,17 @@ void FNiagaraBakerRenderer::RenderSceneCapture(UTextureRenderTarget2D* RenderTar
 	FCanvas Canvas(RenderTarget->GameThread_GetRenderTargetResource(), nullptr, FGameTime::CreateUndilated(WorldTime, FApp::GetDeltaTime()), GetFeatureLevel());
 	Canvas.Clear(FLinearColor::Black);
 
+	// Lazy create the component to avoid it always being present in every Niagara editor
+	if (SceneCaptureComponent == nullptr)
+	{
+		SceneCaptureComponent = NewObject<USceneCaptureComponent2D>(GetTransientPackage(), NAME_None, RF_Transient);
+		SceneCaptureComponent->bTickInEditor = false;
+		SceneCaptureComponent->SetComponentTickEnabled(false);
+		SceneCaptureComponent->SetVisibility(true);
+		SceneCaptureComponent->bCaptureEveryFrame = false;
+		SceneCaptureComponent->bCaptureOnMovement = false;
+	}
+
 	SceneCaptureComponent->RegisterComponentWithWorld(World);
 	SceneCaptureComponent->TextureTarget = RenderTarget;
 	SceneCaptureComponent->CaptureSource = CaptureSource;
@@ -516,7 +518,7 @@ void FNiagaraBakerRenderer::RenderBufferVisualization(UTextureRenderTarget2D* Re
 	if (BakerSettings->bRenderComponentOnly)
 	{
 		ViewInitOptions.ShowOnlyPrimitives.Emplace();
-		ViewInitOptions.ShowOnlyPrimitives->Add(PreviewComponent->GetPrimitiveSceneId());
+		ViewInitOptions.ShowOnlyPrimitives->Add(GetPreviewComponent()->GetPrimitiveSceneId());
 	}
 	
 	FSceneView* NewView = new FSceneView(ViewInitOptions);
@@ -550,7 +552,7 @@ void FNiagaraBakerRenderer::RenderDataInterface(UTextureRenderTarget2D* RenderTa
 	const FName VariableName = FName(SourceString.RightChop(DotIndex + 1));
 
 	// Find data interface
-	FNiagaraSystemInstanceControllerPtr SystemInstanceController = PreviewComponent->GetSystemInstanceController();
+	FNiagaraSystemInstanceControllerPtr SystemInstanceController = GetPreviewComponent()->GetSystemInstanceController();
 	if ( SystemInstanceController.IsValid() == false )
 	{
 		return;
@@ -602,7 +604,7 @@ void FNiagaraBakerRenderer::RenderParticleAttribute(UTextureRenderTarget2D* Rend
 	const FString EmitterName = SourceString.LeftChop(SourceString.Len() - DotIndex);
 	const FName AttributeName = FName(SourceString.RightChop(DotIndex + 1));
 	
-	FNiagaraSystemInstanceControllerPtr SystemInstanceController = PreviewComponent->GetSystemInstanceController();
+	FNiagaraSystemInstanceControllerPtr SystemInstanceController = GetPreviewComponent()->GetSystemInstanceController();
 	if (!ensure(SystemInstanceController.IsValid()))
 	{
 		return;
@@ -805,22 +807,33 @@ void FNiagaraBakerRenderer::RenderStaticMesh(UTextureRenderTarget2D* RenderTarge
 
 UWorld* FNiagaraBakerRenderer::GetWorld() const
 {
-	return PreviewComponent->GetWorld();
+	return GetPreviewComponent()->GetWorld();
 }
 
 float FNiagaraBakerRenderer::GetWorldTime() const
 {
-	return PreviewComponent->GetDesiredAge();
+	return GetPreviewComponent()->GetDesiredAge();
 }
 
 ERHIFeatureLevel::Type FNiagaraBakerRenderer::GetFeatureLevel() const
 {
-	return PreviewComponent->GetWorld()->Scene->GetFeatureLevel();
+	return GetPreviewComponent()->GetWorld()->Scene->GetFeatureLevel();
+}
+
+UNiagaraComponent* FNiagaraBakerRenderer::GetPreviewComponent() const
+{
+	if (PreviewComponent == nullptr)
+	{
+		using namespace NiagaraBakerRendererPrivate;
+		CreatePreviewScene(NiagaraSystem, PreviewComponent, AdvancedPreviewScene);
+	}
+
+	return PreviewComponent;
 }
 
 UNiagaraSystem* FNiagaraBakerRenderer::GetNiagaraSystem() const
 {
-	return PreviewComponent->GetAsset();
+	return NiagaraSystem;
 }
 
 void FNiagaraBakerRenderer::AddReferencedObjects(FReferenceCollector& Collector)
diff --git a/Engine/Plugins/FX/Niagara/Source/NiagaraEditor/Private/NiagaraBakerRenderer.h b/Engine/Plugins/FX/Niagara/Source/NiagaraEditor/Private/NiagaraBakerRenderer.h
--- a/Engine/Plugins/FX/Niagara/Source/NiagaraEditor/Private/NiagaraBakerRenderer.h
+++ b/Engine/Plugins/FX/Niagara/Source/NiagaraEditor/Private/NiagaraBakerRenderer.h
@@ -119,7 +119,7 @@ public:
 	UWorld* GetWorld() const;
 	float GetWorldTime() const;
 	ERHIFeatureLevel::Type GetFeatureLevel() const;
-	UNiagaraComponent* GetPreviewComponent() const { return PreviewComponent; }
+	UNiagaraComponent* GetPreviewComponent() const;
 	UNiagaraSystem* GetNiagaraSystem() const;
 	UNiagaraBakerSettings* GetBakerSettings() const { return NiagaraSystem->GetBakerSettings(); }
 	const UNiagaraBakerSettings* GetBakerGeneratedSettings() const { return NiagaraSystem->GetBakerGeneratedSettings(); }
@@ -139,9 +139,9 @@ public:
 
 private:
 	TObjectPtr<UNiagaraSystem> NiagaraSystem = nullptr;
-	TObjectPtr<UNiagaraComponent> PreviewComponent = nullptr;
-	TSharedPtr<FAdvancedPreviewScene> AdvancedPreviewScene;
-	TObjectPtr<USceneCaptureComponent2D> SceneCaptureComponent = nullptr;
+	mutable TObjectPtr<UNiagaraComponent> PreviewComponent = nullptr;
+	mutable TSharedPtr<FAdvancedPreviewScene> AdvancedPreviewScene;
+	mutable TObjectPtr<USceneCaptureComponent2D> SceneCaptureComponent = nullptr;
 
 	mutable TObjectPtr<UNiagaraComponent> SimCachePreviewComponent = nullptr;
 	mutable TSharedPtr<FAdvancedPreviewScene> SimCacheAdvancedPreviewScene;
diff --git a/Engine/Plugins/FX/Niagara/Source/NiagaraEditor/Private/ViewModels/NiagaraBakerViewModel.cpp b/Engine/Plugins/FX/Niagara/Source/NiagaraEditor/Private/ViewModels/NiagaraBakerViewModel.cpp
--- a/Engine/Plugins/FX/Niagara/Source/NiagaraEditor/Private/ViewModels/NiagaraBakerViewModel.cpp
+++ b/Engine/Plugins/FX/Niagara/Source/NiagaraEditor/Private/ViewModels/NiagaraBakerViewModel.cpp
@@ -40,9 +40,6 @@ void FNiagaraBakerViewModel::Initialize(TWeakPtr<FNiagaraSystemViewModel> InWeak
 
 	UNiagaraSystem* NiagaraSystem = WeakSystemViewModel.Pin()->GetPreviewComponent()->GetAsset();
 	BakerRenderer.Reset(new FNiagaraBakerRenderer(NiagaraSystem));
-
-	Widget = SNew(SNiagaraBakerWidget)
-		.WeakViewModel(this->AsShared());
 }
 
 void FNiagaraBakerViewModel::SetDisplayTimeFromNormalized(float NormalizeTime)
@@ -64,6 +61,12 @@ void FNiagaraBakerViewModel::SetDisplayTimeFromNormalized(float NormalizeTime)
 
 TSharedPtr<class SWidget> FNiagaraBakerViewModel::GetWidget()
 {
+	if (!Widget.IsValid())
+	{
+		Widget = SNew(SNiagaraBakerWidget)
+			.WeakViewModel(this->AsShared());
+	}
+
 	return Widget;
 }
 
