Fixed crash when opening old broken asset.
Deleted old broken asset.

diff --git a/Engine/Plugins/FX/Niagara/Source/NiagaraBlueprintNodes/Private/K2Node_ReadDataChannel.cpp b/Engine/Plugins/FX/Niagara/Source/NiagaraBlueprintNodes/Private/K2Node_ReadDataChannel.cpp
--- a/Engine/Plugins/FX/Niagara/Source/NiagaraBlueprintNodes/Private/K2Node_ReadDataChannel.cpp
+++ b/Engine/Plugins/FX/Niagara/Source/NiagaraBlueprintNodes/Private/K2Node_ReadDataChannel.cpp
@@ -50,6 +50,11 @@ void UK2Node_ReadDataChannel::AddNDCDerivedPins()
 void UK2Node_ReadDataChannel::ExpandNode(FKismetCompilerContext& CompilerContext, UEdGraph* SourceGraph)
 {
 	Super::ExpandNode(CompilerContext, SourceGraph);
+	
+	if (!HasValidDataChannel() && !SupportsDynamicDataChannel())
+	{
+		return;
+	}
 
 	const UEdGraphSchema_K2* Schema = CompilerContext.GetSchema();
 	UEdGraphPin* SuccessPin = FindPinChecked(FName("Success"), EGPD_Output);
diff --git a/Engine/Plugins/FX/Niagara/Source/NiagaraBlueprintNodes/Private/K2Node_WriteDataChannel.cpp b/Engine/Plugins/FX/Niagara/Source/NiagaraBlueprintNodes/Private/K2Node_WriteDataChannel.cpp
--- a/Engine/Plugins/FX/Niagara/Source/NiagaraBlueprintNodes/Private/K2Node_WriteDataChannel.cpp
+++ b/Engine/Plugins/FX/Niagara/Source/NiagaraBlueprintNodes/Private/K2Node_WriteDataChannel.cpp
@@ -46,6 +46,11 @@ void UK2Node_WriteDataChannel::AddNDCDerivedPins()
 void UK2Node_WriteDataChannel::ExpandNode(FKismetCompilerContext& CompilerContext, UEdGraph* SourceGraph)
 {
 	Super::ExpandNode(CompilerContext, SourceGraph);
+	
+	if (!HasValidDataChannel() && !SupportsDynamicDataChannel())
+	{
+		return;
+	}
 
 	// create function call node to init the writer object 
 	UK2Node_CallFunction* CreateWriterNode = CompilerContext.SpawnIntermediateNode<UK2Node_CallFunction>(this, SourceGraph);
