- Disable Source Mode with LWE to reduce confusion
- We don't have a particle stack so there is no way to modify the data without binding to system / user parameters

diff --git a/Engine/Plugins/FX/Niagara/Source/Niagara/Private/NiagaraDecalRendererProperties.cpp b/Engine/Plugins/FX/Niagara/Source/Niagara/Private/NiagaraDecalRendererProperties.cpp
--- a/Engine/Plugins/FX/Niagara/Source/Niagara/Private/NiagaraDecalRendererProperties.cpp
+++ b/Engine/Plugins/FX/Niagara/Source/Niagara/Private/NiagaraDecalRendererProperties.cpp
@@ -75,6 +75,7 @@ void UNiagaraDecalRendererProperties::PostLoad()
 	Super::PostLoad();
 	
 #if WITH_EDITORONLY_DATA
+	SourceMode = ValidateSourceMode(SourceMode);
 	ChangeToPositionBinding(PositionBinding);
 #endif
 	PostLoadBindings(ENiagaraRendererSourceDataMode::Particles);
diff --git a/Engine/Plugins/FX/Niagara/Source/Niagara/Private/NiagaraLightRendererProperties.cpp b/Engine/Plugins/FX/Niagara/Source/Niagara/Private/NiagaraLightRendererProperties.cpp
--- a/Engine/Plugins/FX/Niagara/Source/Niagara/Private/NiagaraLightRendererProperties.cpp
+++ b/Engine/Plugins/FX/Niagara/Source/Niagara/Private/NiagaraLightRendererProperties.cpp
@@ -64,6 +64,7 @@ void UNiagaraLightRendererProperties::PostLoad()
 	Super::PostLoad();
 	
 #if WITH_EDITORONLY_DATA
+	SourceMode = ValidateSourceMode(SourceMode);
 	ChangeToPositionBinding(PositionBinding);
 #endif
 	PostLoadBindings(ENiagaraRendererSourceDataMode::Particles);
diff --git a/Engine/Plugins/FX/Niagara/Source/Niagara/Private/NiagaraMeshRendererProperties.cpp b/Engine/Plugins/FX/Niagara/Source/Niagara/Private/NiagaraMeshRendererProperties.cpp
--- a/Engine/Plugins/FX/Niagara/Source/Niagara/Private/NiagaraMeshRendererProperties.cpp
+++ b/Engine/Plugins/FX/Niagara/Source/Niagara/Private/NiagaraMeshRendererProperties.cpp
@@ -1034,6 +1034,8 @@ void UNiagaraMeshRendererProperties::PostLoad()
 	Super::PostLoad();
 
 #if WITH_EDITORONLY_DATA
+	SourceMode = ValidateSourceMode(SourceMode);
+
 	if (Meshes.Num() == 1 && Meshes[0].Mesh == nullptr && ParticleMesh_DEPRECATED != nullptr)
 	{
 		// Likely predates the mesh array ... just add ParticleMesh to the list of Meshes
diff --git a/Engine/Plugins/FX/Niagara/Source/Niagara/Private/NiagaraRendererProperties.cpp b/Engine/Plugins/FX/Niagara/Source/Niagara/Private/NiagaraRendererProperties.cpp
--- a/Engine/Plugins/FX/Niagara/Source/Niagara/Private/NiagaraRendererProperties.cpp
+++ b/Engine/Plugins/FX/Niagara/Source/Niagara/Private/NiagaraRendererProperties.cpp
@@ -959,6 +959,17 @@ FText UNiagaraRendererProperties::GetWidgetDisplayName() const
 	return GetClass()->GetDisplayNameText();
 }
 
+bool UNiagaraRendererProperties::SupportsEmitterMode() const
+{
+	FVersionedNiagaraEmitter VersionedEmitter = GetOuterEmitter();
+	return VersionedEmitter.Emitter != nullptr;
+}
+
+ENiagaraRendererSourceDataMode UNiagaraRendererProperties::ValidateSourceMode(ENiagaraRendererSourceDataMode InSourceMode) const
+{
+	return SupportsEmitterMode() ? InSourceMode : ENiagaraRendererSourceDataMode::Particles;
+}
+
 #if WITH_NIAGARA_RENDERER_DEBUGDRAW
 bool UNiagaraRendererProperties::SupportsDebugDraw() const
 {
@@ -1257,6 +1268,31 @@ void UNiagaraRendererProperties::PostLoad()
 #endif
 }
 
+#if WITH_EDITOR
+bool UNiagaraRendererProperties::CanEditChange(const FProperty* InProperty) const
+{
+	if (!Super::CanEditChange(InProperty))
+	{
+		return false;
+	}
+
+	if (InProperty != nullptr)
+	{
+		static const FName NAME_SourceMode("SourceMode");
+		if (InProperty->GetFName() == NAME_SourceMode)
+		{
+			const FEnumProperty* EnumProperty = CastField<const FEnumProperty>(InProperty);
+			if (EnumProperty && EnumProperty->GetEnum() == StaticEnum<ENiagaraRendererSourceDataMode>())
+			{
+				return SupportsEmitterMode();
+			}
+		}
+	}
+
+	return true;
+}
+#endif //WITH_EDITOR
+
 #if WITH_EDITORONLY_DATA
 
 void UNiagaraRendererProperties::PostEditChangeProperty(struct FPropertyChangedEvent& PropertyChangedEvent)
diff --git a/Engine/Plugins/FX/Niagara/Source/Niagara/Private/NiagaraSpriteRendererProperties.cpp b/Engine/Plugins/FX/Niagara/Source/Niagara/Private/NiagaraSpriteRendererProperties.cpp
--- a/Engine/Plugins/FX/Niagara/Source/Niagara/Private/NiagaraSpriteRendererProperties.cpp
+++ b/Engine/Plugins/FX/Niagara/Source/Niagara/Private/NiagaraSpriteRendererProperties.cpp
@@ -160,6 +160,8 @@ void UNiagaraSpriteRendererProperties::PostLoad()
 	}
 
 #if WITH_EDITORONLY_DATA
+	SourceMode = ValidateSourceMode(SourceMode);
+
 	if (MaterialUserParamBinding.Parameter.GetType().GetClass() != UMaterialInterface::StaticClass())
 	{
 		FNiagaraTypeDefinition MaterialDef(UMaterialInterface::StaticClass());
diff --git a/Engine/Plugins/FX/Niagara/Source/Niagara/Public/NiagaraRendererProperties.h b/Engine/Plugins/FX/Niagara/Source/Niagara/Public/NiagaraRendererProperties.h
--- a/Engine/Plugins/FX/Niagara/Source/Niagara/Public/NiagaraRendererProperties.h
+++ b/Engine/Plugins/FX/Niagara/Source/Niagara/Public/NiagaraRendererProperties.h
@@ -310,6 +310,9 @@ public:
 	//UObject Interface Begin
 	NIAGARA_API virtual void PostInitProperties() override;
 	NIAGARA_API virtual void PostLoad() override;
+#if WITH_EDITOR
+	NIAGARA_API virtual bool CanEditChange(const FProperty* InProperty) const override;
+#endif //WITH_EDITOR
 #if WITH_EDITORONLY_DATA
 	NIAGARA_API virtual void PostEditChangeProperty(struct FPropertyChangedEvent& PropertyChangedEvent) override;
 #endif
@@ -406,6 +409,9 @@ public:
 
 	// The text to display in the niagara stack widget under the renderer section
 	NIAGARA_API virtual FText GetWidgetDisplayName() const;
+
+	NIAGARA_API bool SupportsEmitterMode() const;
+	NIAGARA_API ENiagaraRendererSourceDataMode ValidateSourceMode(ENiagaraRendererSourceDataMode InSourceMode) const;
 #endif // WITH_EDITORONLY_DATA
 
 #if WITH_NIAGARA_RENDERER_DEBUGDRAW
