- Fix for duplicating an emitter with RunUpdateScript set for interpolation mode changing post duplication to NoInterpolation

#tests Create an emitter, set it to RunUpdateScript, duplicate and observe it doesn't change

diff --git a/Engine/Plugins/FX/Niagara/Source/Niagara/Private/NiagaraEmitter.cpp b/Engine/Plugins/FX/Niagara/Source/Niagara/Private/NiagaraEmitter.cpp
--- a/Engine/Plugins/FX/Niagara/Source/Niagara/Private/NiagaraEmitter.cpp
+++ b/Engine/Plugins/FX/Niagara/Source/Niagara/Private/NiagaraEmitter.cpp
@@ -854,8 +854,9 @@ void FVersionedNiagaraEmitterData::PostLoad(UNiagaraEmitter& Emitter, int32 Niag
 	//Temporarily disabling interpolated spawn if the script type and flag don't match.
 	if (SpawnScriptProps.Script)
 	{
-		bool bActualInterpolatedSpawning = SpawnScriptProps.Script->IsInterpolatedParticleSpawnScript();
-		if (UsesInterpolatedSpawning() != bActualInterpolatedSpawning)
+		const bool bActualInterpolatedSpawning = SpawnScriptProps.Script->IsInterpolatedParticleSpawnScript();
+		const bool bUsesInterpolatedSpawning = InterpolatedSpawnMode != ENiagaraInterpolatedSpawnMode::NoInterpolation;
+		if (bActualInterpolatedSpawning != bUsesInterpolatedSpawning)
 		{
 			InterpolatedSpawnMode = ENiagaraInterpolatedSpawnMode::NoInterpolation;
 			if (bActualInterpolatedSpawning)
