- Validate the interpolation parameters are correctly filled in

#rnx

diff --git a/Engine/Plugins/FX/Niagara/Source/Niagara/Private/NiagaraGPUSystemTick.cpp b/Engine/Plugins/FX/Niagara/Source/Niagara/Private/NiagaraGPUSystemTick.cpp
--- a/Engine/Plugins/FX/Niagara/Source/Niagara/Private/NiagaraGPUSystemTick.cpp
+++ b/Engine/Plugins/FX/Niagara/Source/Niagara/Private/NiagaraGPUSystemTick.cpp
@@ -156,6 +156,9 @@ void FNiagaraGPUSystemTick::Init(FNiagaraSystemInstance* InSystemInstance)
 			InstanceData->Context = GPUContext;
 			check(GPUContext->MainDataSet);
 
+			// Validate that we generated interpolation parameters
+			ensureMsgf(GPUContext->HasInterpolationParameters ? IncludeInterpolationParameters : true, TEXT("Mismatch between required interpolation parameters and calculated."));
+
 			InstanceData->SpawnInfo = GPUContext->GpuSpawnInfo_GT;
 
 			// Consume pending reset
