- Fix crash when sampling from a mutable skeletal mesh
- Mutable doesn't build WholeMeshBuiltData therefore it's empty

#rnx

diff --git a/Engine/Plugins/FX/Niagara/Source/Niagara/Classes/NDISkeletalMeshCommon.h b/Engine/Plugins/FX/Niagara/Source/Niagara/Classes/NDISkeletalMeshCommon.h
--- a/Engine/Plugins/FX/Niagara/Source/Niagara/Classes/NDISkeletalMeshCommon.h
+++ b/Engine/Plugins/FX/Niagara/Source/Niagara/Classes/NDISkeletalMeshCommon.h
@@ -583,10 +583,12 @@ struct TAreaWeightingModeBinder
 				const int32 LODIndex = InstData->GetLODIndex();
 
 				const FSkeletalMeshSamplingInfo& SamplingInfo = InstData->SkeletalMesh->GetSamplingInfo();
-				const FSkeletalMeshSamplingLODBuiltData& WholeMeshBuiltData = SamplingInfo.GetWholeMeshLODBuiltData(LODIndex);
+				const FSkeletalMeshSamplingBuiltData& SamplingBuildData = SamplingInfo.GetBuiltData();
 
-				bAreaWeighting = InstData->SkeletalMesh->GetLODInfo(LODIndex)->bSupportUniformlyDistributedSampling
-					&& (WholeMeshBuiltData.AreaWeightedTriangleSampler.GetNumEntries() > 0);
+				bAreaWeighting =
+					SamplingBuildData.WholeMeshBuiltData.IsValidIndex(LODIndex) &&
+					(SamplingBuildData.WholeMeshBuiltData[LODIndex].AreaWeightedTriangleSampler.GetNumEntries() > 0) &&
+					InstData->SkeletalMesh->GetLODInfo(LODIndex)->bSupportUniformlyDistributedSampling;
 			}
 		}
 
