- Arrays for LW distributions
- Large refactor to make reading a distribution a function call vs a structure
- When sampling a value needed for motion vectors pass in multiple time parameters and it will return both values to ensure random is respected
- Some limited space was added to pack indexing into the data, this can be controlled using meta data on the distribution property
- GPU storage moved to ByteAddressBuffer for LoadX operations which gave a small speed boost
- Perf on a large number of particles was marginally better vs before

diff --git a/Engine/Plugins/FX/Niagara/Shaders/Private/Stateless/Modules/NiagaraStatelessModule_CameraOffset.ush b/Engine/Plugins/FX/Niagara/Shaders/Private/Stateless/Modules/NiagaraStatelessModule_CameraOffset.ush
--- a/Engine/Plugins/FX/Niagara/Shaders/Private/Stateless/Modules/NiagaraStatelessModule_CameraOffset.ush
+++ b/Engine/Plugins/FX/Niagara/Shaders/Private/Stateless/Modules/NiagaraStatelessModule_CameraOffset.ush
@@ -4,11 +4,10 @@ FNiagaraStatelessBuiltDistributionType CameraOffset_Distribution;
 
 void CameraOffset_Simulate(inout FStatelessParticle Particle)
 {
-	FStatelessDistributionSamplerFloat OffsetSampler;
-	OffsetSampler.Init(0, CameraOffset_Distribution);
-	if (OffsetSampler.IsValid())
+	if (FNiagaraStatelessBuiltDistribution::IsValid(CameraOffset_Distribution))
 	{
-		Particle.CameraOffset			= OffsetSampler.GetValue(Particle.NormalizedAge);
-		Particle.PreviousCameraOffset	= OffsetSampler.GetValue(Particle.PreviousNormalizedAge);
+		const float2 CameraOffsets = SampleDistributionValuesFloat(CameraOffset_Distribution, 0, Particle.NormalizedAge, Particle.PreviousNormalizedAge);
+		Particle.CameraOffset			= CameraOffsets[0];
+		Particle.PreviousCameraOffset	= CameraOffsets[1];
 	}
 }
diff --git a/Engine/Plugins/FX/Niagara/Shaders/Private/Stateless/Modules/NiagaraStatelessModule_DynamicMaterialParameters.ush b/Engine/Plugins/FX/Niagara/Shaders/Private/Stateless/Modules/NiagaraStatelessModule_DynamicMaterialParameters.ush
--- a/Engine/Plugins/FX/Niagara/Shaders/Private/Stateless/Modules/NiagaraStatelessModule_DynamicMaterialParameters.ush
+++ b/Engine/Plugins/FX/Niagara/Shaders/Private/Stateless/Modules/NiagaraStatelessModule_DynamicMaterialParameters.ush
@@ -27,57 +27,33 @@ void DynamicMaterialParameters_Simulate(inout FStatelessParticle Particle)
 
 	if ((DynamicMaterialParameters_ChannelMask & 0x000f) != 0)
 	{
-		FStatelessDistributionSamplerFloat ParameterSampler[4];
-		ParameterSampler[0].Init(0, DynamicMaterialParameters_Parameter0X);
-		ParameterSampler[1].Init(1, DynamicMaterialParameters_Parameter0Y);
-		ParameterSampler[2].Init(2, DynamicMaterialParameters_Parameter0Z);
-		ParameterSampler[3].Init(3, DynamicMaterialParameters_Parameter0W);
-		
-		Particle.DynamicMaterialParameter.x = (DynamicMaterialParameters_ChannelMask & 0x0001) != 0 ? ParameterSampler[0].GetValue(Particle.NormalizedAge) : 0.0f;
-		Particle.DynamicMaterialParameter.y = (DynamicMaterialParameters_ChannelMask & 0x0002) != 0 ? ParameterSampler[1].GetValue(Particle.NormalizedAge) : 0.0f;
-		Particle.DynamicMaterialParameter.z = (DynamicMaterialParameters_ChannelMask & 0x0004) != 0 ? ParameterSampler[2].GetValue(Particle.NormalizedAge) : 0.0f;
-		Particle.DynamicMaterialParameter.w = (DynamicMaterialParameters_ChannelMask & 0x0008) != 0 ? ParameterSampler[3].GetValue(Particle.NormalizedAge) : 0.0f;
+		Particle.DynamicMaterialParameter.x = SampleDistributionValueFloat(DynamicMaterialParameters_Parameter0X, 0, Particle.NormalizedAge);
+		Particle.DynamicMaterialParameter.y = SampleDistributionValueFloat(DynamicMaterialParameters_Parameter0Y, 1, Particle.NormalizedAge);
+		Particle.DynamicMaterialParameter.z = SampleDistributionValueFloat(DynamicMaterialParameters_Parameter0Z, 2, Particle.NormalizedAge);
+		Particle.DynamicMaterialParameter.w = SampleDistributionValueFloat(DynamicMaterialParameters_Parameter0W, 3, Particle.NormalizedAge);
 	}
 
 	if ((DynamicMaterialParameters_ChannelMask & 0x00f0) != 0)
 	{
-		FStatelessDistributionSamplerFloat ParameterSampler[4];
-		ParameterSampler[0].Init(4, DynamicMaterialParameters_Parameter1X);
-		ParameterSampler[1].Init(5, DynamicMaterialParameters_Parameter1Y);
-		ParameterSampler[2].Init(6, DynamicMaterialParameters_Parameter1Z);
-		ParameterSampler[3].Init(7, DynamicMaterialParameters_Parameter1W);
-		
-		Particle.DynamicMaterialParameter1.x = (DynamicMaterialParameters_ChannelMask & 0x0010) != 0 ? ParameterSampler[0].GetValue(Particle.NormalizedAge) : 0.0f;
-		Particle.DynamicMaterialParameter1.y = (DynamicMaterialParameters_ChannelMask & 0x0020) != 0 ? ParameterSampler[1].GetValue(Particle.NormalizedAge) : 0.0f;
-		Particle.DynamicMaterialParameter1.z = (DynamicMaterialParameters_ChannelMask & 0x0040) != 0 ? ParameterSampler[2].GetValue(Particle.NormalizedAge) : 0.0f;
-		Particle.DynamicMaterialParameter1.w = (DynamicMaterialParameters_ChannelMask & 0x0080) != 0 ? ParameterSampler[3].GetValue(Particle.NormalizedAge) : 0.0f;
+		Particle.DynamicMaterialParameter1.x = SampleDistributionValueFloat(DynamicMaterialParameters_Parameter1X, 4, Particle.NormalizedAge);
+		Particle.DynamicMaterialParameter1.y = SampleDistributionValueFloat(DynamicMaterialParameters_Parameter1Y, 5, Particle.NormalizedAge);
+		Particle.DynamicMaterialParameter1.z = SampleDistributionValueFloat(DynamicMaterialParameters_Parameter1Z, 6, Particle.NormalizedAge);
+		Particle.DynamicMaterialParameter1.w = SampleDistributionValueFloat(DynamicMaterialParameters_Parameter1W, 7, Particle.NormalizedAge);
 	}
 
 	if ((DynamicMaterialParameters_ChannelMask & 0x0f00) != 0)
 	{
-		FStatelessDistributionSamplerFloat ParameterSampler[4];
-		ParameterSampler[0].Init(8, DynamicMaterialParameters_Parameter2X);
-		ParameterSampler[1].Init(9, DynamicMaterialParameters_Parameter2Y);
-		ParameterSampler[2].Init(10, DynamicMaterialParameters_Parameter2Z);
-		ParameterSampler[3].Init(11, DynamicMaterialParameters_Parameter2W);
-		
-		Particle.DynamicMaterialParameter2.x = (DynamicMaterialParameters_ChannelMask & 0x0100) != 0 ? ParameterSampler[0].GetValue(Particle.NormalizedAge) : 0.0f;
-		Particle.DynamicMaterialParameter2.y = (DynamicMaterialParameters_ChannelMask & 0x0200) != 0 ? ParameterSampler[1].GetValue(Particle.NormalizedAge) : 0.0f;
-		Particle.DynamicMaterialParameter2.z = (DynamicMaterialParameters_ChannelMask & 0x0400) != 0 ? ParameterSampler[2].GetValue(Particle.NormalizedAge) : 0.0f;
-		Particle.DynamicMaterialParameter2.w = (DynamicMaterialParameters_ChannelMask & 0x0800) != 0 ? ParameterSampler[3].GetValue(Particle.NormalizedAge) : 0.0f;
+		Particle.DynamicMaterialParameter2.x = SampleDistributionValueFloat(DynamicMaterialParameters_Parameter2X,  8, Particle.NormalizedAge);
+		Particle.DynamicMaterialParameter2.y = SampleDistributionValueFloat(DynamicMaterialParameters_Parameter2Y,  9, Particle.NormalizedAge);
+		Particle.DynamicMaterialParameter2.z = SampleDistributionValueFloat(DynamicMaterialParameters_Parameter2Z, 10, Particle.NormalizedAge);
+		Particle.DynamicMaterialParameter2.w = SampleDistributionValueFloat(DynamicMaterialParameters_Parameter2W, 11, Particle.NormalizedAge);
 	}
 
 	if ((DynamicMaterialParameters_ChannelMask & 0xf000) != 0)
 	{
-		FStatelessDistributionSamplerFloat ParameterSampler[4];
-		ParameterSampler[0].Init(12, DynamicMaterialParameters_Parameter3X);
-		ParameterSampler[1].Init(13, DynamicMaterialParameters_Parameter3Y);
-		ParameterSampler[2].Init(14, DynamicMaterialParameters_Parameter3Z);
-		ParameterSampler[3].Init(15, DynamicMaterialParameters_Parameter3W);
-		
-		Particle.DynamicMaterialParameter3.x = (DynamicMaterialParameters_ChannelMask & 0x1000) != 0 ? ParameterSampler[0].GetValue(Particle.NormalizedAge) : 0.0f;
-		Particle.DynamicMaterialParameter3.y = (DynamicMaterialParameters_ChannelMask & 0x2000) != 0 ? ParameterSampler[1].GetValue(Particle.NormalizedAge) : 0.0f;
-		Particle.DynamicMaterialParameter3.z = (DynamicMaterialParameters_ChannelMask & 0x4000) != 0 ? ParameterSampler[2].GetValue(Particle.NormalizedAge) : 0.0f;
-		Particle.DynamicMaterialParameter3.w = (DynamicMaterialParameters_ChannelMask & 0x8000) != 0 ? ParameterSampler[3].GetValue(Particle.NormalizedAge) : 0.0f;
+		Particle.DynamicMaterialParameter3.x = SampleDistributionValueFloat(DynamicMaterialParameters_Parameter3X, 12, Particle.NormalizedAge);
+		Particle.DynamicMaterialParameter3.y = SampleDistributionValueFloat(DynamicMaterialParameters_Parameter3Y, 13, Particle.NormalizedAge);
+		Particle.DynamicMaterialParameter3.z = SampleDistributionValueFloat(DynamicMaterialParameters_Parameter3Z, 14, Particle.NormalizedAge);
+		Particle.DynamicMaterialParameter3.w = SampleDistributionValueFloat(DynamicMaterialParameters_Parameter3W, 15, Particle.NormalizedAge);
 	}
 }
diff --git a/Engine/Plugins/FX/Niagara/Shaders/Private/Stateless/Modules/NiagaraStatelessModule_InitializeParticle.ush b/Engine/Plugins/FX/Niagara/Shaders/Private/Stateless/Modules/NiagaraStatelessModule_InitializeParticle.ush
--- a/Engine/Plugins/FX/Niagara/Shaders/Private/Stateless/Modules/NiagaraStatelessModule_InitializeParticle.ush
+++ b/Engine/Plugins/FX/Niagara/Shaders/Private/Stateless/Modules/NiagaraStatelessModule_InitializeParticle.ush
@@ -23,10 +23,7 @@ void InitializeParticle_Simulate(inout FStatelessParticle Particle)
 	const bool bUniformSpriteSize = (InitializeParticle_ModuleFlags & EInitializeParticleModuleFlag_UniformSpriteSize) != 0;
 	const bool bUniformMeshScale  = (InitializeParticle_ModuleFlags & EInitializeParticleModuleFlag_UniformMeshScale) != 0;
 
-	FStatelessDistributionSamplerFloat3 InitialPosition;
-	InitialPosition.Init(0, InitializeParticle_InitialPosition);
-
-	Particle.Position		= InitialPosition.GetValue(0.0f);
+	Particle.Position		= SampleDistributionValueFloat3(InitializeParticle_InitialPosition, 0, 0.0f);
 	Particle.Color			= RandomScaleBiasFloat(1, InitializeParticle_ColorScale, InitializeParticle_ColorBias);
 	Particle.RibbonWidth	= RandomScaleBiasFloat(2, InitializeParticle_RibbonWidthScale, InitializeParticle_RibbonWidthBias);
 	Particle.SpriteSize		= RandomScaleBiasFloat(3, InitializeParticle_SpriteSizeScale, InitializeParticle_SpriteSizeBias, bUniformSpriteSize);
diff --git a/Engine/Plugins/FX/Niagara/Shaders/Private/Stateless/Modules/NiagaraStatelessModule_ScaleColor.ush b/Engine/Plugins/FX/Niagara/Shaders/Private/Stateless/Modules/NiagaraStatelessModule_ScaleColor.ush
--- a/Engine/Plugins/FX/Niagara/Shaders/Private/Stateless/Modules/NiagaraStatelessModule_ScaleColor.ush
+++ b/Engine/Plugins/FX/Niagara/Shaders/Private/Stateless/Modules/NiagaraStatelessModule_ScaleColor.ush
@@ -4,10 +4,8 @@ FNiagaraStatelessBuiltDistributionType	ScaleColor_Distribution;
 
 void ScaleColor_Simulate(inout FStatelessParticle Particle)
 {
-	FStatelessDistributionSamplerFloat4 ScaleDistribution;
-	ScaleDistribution.Init(0, ScaleColor_Distribution);
-	if (ScaleDistribution.IsValid())
+	if (FNiagaraStatelessBuiltDistribution::IsValid(ScaleColor_Distribution))
 	{
-		Particle.Color *= ScaleDistribution.GetValue(Particle.NormalizedAge);
+		Particle.Color *= SampleDistributionValueFloat4(ScaleColor_Distribution, 0, Particle.NormalizedAge);
 	}
 }
diff --git a/Engine/Plugins/FX/Niagara/Shaders/Private/Stateless/Modules/NiagaraStatelessModule_ScaleMeshSize.ush b/Engine/Plugins/FX/Niagara/Shaders/Private/Stateless/Modules/NiagaraStatelessModule_ScaleMeshSize.ush
--- a/Engine/Plugins/FX/Niagara/Shaders/Private/Stateless/Modules/NiagaraStatelessModule_ScaleMeshSize.ush
+++ b/Engine/Plugins/FX/Niagara/Shaders/Private/Stateless/Modules/NiagaraStatelessModule_ScaleMeshSize.ush
@@ -6,14 +6,12 @@ int		ScaleMeshSize_CurveScaleOffset;
 	
 void ScaleMeshSize_Simulate(inout FStatelessParticle Particle)
 {
-	FStatelessDistributionSamplerFloat3 MeshSizeDistribution;
-	MeshSizeDistribution.Init(0, ScaleMeshSize_Distribution);
-
-	if (MeshSizeDistribution.IsValid())
+	if (FNiagaraStatelessBuiltDistribution::IsValid(ScaleMeshSize_Distribution))
 	{
 		const float3 ScaleFactor = GetParameterBufferFloat3(ScaleMeshSize_CurveScaleOffset, 0, ScaleMeshSize_CurveScale);
-			
-		Particle.Scale			*= MeshSizeDistribution.GetValue(Particle.NormalizedAge) * ScaleFactor;
-		Particle.PreviousScale	*= MeshSizeDistribution.GetValue(Particle.PreviousNormalizedAge) * ScaleFactor;
+		const float2x3 Scales = SampleDistributionValuesFloat3(ScaleMeshSize_Distribution, 0, Particle.NormalizedAge, Particle.PreviousNormalizedAge);
+
+		Particle.Scale			*= Scales[0] * ScaleFactor;
+		Particle.PreviousScale	*= Scales[1] * ScaleFactor;
 	}
 }
diff --git a/Engine/Plugins/FX/Niagara/Shaders/Private/Stateless/Modules/NiagaraStatelessModule_ScaleRibbonWidth.ush b/Engine/Plugins/FX/Niagara/Shaders/Private/Stateless/Modules/NiagaraStatelessModule_ScaleRibbonWidth.ush
--- a/Engine/Plugins/FX/Niagara/Shaders/Private/Stateless/Modules/NiagaraStatelessModule_ScaleRibbonWidth.ush
+++ b/Engine/Plugins/FX/Niagara/Shaders/Private/Stateless/Modules/NiagaraStatelessModule_ScaleRibbonWidth.ush
@@ -5,14 +5,12 @@ float	ScaleRibbonWidth_CurveScale;
 
 void ScaleRibbonWidth_Simulate(inout FStatelessParticle Particle)
 {
-	FStatelessDistributionSamplerFloat ScaleDistribution;
-	ScaleDistribution.Init(0, ScaleRibbonWidth_Distribution);
-
-	if (ScaleDistribution.IsValid())
+	if (FNiagaraStatelessBuiltDistribution::IsValid(ScaleRibbonWidth_Distribution))
 	{
 		const float ScaleFactor = ScaleRibbonWidth_CurveScale;
-
-		Particle.RibbonWidth			*= ScaleDistribution.GetValue(Particle.NormalizedAge) * ScaleFactor;
-		Particle.PreviousRibbonWidth	*= ScaleDistribution.GetValue(Particle.PreviousNormalizedAge) * ScaleFactor;
+		const float2 Scales = SampleDistributionValuesFloat(ScaleRibbonWidth_Distribution, 0, Particle.NormalizedAge, Particle.PreviousNormalizedAge);
+	
+		Particle.RibbonWidth			*= Scales[0] * ScaleFactor;
+		Particle.PreviousRibbonWidth	*= Scales[1] * ScaleFactor;
 	}
 }
diff --git a/Engine/Plugins/FX/Niagara/Shaders/Private/Stateless/Modules/NiagaraStatelessModule_ScaleSpriteSize.ush b/Engine/Plugins/FX/Niagara/Shaders/Private/Stateless/Modules/NiagaraStatelessModule_ScaleSpriteSize.ush
--- a/Engine/Plugins/FX/Niagara/Shaders/Private/Stateless/Modules/NiagaraStatelessModule_ScaleSpriteSize.ush
+++ b/Engine/Plugins/FX/Niagara/Shaders/Private/Stateless/Modules/NiagaraStatelessModule_ScaleSpriteSize.ush
@@ -5,14 +5,12 @@ float2	ScaleSpriteSize_CurveScale;
 
 void ScaleSpriteSize_Simulate(inout FStatelessParticle Particle)
 {
-	FStatelessDistributionSamplerFloat2 ScaleDistribution;
-	ScaleDistribution.Init(0, ScaleSpriteSize_Distribution);
-
-	if (ScaleDistribution.IsValid())
+	if (FNiagaraStatelessBuiltDistribution::IsValid(ScaleSpriteSize_Distribution))
 	{
 		const float2 ScaleFactor = ScaleSpriteSize_CurveScale;
+		const float2x2 Scales = SampleDistributionValuesFloat2(ScaleSpriteSize_Distribution, 0, Particle.NormalizedAge, Particle.PreviousNormalizedAge);
 
-		Particle.SpriteSize			*= ScaleDistribution.GetValue(Particle.NormalizedAge) * ScaleFactor;
-		Particle.PreviousSpriteSize	*= ScaleDistribution.GetValue(Particle.PreviousNormalizedAge) * ScaleFactor;
+		Particle.SpriteSize			*= Scales[0] * ScaleFactor;
+		Particle.PreviousSpriteSize	*= Scales[1] * ScaleFactor;
 	}
 }
diff --git a/Engine/Plugins/FX/Niagara/Shaders/Private/Stateless/NiagaraStatelessCommon.ush b/Engine/Plugins/FX/Niagara/Shaders/Private/Stateless/NiagaraStatelessCommon.ush
--- a/Engine/Plugins/FX/Niagara/Shaders/Private/Stateless/NiagaraStatelessCommon.ush
+++ b/Engine/Plugins/FX/Niagara/Shaders/Private/Stateless/NiagaraStatelessCommon.ush
@@ -29,8 +29,8 @@ float4				Common_ToSimulationRotations[3];
 uint				Common_OutputBufferStride;
 RWBuffer<float>		Common_FloatOutputBuffer;
 RWBuffer<int>		Common_IntOutputBuffer;
-Buffer<float>		Common_StaticFloatBuffer;
-Buffer<uint>		Common_ParameterBuffer;
+ByteAddressBuffer	Common_StaticDataBuffer;
+ByteAddressBuffer	Common_DynamicDataBuffer;
 
 ////////////////////////////////////////////////////////////////////////////////////
 
@@ -108,24 +108,26 @@ float4 RandomScaleBiasFloat(uint RandomSeedOffset, float4 Scale, float4 Bias, bo
 float2 RandomUnitFloat2(uint RandomSeedOffset) { return SafeNormalize(RandomFloat2(RandomSeedOffset) - 0.5f); }
 float3 RandomUnitFloat3(uint RandomSeedOffset) { return SafeNormalize(RandomFloat3(RandomSeedOffset) - 0.5f); }
 
-float  GetStaticFloat (int BaseOffset, int Element)	{ return Common_StaticFloatBuffer[BaseOffset + Element]; }
-float2 GetStaticFloat2(int BaseOffset, int Element)	{ return float2(Common_StaticFloatBuffer[BaseOffset + Element*2], Common_StaticFloatBuffer[BaseOffset + 1 + Element*2]); }
-float3 GetStaticFloat3(int BaseOffset, int Element)	{ return float3(Common_StaticFloatBuffer[BaseOffset + Element*3], Common_StaticFloatBuffer[BaseOffset + 1 + Element*3], Common_StaticFloatBuffer[BaseOffset + 2 + Element*3]); }
-float4 GetStaticFloat4(int BaseOffset, int Element)	{ return float4(Common_StaticFloatBuffer[BaseOffset + Element*4], Common_StaticFloatBuffer[BaseOffset + 1 + Element*4], Common_StaticFloatBuffer[BaseOffset + 2 + Element*4], Common_StaticFloatBuffer[BaseOffset + 3 + Element*4]); }
-
-uint   GetParameterBufferUint  (int BaseOffset, int Element) { return Common_ParameterBuffer[BaseOffset + Element]; }
-int    GetParameterBufferInt   (int BaseOffset, int Element) { return asint(Common_ParameterBuffer[BaseOffset + Element]); }
-float  GetParameterBufferFloat (int BaseOffset, int Element) { return asfloat(Common_ParameterBuffer[BaseOffset + Element]); }
-float2 GetParameterBufferFloat2(int BaseOffset, int Element) { return float2(asfloat(Common_ParameterBuffer[BaseOffset + Element*2]), asfloat(Common_ParameterBuffer[BaseOffset + 1 + Element*2])); }
-float3 GetParameterBufferFloat3(int BaseOffset, int Element) { return float3(asfloat(Common_ParameterBuffer[BaseOffset + Element*3]), asfloat(Common_ParameterBuffer[BaseOffset + 1 + Element*3]), asfloat(Common_ParameterBuffer[BaseOffset + 2 + Element*3])); }
-float4 GetParameterBufferFloat4(int BaseOffset, int Element) { return float4(asfloat(Common_ParameterBuffer[BaseOffset + Element*4]), asfloat(Common_ParameterBuffer[BaseOffset + 1 + Element*4]), asfloat(Common_ParameterBuffer[BaseOffset + 2 + Element*4]), asfloat(Common_ParameterBuffer[BaseOffset + 3 + Element*4])); }
-
-uint   GetParameterBufferUint  (int BaseOffset, int Element, uint DefaultValue)   { return BaseOffset >= 0 ? Common_ParameterBuffer[BaseOffset + Element] : DefaultValue; }
-int    GetParameterBufferInt   (int BaseOffset, int Element, int DefaultValue)    { return BaseOffset >= 0 ? asint(Common_ParameterBuffer[BaseOffset + Element]) : DefaultValue; }
-float  GetParameterBufferFloat (int BaseOffset, int Element, float  DefaultValue) { return BaseOffset >= 0 ? asfloat(Common_ParameterBuffer[BaseOffset + Element]) : DefaultValue; }
-float2 GetParameterBufferFloat2(int BaseOffset, int Element, float2 DefaultValue) { return BaseOffset >= 0 ? float2(asfloat(Common_ParameterBuffer[BaseOffset + Element*2]), asfloat(Common_ParameterBuffer[BaseOffset + 1 + Element*2])) : DefaultValue; }
-float3 GetParameterBufferFloat3(int BaseOffset, int Element, float3 DefaultValue) { return BaseOffset >= 0 ? float3(asfloat(Common_ParameterBuffer[BaseOffset + Element*3]), asfloat(Common_ParameterBuffer[BaseOffset + 1 + Element*3]), asfloat(Common_ParameterBuffer[BaseOffset + 2 + Element*3])) : DefaultValue; }
-float4 GetParameterBufferFloat4(int BaseOffset, int Element, float4 DefaultValue) { return BaseOffset >= 0 ? float4(asfloat(Common_ParameterBuffer[BaseOffset + Element*4]), asfloat(Common_ParameterBuffer[BaseOffset + 1 + Element*4]), asfloat(Common_ParameterBuffer[BaseOffset + 2 + Element*4]), asfloat(Common_ParameterBuffer[BaseOffset + 3 + Element*4])) : DefaultValue; }
+uint   GetStaticUInt  (int BaseOffset, int Element)	{ return asuint (Common_StaticDataBuffer.Load ((BaseOffset + Element) * 4)); }
+int    GetStaticInt   (int BaseOffset, int Element)	{ return asint  (Common_StaticDataBuffer.Load ((BaseOffset + Element) * 4)); }
+float  GetStaticFloat (int BaseOffset, int Element)	{ return asfloat(Common_StaticDataBuffer.Load ((BaseOffset + Element) * 4)); }
+float2 GetStaticFloat2(int BaseOffset, int Element)	{ return asfloat(Common_StaticDataBuffer.Load2((BaseOffset + Element*2) * 4)); }
+float3 GetStaticFloat3(int BaseOffset, int Element)	{ return asfloat(Common_StaticDataBuffer.Load3((BaseOffset + Element*3) * 4)); }
+float4 GetStaticFloat4(int BaseOffset, int Element)	{ return asfloat(Common_StaticDataBuffer.Load4((BaseOffset + Element*4) * 4)); }
+
+uint   GetParameterBufferUint  (int BaseOffset, int Element) { return asuint (Common_DynamicDataBuffer.Load ((BaseOffset + Element) * 4)); }
+int    GetParameterBufferInt   (int BaseOffset, int Element) { return asint  (Common_DynamicDataBuffer.Load ((BaseOffset + Element) * 4)); }
+float  GetParameterBufferFloat (int BaseOffset, int Element) { return asfloat(Common_DynamicDataBuffer.Load ((BaseOffset + Element) * 4)); }
+float2 GetParameterBufferFloat2(int BaseOffset, int Element) { return asfloat(Common_DynamicDataBuffer.Load2((BaseOffset + Element*2) * 4)); }
+float3 GetParameterBufferFloat3(int BaseOffset, int Element) { return asfloat(Common_DynamicDataBuffer.Load3((BaseOffset + Element*3) * 4)); }
+float4 GetParameterBufferFloat4(int BaseOffset, int Element) { return asfloat(Common_DynamicDataBuffer.Load4((BaseOffset + Element*4) * 4)); }
+
+uint   GetParameterBufferUint  (int BaseOffset, int Element, uint  DefaultValue)  { return BaseOffset >= 0 ? asuint (Common_DynamicDataBuffer.Load ((BaseOffset + Element) * 4)) : DefaultValue; }
+int    GetParameterBufferInt   (int BaseOffset, int Element, int    DefaultValue) { return BaseOffset >= 0 ? asint  (Common_DynamicDataBuffer.Load ((BaseOffset + Element) * 4)) : DefaultValue; }
+float  GetParameterBufferFloat (int BaseOffset, int Element, float  DefaultValue) { return BaseOffset >= 0 ? asfloat(Common_DynamicDataBuffer.Load ((BaseOffset + Element) * 4)) : DefaultValue; }
+float2 GetParameterBufferFloat2(int BaseOffset, int Element, float2 DefaultValue) { return BaseOffset >= 0 ? asfloat(Common_DynamicDataBuffer.Load2((BaseOffset + Element*2) * 4)) : DefaultValue; }
+float3 GetParameterBufferFloat3(int BaseOffset, int Element, float3 DefaultValue) { return BaseOffset >= 0 ? asfloat(Common_DynamicDataBuffer.Load3((BaseOffset + Element*3) * 4)) : DefaultValue; }
+float4 GetParameterBufferFloat4(int BaseOffset, int Element, float4 DefaultValue) { return BaseOffset >= 0 ? asfloat(Common_DynamicDataBuffer.Load4((BaseOffset + Element*4) * 4)) : DefaultValue; }
 	
 float LerpStaticFloat (uint2 Parameters, float U)
 {
@@ -152,27 +154,9 @@ float3 LerpStaticFloat3(uint2 Parameters, float U)
 }
 
 float  StatelessLerp(float  lhs, float  rhs, float  u, bool bUniform) { return lerp(lhs, rhs, u); }
-float2 StatelessLerp(float2 lhs, float2 rhs, float2 u, bool bUniform) { return bUniform ? lerp(lhs.x, rhs.x, u.x).xx   : lerp(lhs, rhs, u); }
-float3 StatelessLerp(float3 lhs, float3 rhs, float3 u, bool bUniform) { return bUniform ? lerp(lhs.x, rhs.x, u.x).xxx  : lerp(lhs, rhs, u); }
-float4 StatelessLerp(float4 lhs, float4 rhs, float4 u, bool bUniform) { return bUniform ? lerp(lhs.x, rhs.x, u.x).xxxx : lerp(lhs, rhs, u); }
-
-float SampleCurveFloat (FNiagaraStatelessBuiltDistributionType Distribution, float Time)
-{
-	const uint DataOffset = FNiagaraStatelessBuiltDistribution::GetDataOffset(Distribution);
-	const float Offset = FNiagaraStatelessBuiltDistribution::ConvertTimeToLookup(Distribution, Time);
-	const float Value0 = GetStaticFloat(DataOffset, floor(Offset));
-	const float Value1 = GetStaticFloat(DataOffset, ceil(Offset));
-	return lerp(Value0, Value1, frac(Offset));
-}
-
-float3 SampleCurveFloat3(FNiagaraStatelessBuiltDistributionType Distribution, float Time)
-{
-	const uint DataOffset = FNiagaraStatelessBuiltDistribution::GetDataOffset(Distribution);
-	const float Offset = FNiagaraStatelessBuiltDistribution::ConvertTimeToLookup(Distribution, Time);
-	const float3 Value0 = GetStaticFloat3(DataOffset, floor(Offset));
-	const float3 Value1 = GetStaticFloat3(DataOffset, ceil(Offset));
-	return lerp(Value0, Value1, frac(Offset));
-}
+float2 StatelessLerp(float2 lhs, float2 rhs, float2 u, bool bUniform) { return lerp(lhs, rhs, bUniform ? u.xx : u); }
+float3 StatelessLerp(float3 lhs, float3 rhs, float3 u, bool bUniform) { return lerp(lhs, rhs, bUniform ? u.xxx : u); }
+float4 StatelessLerp(float4 lhs, float4 rhs, float4 u, bool bUniform) { return lerp(lhs, rhs, bUniform ? u.xxxx : u); }
 
 uint GetComponentOffset(uint OutputIndex, int ComponentOffset)
 {
diff --git a/Engine/Plugins/FX/Niagara/Shaders/Private/Stateless/NiagaraStatelessDistribution.ush b/Engine/Plugins/FX/Niagara/Shaders/Private/Stateless/NiagaraStatelessDistribution.ush
--- a/Engine/Plugins/FX/Niagara/Shaders/Private/Stateless/NiagaraStatelessDistribution.ush
+++ b/Engine/Plugins/FX/Niagara/Shaders/Private/Stateless/NiagaraStatelessDistribution.ush
@@ -1,50 +1,83 @@
 // Copyright Epic Games, Inc. All Rights Reserved.
 #pragma once
 
-#define DefineStatelessDistributionSampler(HLSLTYPE, NAMETYPE) \
-	struct FStatelessDistributionSampler##NAMETYPE \
+#define SAMPLEFUNC(HLSLTYPE, HLSLPAIRTYPE, NAMETYPE) \
+	HLSLPAIRTYPE Private_GetDistributionTableValues##NAMETYPE##(FNiagaraStatelessBuiltDistributionType BuiltData, uint IndexA, uint IndexB) \
 	{ \
-		void Init(uint RandomSeedOffset, FNiagaraStatelessBuiltDistributionType InBuildDistribution) \
-		{ \
-			BuildDistribution = InBuildDistribution; \
-			RandomOffset = 0.0f; \
-			if ( FNiagaraStatelessBuiltDistribution::IsRandom(BuildDistribution) != 0) \
-			{ \
-				RandomOffset = Random##NAMETYPE##(RandomSeedOffset); \
-			} \
-		} \
-		\
-		bool IsValid() { return FNiagaraStatelessBuiltDistribution::IsValid(BuildDistribution); } \
-		\
-		HLSLTYPE GetValue(float Time) \
-		{ \
-			const uint DataOffset = FNiagaraStatelessBuiltDistribution::GetDataOffset(BuildDistribution); \
-			if (FNiagaraStatelessBuiltDistribution::IsBinding(BuildDistribution)) \
-			{ \
-				return GetParameterBuffer##NAMETYPE##(DataOffset, 0); \
-			} \
-			else if (FNiagaraStatelessBuiltDistribution::IsRandom(BuildDistribution)) \
-			{ \
-				const HLSLTYPE Value0 = GetStatic##NAMETYPE##(DataOffset, 0); \
-				const HLSLTYPE Value1 = GetStatic##NAMETYPE##(DataOffset, 1); \
-				return StatelessLerp(Value0, Value1, RandomOffset, FNiagaraStatelessBuiltDistribution::IsUniform(BuildDistribution)); \
-			} \
-			else \
-			{ \
-				const float Offset = FNiagaraStatelessBuiltDistribution::ConvertTimeToLookup(BuildDistribution, Time); \
-				const HLSLTYPE Value0 = GetStatic##NAMETYPE##(DataOffset, floor(Offset)); \
-				const HLSLTYPE Value1 = GetStatic##NAMETYPE##(DataOffset, ceil(Offset)); \
-				return lerp(Value0, Value1, frac(Offset)); \
-			} \
-		} \
-		\
-		FNiagaraStatelessBuiltDistributionType		BuildDistribution; \
-		HLSLTYPE									RandomOffset; \
-	}
+		const uint BaseOffset = FNiagaraStatelessBuiltDistribution::GetBufferOffset(BuiltData); \
+		if (FNiagaraStatelessBuiltDistribution::IsBufferReadModeStatic(BuiltData)) \
+		{ \
+			return HLSLPAIRTYPE(GetStatic##NAMETYPE##(BaseOffset, IndexA), GetStatic##NAMETYPE##(BaseOffset, IndexB)); \
+		} \
+		else \
+		{ \
+			return HLSLPAIRTYPE(GetParameterBuffer##NAMETYPE##(BaseOffset, IndexA), GetParameterBuffer##NAMETYPE##(BaseOffset, IndexB)); \
+		} \
+	} \
+	\
+	HLSLTYPE Private_GetDistributionRandomValue##NAMETYPE##(FNiagaraStatelessBuiltDistributionType BuiltData, uint RandomSeedOffset) \
+	{ \
+		const float Random = RandomFloat(RandomSeedOffset); \
+		uint Indicies[2]; \
+		float Interp; \
+		FNiagaraStatelessBuiltDistribution::ConvertRandomToLookup(BuiltData, Random, Indicies[0], Indicies[1], Interp); \
+		HLSLPAIRTYPE Values = Private_GetDistributionTableValues##NAMETYPE##(BuiltData, Indicies[0], Indicies[1]); \
+		if (FNiagaraStatelessBuiltDistribution::IsInterpolationModeLinear(BuiltData)) \
+		{ \
+			const bool bUniformRandom = FNiagaraStatelessBuiltDistribution::IsUniformRandomEnabled(BuiltData); \
+			return StatelessLerp(Values[0], Values[1], Random##NAMETYPE##(RandomSeedOffset), bUniformRandom); \
+		} \
+		else \
+		{ \
+			return Values[0]; \
+		} \
+	} \
+	HLSLTYPE Private_GetDistributionTimeValue##NAMETYPE##(FNiagaraStatelessBuiltDistributionType BuiltData, float Time) \
+	{ \
+		uint Indicies[2]; \
+		float Interp; \
+		FNiagaraStatelessBuiltDistribution::ConvertTimeToLookup(BuiltData, Time, Indicies[0], Indicies[1], Interp); \
+		HLSLPAIRTYPE Values = Private_GetDistributionTableValues##NAMETYPE##(BuiltData, Indicies[0], Indicies[1]); \
+		return lerp(Values[0], Values[1], Interp); \
+	} \
+	\
+	HLSLTYPE SampleDistributionValue##NAMETYPE##(FNiagaraStatelessBuiltDistributionType BuiltData, uint RandomSeedOffset, float Time) \
+	{ \
+		if (FNiagaraStatelessBuiltDistribution::IsLookupValueModeRandom(BuiltData)) \
+		{ \
+			return Private_GetDistributionRandomValue##NAMETYPE##(BuiltData, RandomSeedOffset); \
+		} \
+		else \
+		{ \
+			return Private_GetDistributionTimeValue##NAMETYPE##(BuiltData, Time); \
+		} \
+	} \
+	HLSLPAIRTYPE SampleDistributionValues##NAMETYPE##(FNiagaraStatelessBuiltDistributionType BuiltData, uint RandomSeedOffset, float TimeA, float TimeB) \
+	{ \
+		if (FNiagaraStatelessBuiltDistribution::IsLookupValueModeRandom(BuiltData)) \
+		{ \
+			HLSLTYPE Value = Private_GetDistributionRandomValue##NAMETYPE##(BuiltData, RandomSeedOffset); \
+			return HLSLPAIRTYPE(Value, Value); \
+		} \
+		else \
+		{ \
+			return HLSLPAIRTYPE(Private_GetDistributionTimeValue##NAMETYPE##(BuiltData, TimeA), Private_GetDistributionTimeValue##NAMETYPE##(BuiltData, TimeB)); \
+		} \
+	}\
+    HLSLTYPE SampleCurve##NAMETYPE##(FNiagaraStatelessBuiltDistributionType Distribution, float Time) \
+    { \
+	    uint Indicies[2]; \
+	    float Interp; \
+	    FNiagaraStatelessBuiltDistribution::ConvertTimeToLookup(Distribution, Time, Indicies[0], Indicies[1], Interp); \
+	    const uint DataOffset = FNiagaraStatelessBuiltDistribution::GetBufferOffset(Distribution); \
+	    const HLSLTYPE Value0 = GetStatic##NAMETYPE##(DataOffset, Indicies[0]); \
+	    const HLSLTYPE Value1 = GetStatic##NAMETYPE##(DataOffset, Indicies[1]); \
+	    return lerp(Value0, Value1, Interp); \
+    }
 
-DefineStatelessDistributionSampler(float,  Float);
-DefineStatelessDistributionSampler(float2, Float2);
-DefineStatelessDistributionSampler(float3, Float3);
-DefineStatelessDistributionSampler(float4, Float4);
+	SAMPLEFUNC(float,  float2,   Float)
+	SAMPLEFUNC(float2, float2x2, Float2)
+	SAMPLEFUNC(float3, float2x3, Float3)
+	SAMPLEFUNC(float4, float2x4, Float4)
 
-#undef DefineStatelessDistributionSampler
+#undef READFUNCS
diff --git a/Engine/Plugins/FX/Niagara/Shaders/Shared/NiagaraStatelessBuiltDistribution.h b/Engine/Plugins/FX/Niagara/Shaders/Shared/NiagaraStatelessBuiltDistribution.h
--- a/Engine/Plugins/FX/Niagara/Shaders/Shared/NiagaraStatelessBuiltDistribution.h
+++ b/Engine/Plugins/FX/Niagara/Shaders/Shared/NiagaraStatelessBuiltDistribution.h
@@ -5,90 +5,188 @@
 	#define uint				uint32
 	typedef FUintVector3		FNiagaraStatelessBuiltDistributionType;
 	typedef const FUintVector3&	FNiagaraStatelessBuiltDistributionTypeIn;
+
+	#define UEAsFloat(X) reinterpret_cast<const float&>(X);
+	#define UEMin(X, Y) FMath::Min(X, Y)
+	#define UEMax(X, Y) FMath::Max(X, Y)
+	#define UEFrac(X) FMath::Fractional(X)
+	#define UERoundToInt(X) FMath::RoundToInt(X)
+	#define UEOutParam(X) X&
 #else
 	//typedef uint3			FNiagaraStatelessBuiltDistributionType;
 	#define FNiagaraStatelessBuiltDistributionType		uint3
 	#define FNiagaraStatelessBuiltDistributionTypeIn	uint3
+
+	#define UEAsFloat(X) asfloat(X)
+	#define UEMin(X, Y) min(X, Y)
+	#define UEMax(X, Y) max(X, Y)
+	#define UEFrac(X) frac(X)
+	#define UERoundToInt(X) uint(round(X))
+	#define UEOutParam(X) out X
 #endif
 
-#define ENiagaraStatelessBuiltDistributionFlag_Binding			0x20000000u
-#define ENiagaraStatelessBuiltDistributionFlag_Random			0x40000000u
-#define ENiagaraStatelessBuiltDistributionFlag_Uniform			0x80000000u
-#define ENiagaraStatelessBuiltDistributionFlag_TableLengthBits	9u				// 512 entries in LUT
-#define ENiagaraStatelessBuiltDistributionFlag_TableLengthShift	20u
-#define ENiagaraStatelessBuiltDistributionFlag_TableLengthMask	((1u << ENiagaraStatelessBuiltDistributionFlag_TableLengthBits) - 1u)
-#define ENiagaraStatelessBuiltDistributionFlag_DataOffsetBits	20u
-#define ENiagaraStatelessBuiltDistributionFlag_DataOffsetShift	0u
-#define ENiagaraStatelessBuiltDistributionFlag_DataOffsetMask	((1u << ENiagaraStatelessBuiltDistributionFlag_DataOffsetBits) - 1u)
+////////////////////////////////////////////////////////////////////////////////////////////////////////////////////////////////////////////////////////////////////////////
 
 struct FNiagaraStatelessBuiltDistribution
 {
+	struct EFlags
+	{
+		enum
+		{
+			BufferOffsetShift		= 0u,				// Offset in the buffer where the data is stored
+			BufferOffsetBits		= 17u,				// 128k max offset
+			BufferOffsetMask		= (1u << BufferOffsetBits) - 1u,
+
+			TableLengthShift		= 17u,				// Length of the table data minus 1
+			TableLengthBits			= 9u,				// 512 entries max
+			TableLengthMask			= (1u << TableLengthBits) - 1u,
+
+			LookupValueModeShift	= 26u,				// Lookup Value Mode where all bits set is assumed to be random internally
+			LookupValueModeBits		= 2u,				// The caller must decide how to interpret this data
+			LookupValueModeMask		= (1u << LookupValueModeBits) - 1u,
+
+			UniformRandom			= 1u << 28u,		// When true random will be uniform when interpolating vs non-uniform
+			AddressModeWrap			= 1u << 29u,		// The address mode, when 0 clamp, when 1 wrap
+			InterpolationModeLinear	= 1u << 30u,		// The interpolation mode, when 0 none, when 1 linear
+			BufferReadModeStatic	= 1u << 31u,		// The buffer read mode, when 0 dynamic, when 1 static
+		};
+	};
+
 	static bool IsValid(FNiagaraStatelessBuiltDistributionTypeIn BuiltData) { return BuiltData[2] != 0; }
-	static bool IsBinding(FNiagaraStatelessBuiltDistributionTypeIn BuiltData) { return (BuiltData[0] & ENiagaraStatelessBuiltDistributionFlag_Binding) != 0; }
-	static bool IsRandom(FNiagaraStatelessBuiltDistributionTypeIn BuiltData) { return (BuiltData[0] & ENiagaraStatelessBuiltDistributionFlag_Random) != 0; }
-	static bool IsUniform(FNiagaraStatelessBuiltDistributionTypeIn BuiltData) { return (BuiltData[0] & ENiagaraStatelessBuiltDistributionFlag_Uniform) != 0; }
 
-	static uint GetDataOffset(FNiagaraStatelessBuiltDistributionTypeIn BuiltData) { return (BuiltData[0] >> ENiagaraStatelessBuiltDistributionFlag_DataOffsetShift) & ENiagaraStatelessBuiltDistributionFlag_DataOffsetMask; }
-	static uint GetTableLength(FNiagaraStatelessBuiltDistributionTypeIn BuiltData) { return (BuiltData[0] >> ENiagaraStatelessBuiltDistributionFlag_TableLengthShift) & ENiagaraStatelessBuiltDistributionFlag_TableLengthMask; }
+	static bool IsUniformRandomEnabled(FNiagaraStatelessBuiltDistributionTypeIn BuiltData) { return (BuiltData[0] & EFlags::UniformRandom) != 0; }
+	static bool IsAddressModeWrap(FNiagaraStatelessBuiltDistributionTypeIn BuiltData) { return (BuiltData[0] & EFlags::AddressModeWrap) != 0; }
+	static bool IsInterpolationModeLinear(FNiagaraStatelessBuiltDistributionTypeIn BuiltData) { return (BuiltData[0] & EFlags::InterpolationModeLinear) != 0; }
+	static bool IsBufferReadModeStatic(FNiagaraStatelessBuiltDistributionTypeIn BuiltData) { return (BuiltData[0] & EFlags::BufferReadModeStatic) != 0; }
+	static bool IsLookupValueModeRandom(FNiagaraStatelessBuiltDistributionTypeIn BuiltData) { const uint Mask = EFlags::LookupValueModeMask << EFlags::LookupValueModeShift; return (BuiltData[0] & Mask) == Mask; }
+
+	static uint GetBufferOffset(FNiagaraStatelessBuiltDistributionTypeIn BuiltData) { return (BuiltData[0] >> EFlags::BufferOffsetShift) & EFlags::BufferOffsetMask; }
+	static uint GetTableLength(FNiagaraStatelessBuiltDistributionTypeIn BuiltData) { return (BuiltData[0] >> EFlags::TableLengthShift) & EFlags::TableLengthMask; }
+	static uint GetLookupValueMode(FNiagaraStatelessBuiltDistributionTypeIn BuiltData) { return (BuiltData[0] >> EFlags::LookupValueModeShift) & EFlags::LookupValueModeMask; }
 
-#ifdef GPU_SIMULATION
-	static float ConvertTimeToLookup(FNiagaraStatelessBuiltDistributionTypeIn BuiltData, float Time)
+	static void ConvertRandomToLookup(FNiagaraStatelessBuiltDistributionTypeIn BuiltData, float Random, UEOutParam(uint) OutIndexA, UEOutParam(uint) OutIndexB, UEOutParam(float) OutInterp)
 	{
-		const float TimeBias	= asfloat(BuiltData[1]);
-		const float TimeScale	= asfloat(BuiltData[2]);
-		const float TableLength	= float(GetTableLength(BuiltData));
-		const float Offset		= clamp((Time - TimeBias) * TimeScale, 0.0f, TableLength);
-		return Offset;
+		const uint TableLength = GetTableLength(BuiltData);
+		const bool bInterpolate = IsInterpolationModeLinear(BuiltData);
+		const float Index = Random * float(TableLength);
+		OutIndexA = bInterpolate ? uint(Index) : UERoundToInt(Index);
+		OutIndexB = OutIndexA + 1;
+		OutInterp = bInterpolate ? UEFrac(Index) : 0.0f;
+		if (IsAddressModeWrap(BuiltData))
+		{
+			OutIndexB = OutIndexB > TableLength ? 0 : OutIndexB;
+		}
+		else
+		{
+			OutIndexB = UEMin(OutIndexB, TableLength);
+		}
 	}
-#else
-	static float ConvertTimeToLookup(FNiagaraStatelessBuiltDistributionTypeIn BuiltData, float Time)
+
+	static void ConvertTimeToLookup(FNiagaraStatelessBuiltDistributionTypeIn BuiltData, float Time, UEOutParam(uint) OutIndexA, UEOutParam(uint) OutIndexB, UEOutParam(float) OutInterp)
 	{
-		const float TimeBias	= reinterpret_cast<const float&>(BuiltData[1]);
-		const float TimeScale	= reinterpret_cast<const float&>(BuiltData[2]);
-		const float TableLength	= float(GetTableLength(BuiltData));
-		const float Offset		= FMath::Clamp((Time - TimeBias) * TimeScale, 0.0f, TableLength);
-		return Offset;
+		const float TimeBias	= UEAsFloat(BuiltData[1]);
+		const float TimeScale	= UEAsFloat(BuiltData[2]);
+		const uint  TableLength	= GetTableLength(BuiltData);
+		const bool bInterpolate = IsInterpolationModeLinear(BuiltData);
+		Time = UEMax((Time - TimeBias) * TimeScale, 0.0f);
+
+		OutIndexA = bInterpolate ? uint(Time) : UERoundToInt(Time);
+		OutIndexB = OutIndexA + 1;
+		if (IsAddressModeWrap(BuiltData))
+		{
+			OutIndexA = OutIndexA % (TableLength + 1);
+			OutIndexB = OutIndexB > TableLength ? 0 : OutIndexB;
+		}
+		else
+		{
+			OutIndexA = UEMin(OutIndexA, TableLength);
+			OutIndexB = UEMin(OutIndexB, TableLength);
+		}
+		OutInterp = bInterpolate ? UEFrac(Time) : 0.0f;
 	}
-	
+
+#ifndef GPU_SIMULATION
 	static FNiagaraStatelessBuiltDistributionType GetDefault() { return FNiagaraStatelessBuiltDistributionType::ZeroValue; }
 
-	static void SetIsBinding(FNiagaraStatelessBuiltDistributionType& BuiltData) { BuiltData[0] |= ENiagaraStatelessBuiltDistributionFlag_Binding; }
-	static void SetIsRandom(FNiagaraStatelessBuiltDistributionType& BuiltData) { BuiltData[0] |= ENiagaraStatelessBuiltDistributionFlag_Random; }
-	static void SetIsUniform(FNiagaraStatelessBuiltDistributionType& BuiltData) { BuiltData[0] |= ENiagaraStatelessBuiltDistributionFlag_Uniform; }
+	static void SetInterpolationMode(FNiagaraStatelessBuiltDistributionType& BuiltData, bool bLinearInterpolation)
+	{
+		BuiltData[0] &= ~EFlags::InterpolationModeLinear;
+		BuiltData[0] |= bLinearInterpolation ? EFlags::InterpolationModeLinear : 0;
+	}
+
+	static void SetAddressMode(FNiagaraStatelessBuiltDistributionType& BuiltData, bool bClampAddress)
+	{
+		BuiltData[0] &= ~EFlags::AddressModeWrap;
+		BuiltData[0] |= bClampAddress ? 0 : EFlags::AddressModeWrap;
+	}
 
-	static void SetLookupParameters(FNiagaraStatelessBuiltDistributionType& BuiltData, uint DataOffset)
+	static void SetBufferReadMode(FNiagaraStatelessBuiltDistributionType& BuiltData, bool bStaticBuffer)
 	{
-		check(IsBinding(BuiltData));
-		check(DataOffset <= ENiagaraStatelessBuiltDistributionFlag_DataOffsetMask);
+		BuiltData[0] &= ~EFlags::BufferReadModeStatic;
+		BuiltData[0] |= bStaticBuffer ? EFlags::BufferReadModeStatic : 0;
+	}
 
-		BuiltData[0] &= ~(ENiagaraStatelessBuiltDistributionFlag_DataOffsetMask << ENiagaraStatelessBuiltDistributionFlag_DataOffsetShift);
-		BuiltData[0] &= ~(ENiagaraStatelessBuiltDistributionFlag_TableLengthMask << ENiagaraStatelessBuiltDistributionFlag_TableLengthShift);
-		BuiltData[0] |= DataOffset << ENiagaraStatelessBuiltDistributionFlag_DataOffsetShift;
+	static void SetLookupValueMode(FNiagaraStatelessBuiltDistributionType& BuiltData, uint8 ValueMode)
+	{
+		BuiltData[0] &= ~(EFlags::LookupValueModeMask << EFlags::LookupValueModeShift);
+		if (ValueMode == 0xff)
+		{
+			BuiltData[0] |= EFlags::LookupValueModeMask << EFlags::LookupValueModeShift;
+		}
+		else
+		{
+			checkf(ValueMode <= EFlags::LookupValueModeMask - 1, TEXT("ValueMode (%d) is out of range (0 - %d)"), ValueMode, EFlags::LookupValueModeMask - 1);
+			BuiltData[0] |= ValueMode << EFlags::LookupValueModeShift;
+		}
+	}
 
-		reinterpret_cast<float&>(BuiltData[1]) = 0.0f;
-		reinterpret_cast<float&>(BuiltData[2]) = 1.0f;
+	static void SetUniformRandom(FNiagaraStatelessBuiltDistributionType& BuiltData, bool bEnabled)
+	{
+		BuiltData[0] = (BuiltData[0] & ~EFlags::UniformRandom) | (bEnabled ? EFlags::UniformRandom : 0);
 	}
 
-	static void SetLookupParameters(FNiagaraStatelessBuiltDistributionType& BuiltData, uint DataOffset, uint TableLength, const FVector2f& TimeRange)
+	static void SetTimeScaleBias(FNiagaraStatelessBuiltDistributionType& BuiltData, float TimeScale, float TimeBias)
 	{
-		const uint TableLengthMinusOne = TableLength - 1;
+		reinterpret_cast<float&>(BuiltData[1]) = TimeBias;
+		reinterpret_cast<float&>(BuiltData[2]) = FMath::Max(TimeScale, UE_KINDA_SMALL_NUMBER);
+	}
 
-		check(DataOffset <= ENiagaraStatelessBuiltDistributionFlag_DataOffsetMask);
-		check(TableLengthMinusOne <= ENiagaraStatelessBuiltDistributionFlag_TableLengthMask);
+	static void SetLookupParameters(FNiagaraStatelessBuiltDistributionType& BuiltData, uint BufferOffset, uint TableLength)
+	{
+		check(BufferOffset <= EFlags::BufferOffsetMask);
+		check(TableLength > 0 && TableLength <= EFlags::TableLengthMask);
+
+		BuiltData[0] &= ~(EFlags::BufferOffsetMask << EFlags::BufferOffsetShift);
+		BuiltData[0] |= BufferOffset << EFlags::BufferOffsetShift;
+
+		BuiltData[0] &= ~(EFlags::TableLengthMask << EFlags::TableLengthShift);
+		BuiltData[0] |= (TableLength - 1) << EFlags::TableLengthShift;
+
+		SetTimeScaleBias(BuiltData, float(TableLength - 1), 0.0f);
+	}
 
-		BuiltData[0] &= ~(ENiagaraStatelessBuiltDistributionFlag_DataOffsetMask << ENiagaraStatelessBuiltDistributionFlag_DataOffsetShift);
-		BuiltData[0] |= DataOffset << ENiagaraStatelessBuiltDistributionFlag_DataOffsetShift;
+	static void SetLookupParameters(FNiagaraStatelessBuiltDistributionType& BuiltData, uint BufferOffset, uint TableLength, float MinTime, float MaxTime)
+	{
+		check(TableLength > 0);
 
-		BuiltData[0] &= ~(ENiagaraStatelessBuiltDistributionFlag_TableLengthMask << ENiagaraStatelessBuiltDistributionFlag_TableLengthShift);
-		BuiltData[0] |= TableLengthMinusOne << ENiagaraStatelessBuiltDistributionFlag_TableLengthShift;
+		SetLookupParameters(BuiltData, BufferOffset, TableLength);
 
-		const float TimeDuration = TimeRange.Y - TimeRange.X;
-		reinterpret_cast<float&>(BuiltData[1]) = TimeRange.X;
-		reinterpret_cast<float&>(BuiltData[2]) = TableLengthMinusOne > 0 && TimeDuration > 0.0f ? float(TableLengthMinusOne) / TimeDuration : 0.0f;
+		const uint TableLengthMinusOne = TableLength - 1;
+		const float TimeDuration = MaxTime - MinTime;
+		const float TimeScale = TableLengthMinusOne > 0 && TimeDuration > 0.0f ? float(TableLengthMinusOne) / TimeDuration : 0.0f;
+		SetTimeScaleBias(BuiltData, TimeScale, MinTime);
 	}
 #endif
 };
 
+////////////////////////////////////////////////////////////////////////////////////////////////////////////////////////////////////////////////////////////////////////////
+
 #ifndef GPU_SIMULATION
 	#undef uint
 #endif
+#undef UEAsFloat
+#undef UEMin
+#undef UEMax
+#undef UEFrac
+#undef UERoundToInt
+#undef UEOutParam
diff --git a/Engine/Plugins/FX/Niagara/Source/Niagara/Internal/Stateless/Modules/NiagaraStatelessModule_InitializeParticle.h b/Engine/Plugins/FX/Niagara/Source/Niagara/Internal/Stateless/Modules/NiagaraStatelessModule_InitializeParticle.h
--- a/Engine/Plugins/FX/Niagara/Source/Niagara/Internal/Stateless/Modules/NiagaraStatelessModule_InitializeParticle.h
+++ b/Engine/Plugins/FX/Niagara/Source/Niagara/Internal/Stateless/Modules/NiagaraStatelessModule_InitializeParticle.h
@@ -41,7 +41,7 @@ public:
 	UPROPERTY(EditAnywhere, Category = "Parameters", meta = (DisplayName = "Ribbon Width", EditCondition="bWriteRibbonWidth"))
 	FNiagaraDistributionRangeFloat RibbonWidthDistribution = FNiagaraDistributionRangeFloat(FNiagaraStatelessGlobals::GetDefaultRibbonWidthValue());
 
-	UPROPERTY(EditAnywhere, Category = "Parameters", meta = (DisableCurveDistribution))
+	UPROPERTY(EditAnywhere, Category = "Parameters", meta = (DisableCurveDistribution, DisableArrayDistribution))
 	FNiagaraDistributionPosition InitialPositionDistribution = FNiagaraDistributionPosition(FVector3f::ZeroVector);
 
 	virtual void BuildEmitterData(const FNiagaraStatelessEmitterDataBuildContext& BuildContext) const override;
diff --git a/Engine/Plugins/FX/Niagara/Source/Niagara/Internal/Stateless/NiagaraStatelessEmitterData.h b/Engine/Plugins/FX/Niagara/Source/Niagara/Internal/Stateless/NiagaraStatelessEmitterData.h
--- a/Engine/Plugins/FX/Niagara/Source/Niagara/Internal/Stateless/NiagaraStatelessEmitterData.h
+++ b/Engine/Plugins/FX/Niagara/Source/Niagara/Internal/Stateless/NiagaraStatelessEmitterData.h
@@ -58,8 +58,8 @@ struct FNiagaraStatelessEmitterData
 	const UNiagaraStatelessEmitterTemplate*			EmitterTemplate = nullptr;	// Used to access shader information
 
 	TArray<uint8>									BuiltData;					// Built data, generally allocated by modules if any
-	TArray<float>									StaticFloatData;			// Used with CPU generation, must be valid if ParticleSimExecData is also valid
-	FReadBuffer										StaticFloatBuffer;
+	TArray<uint8>									StaticDataBufferCpu;		// Used with CPU generation, must be valid if ParticleSimExecData is also valid
+	FByteAddressBuffer								StaticDataBufferGpu;		// Used with GPU generation, guaranteed to be minimum safe size
 
 	NiagaraStateless::FParticleSimulationExecData*	ParticleSimExecData = nullptr;	// CPU simulation execution data, when null we don't provide a CPU path
 
diff --git a/Engine/Plugins/FX/Niagara/Source/Niagara/Internal/Stateless/NiagaraStatelessEmitterDataBuildContext.h b/Engine/Plugins/FX/Niagara/Source/Niagara/Internal/Stateless/NiagaraStatelessEmitterDataBuildContext.h
--- a/Engine/Plugins/FX/Niagara/Source/Niagara/Internal/Stateless/NiagaraStatelessEmitterDataBuildContext.h
+++ b/Engine/Plugins/FX/Niagara/Source/Niagara/Internal/Stateless/NiagaraStatelessEmitterDataBuildContext.h
@@ -22,23 +22,25 @@ class FNiagaraStatelessEmitterDataBuildContext
 public:
 	UE_NONCOPYABLE(FNiagaraStatelessEmitterDataBuildContext);
 
-	FNiagaraStatelessEmitterDataBuildContext(FNiagaraDataSetCompiledData& InParticleDataSet, FNiagaraParameterStore& InRendererBindings, TArray<TPair<int32, FInstancedStruct>>& InExpressions, TArray<uint8>& InBuiltData, TArray<float>& InStaticFloatData, NiagaraStateless::FParticleSimulationExecData* InParticleExecData)
+	FNiagaraStatelessEmitterDataBuildContext(FNiagaraDataSetCompiledData& InParticleDataSet, FNiagaraParameterStore& InRendererBindings, TArray<TPair<int32, FInstancedStruct>>& InExpressions, TArray<uint8>& InBuiltData, TArray<uint8>& InStaticDataBuffer, NiagaraStateless::FParticleSimulationExecData* InParticleExecData)
 		: ParticleDataSet(InParticleDataSet)
 		, RendererBindings(InRendererBindings)
 		, Expressions(InExpressions)
 		, BuiltData(InBuiltData)
-		, StaticFloatData(InStaticFloatData)
+		, StaticDataBuffer(InStaticDataBuffer)
 		, ParticleExecData(InParticleExecData)
 	{
 	}
 
 	void PreModuleBuild(int32 InShaderParameterOffset);
 
-	uint32 AddStaticData(TConstArrayView<float> FloatData) const;
-	uint32 AddStaticData(TConstArrayView<FVector2f> FloatData) const;
-	uint32 AddStaticData(TConstArrayView<FVector3f> FloatData) const;
-	uint32 AddStaticData(TConstArrayView<FVector4f> FloatData) const;
-	uint32 AddStaticData(TConstArrayView<FLinearColor> FloatData) const;
+	uint32 AddStaticData(TConstArrayView<uint32> StaticData) const;
+	uint32 AddStaticData(TConstArrayView<int32> StaticData) const;
+	uint32 AddStaticData(TConstArrayView<float> StaticData) const;
+	uint32 AddStaticData(TConstArrayView<FVector2f> StaticData) const;
+	uint32 AddStaticData(TConstArrayView<FVector3f> StaticData) const;
+	uint32 AddStaticData(TConstArrayView<FVector4f> StaticData) const;
+	uint32 AddStaticData(TConstArrayView<FLinearColor> StaticData) const;
 
 	template<typename T>
 	T* AllocateBuiltData() const
@@ -74,75 +76,116 @@ private:
 	// The returned value is INDEX_NONE is the variables is index otherwise the offset in DWORDs
 	int32 AddExpression(const FInstancedStruct& Expression) const;
 
-public:
-	// Adds an distribution into the LUT if enabled
+private:
+	template<typename TType>
+	FNiagaraStatelessBuiltDistributionType InternalAddDistribution(ENiagaraDistributionMode Mode, TConstArrayView<TType> Values, const FVector2f& TimeRange, TType DefaultValue) const
+	{
+		using namespace NiagaraStateless;
+		
+		if (Values.Num() == 0)
+		{
+			Values = MakeArrayView(&DefaultValue, 1);
+		}
+
+		FNiagaraStatelessBuiltDistributionType BuiltDistribution = FNiagaraStatelessBuiltDistribution::GetDefault();
+
+		FNiagaraStatelessBuiltDistribution::SetBufferReadMode(BuiltDistribution, true);
+		FNiagaraStatelessBuiltDistribution::SetInterpolationMode(BuiltDistribution, true);
+		FNiagaraStatelessBuiltDistribution::SetAddressMode(BuiltDistribution, true);
+
+		// Uniform random controls if we should interpolate per channel or not
+		// ColorGradient is here as we would not want to interpolate between keys differently per chnanel like the other Uniform types
+		const bool bUniformRandom = 
+			Mode == ENiagaraDistributionMode::UniformConstant ||
+			Mode == ENiagaraDistributionMode::UniformRange ||
+			Mode == ENiagaraDistributionMode::UniformCurve ||
+			Mode == ENiagaraDistributionMode::ColorGradient;
+
+		const bool bIsRandom = 
+			Mode == ENiagaraDistributionMode::UniformRange ||
+			Mode == ENiagaraDistributionMode::NonUniformRange;
+
+		switch (Mode)
+		{
+			case ENiagaraDistributionMode::UniformConstant:
+				Values = MakeArrayView(Values.GetData(), 1);
+				break;
+
+			case ENiagaraDistributionMode::UniformRange:
+				check(Values.Num() >= 2);
+				Values = MakeArrayView(Values.GetData(), 2);
+				break;
+
+			case ENiagaraDistributionMode::Array:
+			case ENiagaraDistributionMode::Binding:
+			case ENiagaraDistributionMode::Expression:
+				checkNoEntry();
+				break;
+
+			default:
+				break;
+		}
+
+		const uint8 LookupValueMode = 0;
+
+		FNiagaraStatelessBuiltDistribution::SetUniformRandom(BuiltDistribution, bUniformRandom);
+		FNiagaraStatelessBuiltDistribution::SetLookupValueMode(BuiltDistribution, bIsRandom ? uint8(ENiagaraDistributionLookupValueMode::Random) : LookupValueMode);
+
+		FNiagaraStatelessBuiltDistribution::SetLookupParameters(BuiltDistribution, AddStaticData(Values), Values.Num(), TimeRange.X, TimeRange.Y);
+		return BuiltDistribution;
+	}
+
 	template<typename TType>
-	FNiagaraStatelessBuiltDistributionType AddDistribution(ENiagaraDistributionMode Mode, TConstArrayView<TType> Values, const FVector2f& TimeRange) const
+	FNiagaraStatelessBuiltDistributionType InternalAddDistribution(uint32 DataOffset, uint32 NumEntries, TType DefaultValue, bool bDataIsStaticBuffer) const
 	{
 		using namespace NiagaraStateless;
 
 		FNiagaraStatelessBuiltDistributionType BuiltDistribution = FNiagaraStatelessBuiltDistribution::GetDefault();
-		if (Values.Num() > 0)
+		FNiagaraStatelessBuiltDistribution::SetInterpolationMode(BuiltDistribution, false);
+		FNiagaraStatelessBuiltDistribution::SetAddressMode(BuiltDistribution, true);
+		FNiagaraStatelessBuiltDistribution::SetUniformRandom(BuiltDistribution, true);
+		FNiagaraStatelessBuiltDistribution::SetLookupValueMode(BuiltDistribution, 0);
+
+		if ( DataOffset == INDEX_NONE || NumEntries == 0 )
 		{
-			switch (Mode)
-			{
-				case ENiagaraDistributionMode::Binding:				checkNoEntry(); break;
-				case ENiagaraDistributionMode::Expression:			checkNoEntry(); break;
-				case ENiagaraDistributionMode::UniformConstant:		FNiagaraStatelessBuiltDistribution::SetIsRandom(BuiltDistribution); FNiagaraStatelessBuiltDistribution::SetIsUniform(BuiltDistribution); break;
-				case ENiagaraDistributionMode::NonUniformConstant:	FNiagaraStatelessBuiltDistribution::SetIsRandom(BuiltDistribution); break;
-				case ENiagaraDistributionMode::UniformRange:		FNiagaraStatelessBuiltDistribution::SetIsRandom(BuiltDistribution); FNiagaraStatelessBuiltDistribution::SetIsUniform(BuiltDistribution); break;
-				case ENiagaraDistributionMode::NonUniformRange:		FNiagaraStatelessBuiltDistribution::SetIsRandom(BuiltDistribution); break;
-				case ENiagaraDistributionMode::UniformCurve:		FNiagaraStatelessBuiltDistribution::SetIsUniform(BuiltDistribution); break;
-				case ENiagaraDistributionMode::NonUniformCurve:		break;
-				case ENiagaraDistributionMode::ColorGradient:		break;
-				default:											checkNoEntry(); break;
-			}
-
-			FNiagaraStatelessBuiltDistribution::SetLookupParameters(BuiltDistribution, AddStaticData(Values), Values.Num(), TimeRange);
+			FNiagaraStatelessBuiltDistribution::SetBufferReadMode(BuiltDistribution, true);
+			FNiagaraStatelessBuiltDistribution::SetLookupParameters(BuiltDistribution, AddStaticData(MakeArrayView(&DefaultValue, 1)), 1);
+		}
+		else
+		{
+			FNiagaraStatelessBuiltDistribution::SetBufferReadMode(BuiltDistribution, bDataIsStaticBuffer);
+			FNiagaraStatelessBuiltDistribution::SetLookupParameters(BuiltDistribution, DataOffset, NumEntries);
 		}
 		return BuiltDistribution;
 	}
 
+public:
 	// Adds a distribution into the LUT if enabled and returns the packed information to send to the shader
 	template<typename TDistribution>
 	FNiagaraStatelessBuiltDistributionType AddDistribution(const TDistribution& Distribution) const
 	{
-		FNiagaraStatelessBuiltDistributionType BuiltDistribution = FNiagaraStatelessBuiltDistribution::GetDefault();
-		if (Distribution.Mode == ENiagaraDistributionMode::Binding)
+		if (Distribution.Mode == ENiagaraDistributionMode::Array)
 		{
-			//-OPT: If the expression is constant should we pack into the static values data rather than going into the parameter store?
-			const int32 ParameterOffset = AddRendererBinding(Distribution.ParameterBinding);
-			if (ParameterOffset >= 0)
-			{
-				FNiagaraStatelessBuiltDistribution::SetIsBinding(BuiltDistribution);
-				FNiagaraStatelessBuiltDistribution::SetLookupParameters(BuiltDistribution, ParameterOffset);
-			}
-			else
-			{
-				const auto DefaultValues = { TDistribution::GetDefaultValue(), TDistribution::GetDefaultValue() };
-				BuiltDistribution = AddDistribution(ENiagaraDistributionMode::NonUniformConstant, MakeArrayView(DefaultValues), FVector2f::ZeroVector);
-			}
+			FNiagaraStatelessBuiltDistributionType BuiltDistribution = InternalAddDistribution(AddStaticData(Distribution.Values), Distribution.Values.Num(), TDistribution::GetDefaultValue(), true);
+
+			FNiagaraStatelessBuiltDistribution::SetInterpolationMode(BuiltDistribution, Distribution.GetInterpolationMode() == ENiagaraDistributionInterpolationMode::Linear);
+			FNiagaraStatelessBuiltDistribution::SetAddressMode(BuiltDistribution, Distribution.GetAddressMode() == ENiagaraDistributionAddressMode::Clamp);
+			FNiagaraStatelessBuiltDistribution::SetUniformRandom(BuiltDistribution, true);
+
+			const uint8 LookupValueMode = Distribution.GetLookupValueMode();
+			FNiagaraStatelessBuiltDistribution::SetLookupValueMode(BuiltDistribution, LookupValueMode);
+			FNiagaraStatelessBuiltDistribution::SetTimeScaleBias(BuiltDistribution, Distribution.ValuesTimeRange.X, Distribution.ValuesTimeRange.Y);
+			return BuiltDistribution;
 		}
-		else if (Distribution.Mode == ENiagaraDistributionMode::Expression)
+		else if (Distribution.Mode == ENiagaraDistributionMode::Binding)
 		{
-			//-OPT: If the expression is constant should we pack into the static values data rather than going into the parameter store?
-			const int32 ParameterOffset = AddExpression(Distribution.ParameterExpression);
-			if (ParameterOffset >= 0)
-			{
-				FNiagaraStatelessBuiltDistribution::SetIsBinding(BuiltDistribution);
-				FNiagaraStatelessBuiltDistribution::SetLookupParameters(BuiltDistribution, ParameterOffset);
-			}
-			else
-			{
-				const auto DefaultValues = { TDistribution::GetDefaultValue(), TDistribution::GetDefaultValue() };
-				BuiltDistribution = AddDistribution(ENiagaraDistributionMode::NonUniformConstant, MakeArrayView(DefaultValues), FVector2f::ZeroVector);
-			}
+			return InternalAddDistribution(AddRendererBinding(Distribution.ParameterBinding), 1, TDistribution::GetDefaultValue(), false);
 		}
-		else
+		else if (Distribution.Mode == ENiagaraDistributionMode::Expression)
 		{
-			BuiltDistribution = AddDistribution(Distribution.Mode, MakeArrayView(Distribution.Values), Distribution.ValuesTimeRange);
+			return InternalAddDistribution(AddExpression(Distribution.ParameterExpression), 1, TDistribution::GetDefaultValue(), false);
 		}
-		return BuiltDistribution;
+		return InternalAddDistribution(Distribution.Mode, MakeArrayView(Distribution.Values), Distribution.ValuesTimeRange, TDistribution::GetDefaultValue());
 	}
 
 	// Adds a distribution and forces it to generate as a curve for lookup
@@ -152,14 +195,13 @@ public:
 		FNiagaraStatelessBuiltDistributionType BuiltDistribution = FNiagaraStatelessBuiltDistribution::GetDefault();
 		if (ensure((Distribution.IsCurve() || Distribution.IsGradient()) && Distribution.Values.Num() > 1))
 		{
-			FNiagaraStatelessBuiltDistribution::SetLookupParameters(BuiltDistribution, AddStaticData(Distribution.Values), Distribution.Values.Num(), Distribution.ValuesTimeRange);
+			return InternalAddDistribution<TType>(Distribution.Mode, MakeArrayView(Distribution.Values), Distribution.ValuesTimeRange, DefaultValue);
 		}
 		else
 		{
 			const TType DefaultValues[] = {DefaultValue , DefaultValue};
-			FNiagaraStatelessBuiltDistribution::SetLookupParameters(BuiltDistribution, AddStaticData(DefaultValues));
+			return InternalAddDistribution(ENiagaraDistributionMode::NonUniformCurve, MakeArrayView(DefaultValues), Distribution.ValuesTimeRange, DefaultValue);
 		}
-		return BuiltDistribution;
 	}
 
 	template<typename TRange, typename TDistribution, typename TDefaultValue>
@@ -204,7 +246,7 @@ private:
 	FNiagaraParameterStore&							RendererBindings;
 	TArray<TPair<int32, FInstancedStruct>>&			Expressions;
 	TArray<uint8>&									BuiltData;
-	TArray<float>&									StaticFloatData;
+	TArray<uint8>&									StaticDataBuffer;
 	NiagaraStateless::FParticleSimulationExecData*	ParticleExecData = nullptr;
 
 	int32											ModuleBuiltDataOffset = 0;
diff --git a/Engine/Plugins/FX/Niagara/Source/Niagara/Internal/Stateless/NiagaraStatelessEmitterInstance.h b/Engine/Plugins/FX/Niagara/Source/Niagara/Internal/Stateless/NiagaraStatelessEmitterInstance.h
--- a/Engine/Plugins/FX/Niagara/Source/Niagara/Internal/Stateless/NiagaraStatelessEmitterInstance.h
+++ b/Engine/Plugins/FX/Niagara/Source/Niagara/Internal/Stateless/NiagaraStatelessEmitterInstance.h
@@ -39,7 +39,7 @@ namespace NiagaraStateless
 
 		mutable bool											bBindingBufferDirty = true;
 		mutable TArray<uint8>									BindingBufferData;
-		mutable FReadBuffer										BindingBuffer;
+		mutable FByteAddressBuffer								BindingBuffer;
 
 		// Begin: INiagaraComputeDataBufferInterface
 		virtual bool HasTranslucentDataToRender() const override { return false; }
diff --git a/Engine/Plugins/FX/Niagara/Source/Niagara/Internal/Stateless/NiagaraStatelessParticleSimContext.h b/Engine/Plugins/FX/Niagara/Source/Niagara/Internal/Stateless/NiagaraStatelessParticleSimContext.h
--- a/Engine/Plugins/FX/Niagara/Source/Niagara/Internal/Stateless/NiagaraStatelessParticleSimContext.h
+++ b/Engine/Plugins/FX/Niagara/Source/Niagara/Internal/Stateless/NiagaraStatelessParticleSimContext.h
@@ -101,7 +101,7 @@ public:
 	TType GetParameterBufferValue(int32 Offset, int32 Element) const
 	{
 		TType Value;
-		FMemory::Memcpy(&Value, &DynamicBufferData[Offset * 4], sizeof(TType));
+		FMemory::Memcpy(&Value, &DynamicBufferData[(Offset * 4) + (Element * sizeof(TType))], sizeof(TType));
 		return Value;
 	}
 
@@ -110,7 +110,7 @@ public:
 	template<typename TType>
 	TType GetParameterBufferFloat(int32 Offset, int32 Element) const { return GetParameterBufferValue<TType>(Offset, Element); }
 	template<typename TType>
-	TType GetParameterBufferFloat(int32 Offset, const TType& DefaultValue) const { return Offset >= 0 ? GetParameterBufferValue<TType>(Offset, 0) : DefaultValue; }
+	TType GetParameterBufferFloatDefaulted(int32 Offset, const TType& DefaultValue) const { return Offset >= 0 ? GetParameterBufferValue<TType>(Offset, 0) : DefaultValue; }
 
 	template<typename TType>
 	TType GetStaticFloat(uint32 Offset, uint32 Element) const;
@@ -129,19 +129,103 @@ public:
 		return FMath::Lerp(Value0, Value1, FMath::Fractional(Offset));
 	}
 
-	float			Lerp(const float& Lhs, const float& Rhs, const float& U, bool bUniform) const { return FMath::Lerp(Lhs, Rhs, U); }
-	FVector2f		Lerp(const FVector2f& Lhs, const FVector2f& Rhs, const FVector2f& U, bool bUniform) const { const FVector2f V = FMath::Lerp(Lhs, Rhs, U); return bUniform ? FVector2f(V.X, V.X) : V; }
-	FVector3f		Lerp(const FVector3f& Lhs, const FVector3f& Rhs, const FVector3f& U, bool bUniform) const { const FVector3f V = FMath::Lerp(Lhs, Rhs, U); return bUniform ? FVector3f(V.X, V.X, V.X) : V; }
-	FVector4f		Lerp(const FVector4f& Lhs, const FVector4f& Rhs, const FVector4f& U, bool bUniform) const { const FVector4f V = FMath::Lerp(Lhs, Rhs, U); return bUniform ? FVector4f(V.X, V.X, V.X, V.X) : V; }
+	float		Lerp(const float& Lhs, const float& Rhs, const float& U, bool bUniform) const { return FMath::Lerp(Lhs, Rhs, U); }
+	FVector2f	Lerp(const FVector2f& Lhs, const FVector2f& Rhs, const FVector2f& U, bool bUniform) const { return FMath::Lerp(Lhs, Rhs, bUniform ? FVector2f(U.X, U.X) : U); }
+	FVector3f	Lerp(const FVector3f& Lhs, const FVector3f& Rhs, const FVector3f& U, bool bUniform) const { return FMath::Lerp(Lhs, Rhs, bUniform ? FVector3f(U.X, U.X, U.X) : U); }
+	FVector4f	Lerp(const FVector4f& Lhs, const FVector4f& Rhs, const FVector4f& U, bool bUniform) const { return FMath::Lerp(Lhs, Rhs, bUniform ? FVector4f(U.X, U.X, U.X, U.X) : U); }
 
 	template<typename TType>
 	TType SampleCurve(const FNiagaraStatelessBuiltDistributionType& Distribution, float Time) const
 	{
-		const uint32 DataOffset = FNiagaraStatelessBuiltDistribution::GetDataOffset(Distribution);
-		const float Offset = FNiagaraStatelessBuiltDistribution::ConvertTimeToLookup(Distribution, Time);
-		const TType Value0 = GetStaticFloat<TType>(DataOffset, FMath::FloorToInt(Offset));
-		const TType Value1 = GetStaticFloat<TType>(DataOffset, FMath::CeilToInt(Offset));
-		return FMath::Lerp(Value0, Value1, FMath::Fractional(Offset));
+		const uint32 DataOffset = FNiagaraStatelessBuiltDistribution::GetBufferOffset(Distribution);
+		uint32 IndexA, IndexB;
+		float Interp;
+		FNiagaraStatelessBuiltDistribution::ConvertTimeToLookup(Distribution, Time, IndexA, IndexB, Interp);
+		const TType Value0 = GetStaticFloat<TType>(DataOffset, IndexA);
+		const TType Value1 = GetStaticFloat<TType>(DataOffset, IndexB);
+		return FMath::Lerp(Value0, Value1, Interp);
+	}
+
+private:
+	template<typename TType>
+	void GetDistributionTableValues(FNiagaraStatelessBuiltDistributionTypeIn BuiltDistribution, uint32 IndexA, uint32 IndexB, TType& OutValueA, TType& OutValueB) const
+	{
+		const uint32 BaseOffset = FNiagaraStatelessBuiltDistribution::GetBufferOffset(BuiltDistribution);
+		if (FNiagaraStatelessBuiltDistribution::IsBufferReadModeStatic(BuiltDistribution))
+		{
+			OutValueA = GetStaticFloat<TType>(BaseOffset, IndexA);
+			OutValueB = GetStaticFloat<TType>(BaseOffset, IndexB);
+		}
+		else
+		{
+			OutValueA = GetParameterBufferFloat<TType>(BaseOffset, IndexA);
+			OutValueB = GetParameterBufferFloat<TType>(BaseOffset, IndexB);
+		}
+	}
+
+	template<typename TType>
+	TType GetDistributionRandomValue(FNiagaraStatelessBuiltDistributionTypeIn BuiltDistribution, uint32 iInstance, uint32 RandomSeedOffset) const
+	{
+		const float Random = RandomFloat(iInstance, RandomSeedOffset);
+
+		uint32 Indicies[2];
+		float Interp;
+		FNiagaraStatelessBuiltDistribution::ConvertRandomToLookup(BuiltDistribution, Random, Indicies[0], Indicies[1], Interp);
+
+		TType Values[2];
+		GetDistributionTableValues(BuiltDistribution, Indicies[0], Indicies[1], Values[0], Values[1]);
+
+		if (FNiagaraStatelessBuiltDistribution::IsInterpolationModeLinear(BuiltDistribution))
+		{
+			const bool bUniformRandom = FNiagaraStatelessBuiltDistribution::IsUniformRandomEnabled(BuiltDistribution);
+			return Lerp(Values[0], Values[1], TRandomFloat<TType>(iInstance, RandomSeedOffset), bUniformRandom);
+		}
+		else
+		{
+			return Values[0];
+		}
+	}
+
+	template<typename TType>
+	TType GetDistributionTimeValue(FNiagaraStatelessBuiltDistributionTypeIn BuiltDistribution, float Time) const
+	{
+		uint32 Indicies[2];
+		float Interp;
+		FNiagaraStatelessBuiltDistribution::ConvertTimeToLookup(BuiltDistribution, Time, Indicies[0], Indicies[1], Interp);
+
+		TType Values[2];
+		GetDistributionTableValues(BuiltDistribution, Indicies[0], Indicies[1], Values[0], Values[1]);
+
+		return FMath::Lerp(Values[0], Values[1], Interp);
+	}
+
+public:
+	template<typename TType>
+	TType SampleDistributionValue(FNiagaraStatelessBuiltDistributionTypeIn BuiltDistribution, uint32 iInstance, uint32 RandomSeedOffset, float Time) const
+	{
+		if (FNiagaraStatelessBuiltDistribution::IsLookupValueModeRandom(BuiltDistribution))
+		{
+			return GetDistributionRandomValue<TType>(BuiltDistribution, iInstance, RandomSeedOffset);
+		}
+		else
+		{
+			return GetDistributionTimeValue<TType>(BuiltDistribution, Time);
+		}
+	}
+
+	template<typename TType>
+	void SampleDistributionValues(FNiagaraStatelessBuiltDistributionTypeIn BuiltDistribution, uint32 iInstance, uint32 RandomSeedOffset, float TimeA, float TimeB, TType& OutValueA, TType& OutValueB) const
+	{
+		if (FNiagaraStatelessBuiltDistribution::IsLookupValueModeRandom(BuiltDistribution))
+		{
+			OutValueA = GetDistributionRandomValue<TType>(BuiltDistribution, iInstance, RandomSeedOffset);
+			OutValueB = OutValueA;
+		}
+		else
+		{
+			OutValueA = GetDistributionTimeValue<TType>(BuiltDistribution, TimeA);
+			OutValueB = GetDistributionTimeValue<TType>(BuiltDistribution, TimeB);
+		}
 	}
 
 	template<typename TType>
@@ -251,45 +335,4 @@ private:
 	const TConstArrayView<uint8>			DynamicBufferData;
 };
 
-template<typename TType>
-struct FStatelessDistributionSampler
-{
-	explicit FStatelessDistributionSampler(const FParticleSimulationContext& ParticleSimulationContext, const FNiagaraStatelessBuiltDistributionType& InBuiltDistribution, int32 iInstance, uint32 RandomSeedOffset)
-		: BuiltDistribution(InBuiltDistribution)
-	{
-		if (FNiagaraStatelessBuiltDistribution::IsRandom(BuiltDistribution))
-		{
-			RandomOffset = ParticleSimulationContext.TRandomFloat<TType>(iInstance, RandomSeedOffset);
-		}
-	}
-
-	bool IsValid() const { return FNiagaraStatelessBuiltDistribution::IsValid(BuiltDistribution); }
-
-	TType GetValue(const FParticleSimulationContext& ParticleSimulationContext, float Time) const
-	{
-		const uint32 DataOffset = FNiagaraStatelessBuiltDistribution::GetDataOffset(BuiltDistribution);
-		if (FNiagaraStatelessBuiltDistribution::IsBinding(BuiltDistribution))
-		{
-			return ParticleSimulationContext.GetParameterBufferFloat<TType>(DataOffset, 0);
-		}
-		else if (FNiagaraStatelessBuiltDistribution::IsRandom(BuiltDistribution))
-		{
-			//-OPT: Could move into constructor
-			const TType	Value0 = ParticleSimulationContext.GetStaticFloat<TType>(DataOffset, 0);
-			const TType	Value1 = ParticleSimulationContext.GetStaticFloat<TType>(DataOffset, 1);
-			return ParticleSimulationContext.Lerp(Value0, Value1, RandomOffset, FNiagaraStatelessBuiltDistribution::IsUniform(BuiltDistribution));
-		}
-		else
-		{
-			const float Offset = FNiagaraStatelessBuiltDistribution::ConvertTimeToLookup(BuiltDistribution, Time);
-			const TType	Value0 = ParticleSimulationContext.GetStaticFloat<TType>(DataOffset, FMath::FloorToInt(Offset));
-			const TType	Value1 = ParticleSimulationContext.GetStaticFloat<TType>(DataOffset, FMath::CeilToInt(Offset));
-			return FMath::Lerp(Value0, Value1, FMath::Fractional(Offset));
-		}
-	}
-
-	const FNiagaraStatelessBuiltDistributionType	BuiltDistribution;
-	TType											RandomOffset = {};
-};
-
 } //namespace NiagaraStateless
diff --git a/Engine/Plugins/FX/Niagara/Source/Niagara/Private/NiagaraRenderer.cpp b/Engine/Plugins/FX/Niagara/Source/Niagara/Private/NiagaraRenderer.cpp
--- a/Engine/Plugins/FX/Niagara/Source/Niagara/Private/NiagaraRenderer.cpp
+++ b/Engine/Plugins/FX/Niagara/Source/Niagara/Private/NiagaraRenderer.cpp
@@ -160,6 +160,41 @@ public:
 	}
 };
 
+class FNiagaraEmptyByteAddressBufferSRV : public FRenderResource
+{
+public:
+	FNiagaraEmptyByteAddressBufferSRV(const FString& InDebugName) : DebugName(InDebugName) {}
+
+	FString DebugName;
+	FBufferRHIRef Buffer;
+	FShaderResourceViewRHIRef SRV;
+
+	virtual void InitRHI(FRHICommandListBase& RHICmdList) override
+	{
+		constexpr int32 NumBytes = sizeof(FUintVector4);
+
+		FRHIBufferCreateDesc CreateDesc =
+			FRHIBufferCreateDesc::CreateByteAddress(*DebugName, NumBytes, 0)
+			.AddUsage(EBufferUsageFlags::Static | EBufferUsageFlags::ShaderResource)
+			.SetInitialState(ERHIAccess::SRVMask)
+			.SetInitActionZeroData();
+
+		Buffer = RHICmdList.CreateBuffer(CreateDesc);
+
+		SRV = RHICmdList.CreateShaderResourceView(
+			Buffer,
+			FRHIViewDesc::CreateBufferSRV()
+				.SetType(FRHIViewDesc::EBufferType::Raw)
+		);
+	}
+
+	virtual void ReleaseRHI() override
+	{
+		SRV.SafeRelease();
+		Buffer.SafeRelease();
+	}
+};
+
 static TGlobalResource<FNiagaraEmptyBufferSRV> DummyFloatBuffer(PF_R32_FLOAT, TEXT("NiagaraRenderer::DummyFloat"));
 static TGlobalResource<FNiagaraEmptyBufferSRV> DummyFloat2Buffer(PF_G16R16F, TEXT("NiagaraRenderer::DummyFloat2"));
 static TGlobalResource<FNiagaraEmptyBufferSRV> DummyFloat4Buffer(PF_A32B32G32R32F, TEXT("NiagaraRenderer::DummyFloat4"));
@@ -169,6 +204,7 @@ static TGlobalResource<FNiagaraEmptyBufferSRV> DummyUIntBuffer(PF_R32_UINT, TEXT
 static TGlobalResource<FNiagaraEmptyBufferSRV> DummyUInt2Buffer(PF_R32G32_UINT, TEXT("NiagaraRenderer::DummyUInt2"));
 static TGlobalResource<FNiagaraEmptyBufferSRV> DummyUInt4Buffer(PF_R32G32B32A32_UINT, TEXT("NiagaraRenderer::DummyUInt4"));
 static TGlobalResource<FNiagaraEmptyBufferSRV> DummyHalfBuffer(PF_R16F, TEXT("NiagaraRenderer::DummyHalf"));
+static TGlobalResource<FNiagaraEmptyByteAddressBufferSRV> DummyByteAddressBuffer(TEXT("NiagaraRenderer::DummyByteAddressBuffer"));
 
 FRHIShaderResourceView* FNiagaraRenderer::GetDummyFloatBuffer()
 {
@@ -215,6 +251,11 @@ FRHIShaderResourceView* FNiagaraRenderer::GetDummyHalfBuffer()
 	return DummyHalfBuffer.SRV;
 }
 
+FRHIShaderResourceView* FNiagaraRenderer::GetDummyByteAddressBuffer()
+{
+	return DummyByteAddressBuffer.SRV;
+}
+
 FParticleRenderData FNiagaraRenderer::TransferDataToGPU(FRHICommandListBase& RHICmdList, FGlobalDynamicReadBuffer& DynamicReadBuffer, const FNiagaraRendererLayout* RendererLayout, TConstArrayView<uint32> IntComponents, const FNiagaraDataBuffer* SrcData)
 {
 	const int32 NumInstances = SrcData->GetNumInstances();
diff --git a/Engine/Plugins/FX/Niagara/Source/Niagara/Private/Stateless/Modules/NiagaraStatelessModule_CameraOffset.cpp b/Engine/Plugins/FX/Niagara/Source/Niagara/Private/Stateless/Modules/NiagaraStatelessModule_CameraOffset.cpp
--- a/Engine/Plugins/FX/Niagara/Source/Niagara/Private/Stateless/Modules/NiagaraStatelessModule_CameraOffset.cpp
+++ b/Engine/Plugins/FX/Niagara/Source/Niagara/Private/Stateless/Modules/NiagaraStatelessModule_CameraOffset.cpp
@@ -29,16 +29,11 @@ namespace NSMCameraOffsetPrivate
 
 		for (uint32 i = 0; i < ParticleSimulationContext.GetNumInstances(); ++i)
 		{
-			const FStatelessDistributionSampler<float> CameraOffsetSampler(ParticleSimulationContext, ModuleBuiltData->DistributionParameters, i, 0);
+			float CameraOffsets[2];
+			ParticleSimulationContext.SampleDistributionValues(ModuleBuiltData->DistributionParameters, i, 0, NormalizedAgeData[i], PreviousNormalizedAgeData[i], CameraOffsets[0], CameraOffsets[1]);
 
-			const float NormalizedAge = NormalizedAgeData[i];
-			const float PreviousNormalizedAge = PreviousNormalizedAgeData[i];
-
-			const float CameraOffset = CameraOffsetSampler.GetValue(ParticleSimulationContext, NormalizedAge);
-			const float PreviousCameraOffset = CameraOffsetSampler.GetValue(ParticleSimulationContext, PreviousNormalizedAge);
-
-			ParticleSimulationContext.WriteParticleVariable(ModuleBuiltData->CameraVariableOffset, i, CameraOffset);
-			ParticleSimulationContext.WriteParticleVariable(ModuleBuiltData->PreviousCameraVariableOffset, i, PreviousCameraOffset);
+			ParticleSimulationContext.WriteParticleVariable(ModuleBuiltData->CameraVariableOffset, i, CameraOffsets[0]);
+			ParticleSimulationContext.WriteParticleVariable(ModuleBuiltData->PreviousCameraVariableOffset, i, CameraOffsets[1]);
 		}
 	}
 }
diff --git a/Engine/Plugins/FX/Niagara/Source/Niagara/Private/Stateless/Modules/NiagaraStatelessModule_DecalAttributes.cpp b/Engine/Plugins/FX/Niagara/Source/Niagara/Private/Stateless/Modules/NiagaraStatelessModule_DecalAttributes.cpp
--- a/Engine/Plugins/FX/Niagara/Source/Niagara/Private/Stateless/Modules/NiagaraStatelessModule_DecalAttributes.cpp
+++ b/Engine/Plugins/FX/Niagara/Source/Niagara/Private/Stateless/Modules/NiagaraStatelessModule_DecalAttributes.cpp
@@ -67,22 +67,19 @@ namespace NSMDecalAttributesPrivate
 		{
 			if (ModuleBuiltData->AttributeOffset[int(EModuleAttribute::Orientation)] != INDEX_NONE)
 			{
-				FStatelessDistributionSampler<FVector3f> DistributionSampler(ParticleSimulationContext, ModuleBuiltData->AttributeDistributionParameters[int(EModuleAttribute::Orientation)], i, 0);
-				const FVector3f Rotation = DistributionSampler.GetValue(ParticleSimulationContext, NormalizedAgeData[i]);
+				const FVector3f Rotation = ParticleSimulationContext.SampleDistributionValue<FVector3f>(ModuleBuiltData->AttributeDistributionParameters[int(EModuleAttribute::Orientation)], i, 0, NormalizedAgeData[i]);
 				const FQuat4f SourceValue = FRotator3f(Rotation.X, Rotation.Y, Rotation.Z).Quaternion();
 				const FQuat4f Value = OrientationRotation * SourceValue;
 				ParticleSimulationContext.WriteParticleVariable(ModuleBuiltData->AttributeOffset[int(EModuleAttribute::Orientation)], i, Value);
 			}
 			if (ModuleBuiltData->AttributeOffset[int(EModuleAttribute::Size)] != INDEX_NONE)
 			{
-				FStatelessDistributionSampler<FVector3f> DistributionSampler(ParticleSimulationContext, ModuleBuiltData->AttributeDistributionParameters[int(EModuleAttribute::Size)], i, 0);
-				const FVector3f Value = DistributionSampler.GetValue(ParticleSimulationContext, NormalizedAgeData[i]);
+				const FVector3f Value = ParticleSimulationContext.SampleDistributionValue<FVector3f>(ModuleBuiltData->AttributeDistributionParameters[int(EModuleAttribute::Size)], i, 0, NormalizedAgeData[i]);
 				ParticleSimulationContext.WriteParticleVariable(ModuleBuiltData->AttributeOffset[int(EModuleAttribute::Size)], i, Value);
 			}
 			if (ModuleBuiltData->AttributeOffset[int(EModuleAttribute::Fade)] != INDEX_NONE)
 			{
-				FStatelessDistributionSampler<float> DistributionSampler(ParticleSimulationContext, ModuleBuiltData->AttributeDistributionParameters[int(EModuleAttribute::Fade)], i, 0);
-				const float Value = DistributionSampler.GetValue(ParticleSimulationContext, NormalizedAgeData[i]);
+				const float Value = ParticleSimulationContext.SampleDistributionValue<float>(ModuleBuiltData->AttributeDistributionParameters[int(EModuleAttribute::Fade)], i, 0, NormalizedAgeData[i]);
 				ParticleSimulationContext.WriteParticleVariable(ModuleBuiltData->AttributeOffset[int(EModuleAttribute::Fade)], i, Value);
 			}
 		}
diff --git a/Engine/Plugins/FX/Niagara/Source/Niagara/Private/Stateless/Modules/NiagaraStatelessModule_DynamicMaterialParameters.cpp b/Engine/Plugins/FX/Niagara/Source/Niagara/Private/Stateless/Modules/NiagaraStatelessModule_DynamicMaterialParameters.cpp
--- a/Engine/Plugins/FX/Niagara/Source/Niagara/Private/Stateless/Modules/NiagaraStatelessModule_DynamicMaterialParameters.cpp
+++ b/Engine/Plugins/FX/Niagara/Source/Niagara/Private/Stateless/Modules/NiagaraStatelessModule_DynamicMaterialParameters.cpp
@@ -53,18 +53,17 @@ namespace NSMDynamicMaterialParametersPrivate
 				}
 
 				const int32 FirstChannel = iParameter * NumChannelPerParameter;
-				const FStatelessDistributionSampler<float> XChannelSampler(ParticleSimulationContext, ModuleBuiltData->ParameterDistributions[FirstChannel + 0], i, FirstChannel + 0);
-				const FStatelessDistributionSampler<float> YChannelSampler(ParticleSimulationContext, ModuleBuiltData->ParameterDistributions[FirstChannel + 1], i, FirstChannel + 1);
-				const FStatelessDistributionSampler<float> ZChannelSampler(ParticleSimulationContext, ModuleBuiltData->ParameterDistributions[FirstChannel + 2], i, FirstChannel + 2);
-				const FStatelessDistributionSampler<float> WChannelSampler(ParticleSimulationContext, ModuleBuiltData->ParameterDistributions[FirstChannel + 3], i, FirstChannel + 3);
-
 				const float NormalizedAge = NormalizedAgeData[i];
-				const FVector4f DynamicParameter(
-					((ChannelMask & 0x1) != 0) ? XChannelSampler.GetValue(ParticleSimulationContext, NormalizedAge) : 0.0f,
-					((ChannelMask & 0x2) != 0) ? YChannelSampler.GetValue(ParticleSimulationContext, NormalizedAge) : 0.0f,
-					((ChannelMask & 0x4) != 0) ? ZChannelSampler.GetValue(ParticleSimulationContext, NormalizedAge) : 0.0f,
-					((ChannelMask & 0x8) != 0) ? WChannelSampler.GetValue(ParticleSimulationContext, NormalizedAge) : 0.0f
-				);
+
+				FVector4f DynamicParameter;
+				for (int32 iChannel = 0; iChannel < NumChannelPerParameter; ++iChannel)
+				{
+					DynamicParameter[iChannel] = 0.0f;
+					if ((ChannelMask & (1 << iChannel)) != 0 )
+					{
+						DynamicParameter[iChannel] = ParticleSimulationContext.SampleDistributionValue<float>(ModuleBuiltData->ParameterDistributions[FirstChannel + iChannel], i, FirstChannel + iChannel, NormalizedAge);
+					}
+				}
 
 				ParticleSimulationContext.WriteParticleVariable(ModuleBuiltData->ParameterVariableOffsets[iParameter], i, DynamicParameter);
 			}
@@ -169,7 +168,15 @@ void UNiagaraStatelessModule_DynamicMaterialParameters::BuildEmitterData(const F
 			BuiltData->ChannelMask |= ParameterChannelMask << (iParameter * NumChannelPerParameter);
 			for (int32 iChannel = 0; iChannel < NumChannelPerParameter; ++iChannel)
 			{
-				BuiltData->ParameterDistributions[(iParameter * NumChannelPerParameter) + iChannel] = BuildContext.AddDistribution(GetParameterChannelDistribution(iParameter, iChannel));
+				FUintVector3& BuiltDistribution = BuiltData->ParameterDistributions[(iParameter * NumChannelPerParameter) + iChannel];
+				if ((ParameterChannelMask & (1 << iChannel)) != 0)
+				{
+					BuiltDistribution = BuildContext.AddDistribution(GetParameterChannelDistribution(iParameter, iChannel));
+				}
+				else
+				{
+					BuiltDistribution = BuildContext.AddDistribution(FNiagaraDistributionFloat(0.0f));
+				}
 			}
 		}
 
diff --git a/Engine/Plugins/FX/Niagara/Source/Niagara/Private/Stateless/Modules/NiagaraStatelessModule_InitializeParticle.cpp b/Engine/Plugins/FX/Niagara/Source/Niagara/Private/Stateless/Modules/NiagaraStatelessModule_InitializeParticle.cpp
--- a/Engine/Plugins/FX/Niagara/Source/Niagara/Private/Stateless/Modules/NiagaraStatelessModule_InitializeParticle.cpp
+++ b/Engine/Plugins/FX/Niagara/Source/Niagara/Private/Stateless/Modules/NiagaraStatelessModule_InitializeParticle.cpp
@@ -57,9 +57,7 @@ namespace NSMInitializeParticlePrivate
 
 		for (uint32 i = 0; i < ParticleSimulationContext.GetNumInstances(); ++i)
 		{
-			const FStatelessDistributionSampler<FVector3f> PositionSampler(ParticleSimulationContext, ModuleBuiltData->InitialPosition, i, 0);
-
-			const FVector3f		Position = PositionSampler.GetValue(ParticleSimulationContext, 0.0f);
+			const FVector3f		Position = ParticleSimulationContext.SampleDistributionValue<FVector3f>(ModuleBuiltData->InitialPosition, i, 0, 0.0f);
 			const FLinearColor	Color = ParticleSimulationContext.RandomScaleBiasFloat(i, 1, ShaderParameters->InitializeParticle_ColorScale, ShaderParameters->InitializeParticle_ColorBias);
 			const float			RibbonWidth = ParticleSimulationContext.RandomScaleBiasFloat(i, 2, ShaderParameters->InitializeParticle_RibbonWidthScale, ShaderParameters->InitializeParticle_RibbonWidthBias);
 			const FVector2f		SpriteSize = ParticleSimulationContext.RandomScaleBiasFloat(i, 3, ShaderParameters->InitializeParticle_SpriteSizeScale, ShaderParameters->InitializeParticle_SpriteSizeBias, bUniformSpriteSize);
@@ -80,7 +78,6 @@ namespace NSMInitializeParticlePrivate
 			ParticleSimulationContext.WriteParticleVariable(ModuleBuiltData->PreviousScaleVariableOffset, i, Scale);
 		}
 	}
-
 }
 
 void UNiagaraStatelessModule_InitializeParticle::BuildEmitterData(const FNiagaraStatelessEmitterDataBuildContext& BuildContext) const
diff --git a/Engine/Plugins/FX/Niagara/Source/Niagara/Private/Stateless/Modules/NiagaraStatelessModule_LightAttributes.cpp b/Engine/Plugins/FX/Niagara/Source/Niagara/Private/Stateless/Modules/NiagaraStatelessModule_LightAttributes.cpp
--- a/Engine/Plugins/FX/Niagara/Source/Niagara/Private/Stateless/Modules/NiagaraStatelessModule_LightAttributes.cpp
+++ b/Engine/Plugins/FX/Niagara/Source/Niagara/Private/Stateless/Modules/NiagaraStatelessModule_LightAttributes.cpp
@@ -84,8 +84,7 @@ namespace NSMLightAttributesPrivate
 			{
 				if (ModuleBuiltData->AttributeOffset[a] != INDEX_NONE)
 				{
-					FStatelessDistributionSampler<float> DistributionSampler(ParticleSimulationContext, ModuleBuiltData->AttributeDistributionParameters[a], i, 0);
-					const float Value = DistributionSampler.GetValue(ParticleSimulationContext, NormalizedAgeData[i]);
+					const float Value = ParticleSimulationContext.SampleDistributionValue<float>(ModuleBuiltData->AttributeDistributionParameters[a], i, 0, NormalizedAgeData[i]);
 					ParticleSimulationContext.WriteParticleVariable(ModuleBuiltData->AttributeOffset[a], i, Value);
 				}
 			}
diff --git a/Engine/Plugins/FX/Niagara/Source/Niagara/Private/Stateless/Modules/NiagaraStatelessModule_MeshRotationRate.cpp b/Engine/Plugins/FX/Niagara/Source/Niagara/Private/Stateless/Modules/NiagaraStatelessModule_MeshRotationRate.cpp
--- a/Engine/Plugins/FX/Niagara/Source/Niagara/Private/Stateless/Modules/NiagaraStatelessModule_MeshRotationRate.cpp
+++ b/Engine/Plugins/FX/Niagara/Source/Niagara/Private/Stateless/Modules/NiagaraStatelessModule_MeshRotationRate.cpp
@@ -72,11 +72,11 @@ void UNiagaraStatelessModule_MeshRotationRate::BuildEmitterData(const FNiagaraSt
 		BuiltData->RotationRange.Max *= 1.0f / 360.0f;
 		if (bUseRateScale)
 		{
-			BuiltData->RateScaleParameters = BuildContext.AddDistributionAsCurve(RateScaleDistribution, 1.0f);
+			BuiltData->RateScaleParameters = BuildContext.AddDistributionAsCurve(RateScaleDistribution, FVector3f::OneVector);
 		}
 		else
 		{
-			BuiltData->RateScaleParameters = BuildContext.AddDistributionAsCurve(FNiagaraDistributionCurveVector3(ENiagaraDistributionCurveLUTMode::Accumulate), 1.0f);
+			BuiltData->RateScaleParameters = BuildContext.AddDistributionAsCurve(FNiagaraDistributionCurveVector3(ENiagaraDistributionCurveLUTMode::Accumulate), FVector3f::OneVector);
 		}
 
 		BuildContext.AddParticleSimulationExecSimulate(&ParticleSimulate);
diff --git a/Engine/Plugins/FX/Niagara/Source/Niagara/Private/Stateless/Modules/NiagaraStatelessModule_ScaleColor.cpp b/Engine/Plugins/FX/Niagara/Source/Niagara/Private/Stateless/Modules/NiagaraStatelessModule_ScaleColor.cpp
--- a/Engine/Plugins/FX/Niagara/Source/Niagara/Private/Stateless/Modules/NiagaraStatelessModule_ScaleColor.cpp
+++ b/Engine/Plugins/FX/Niagara/Source/Niagara/Private/Stateless/Modules/NiagaraStatelessModule_ScaleColor.cpp
@@ -26,11 +26,9 @@ namespace NSMScaleColorPrivate
 
 		for (uint32 i = 0; i < ParticleSimulationContext.GetNumInstances(); ++i)
 		{
-			FStatelessDistributionSampler<FLinearColor> ColorSampler(ParticleSimulationContext, ModuleBuiltData->DistributionParameters, i, 0);
-
 			FLinearColor Color = ParticleSimulationContext.ReadParticleVariable(ModuleBuiltData->ColorVariableOffset, i, FLinearColor::White);
 
-			Color *= ColorSampler.GetValue(ParticleSimulationContext, NormalizedAgeData[i]);
+			Color *= ParticleSimulationContext.SampleDistributionValue<FLinearColor>(ModuleBuiltData->DistributionParameters, i, 0, NormalizedAgeData[i]);
 
 			ParticleSimulationContext.WriteParticleVariable(ModuleBuiltData->ColorVariableOffset, i, Color);
 		}
diff --git a/Engine/Plugins/FX/Niagara/Source/Niagara/Private/Stateless/Modules/NiagaraStatelessModule_ScaleMeshSize.cpp b/Engine/Plugins/FX/Niagara/Source/Niagara/Private/Stateless/Modules/NiagaraStatelessModule_ScaleMeshSize.cpp
--- a/Engine/Plugins/FX/Niagara/Source/Niagara/Private/Stateless/Modules/NiagaraStatelessModule_ScaleMeshSize.cpp
+++ b/Engine/Plugins/FX/Niagara/Source/Niagara/Private/Stateless/Modules/NiagaraStatelessModule_ScaleMeshSize.cpp
@@ -29,20 +29,18 @@ namespace NSMScaleMeshSizePrivate
 		const float* NormalizedAgeData = ParticleSimulationContext.GetParticleNormalizedAge();
 		const float* PreviousNormalizedAgeData = ParticleSimulationContext.GetParticlePreviousNormalizedAge();
 
-		const FVector3f CurveScale = ParticleSimulationContext.GetParameterBufferFloat(ModuleBuiltData->CurveScaleOffset, ModuleBuiltData->CurveScale);
+		const FVector3f CurveScale = ParticleSimulationContext.GetParameterBufferFloatDefaulted(ModuleBuiltData->CurveScaleOffset, ModuleBuiltData->CurveScale);
 
 		for (uint32 i = 0; i < ParticleSimulationContext.GetNumInstances(); ++i)
 		{
-			FStatelessDistributionSampler<FVector3f> ScaleSampler(ParticleSimulationContext, ModuleBuiltData->DistributionParameters, i, 0);
+			const FVector3f MeshScale = ParticleSimulationContext.ReadParticleVariable(ModuleBuiltData->ScaleVariableOffset, i, FVector3f::OneVector);
+			const FVector3f PrevMeshScale = ParticleSimulationContext.ReadParticleVariable(ModuleBuiltData->PreviousScaleVariableOffset, i, FVector3f::OneVector);
 
-			FVector3f Scale = ParticleSimulationContext.ReadParticleVariable(ModuleBuiltData->ScaleVariableOffset, i, FVector3f::OneVector);
-			FVector3f PreviousScale = ParticleSimulationContext.ReadParticleVariable(ModuleBuiltData->PreviousScaleVariableOffset, i, FVector3f::OneVector);
+			FVector3f Scales[2];
+			ParticleSimulationContext.SampleDistributionValues(ModuleBuiltData->DistributionParameters, i, 0, NormalizedAgeData[i], PreviousNormalizedAgeData[i], Scales[0], Scales[1]);
 
-			Scale = Scale * ScaleSampler.GetValue(ParticleSimulationContext, NormalizedAgeData[i]) * CurveScale;
-			PreviousScale = PreviousScale * ScaleSampler.GetValue(ParticleSimulationContext, PreviousNormalizedAgeData[i]) * CurveScale;
-
-			ParticleSimulationContext.WriteParticleVariable(ModuleBuiltData->ScaleVariableOffset, i, Scale);
-			ParticleSimulationContext.WriteParticleVariable(ModuleBuiltData->PreviousScaleVariableOffset, i, PreviousScale);
+			ParticleSimulationContext.WriteParticleVariable(ModuleBuiltData->ScaleVariableOffset, i, Scales[0] * CurveScale * MeshScale);
+			ParticleSimulationContext.WriteParticleVariable(ModuleBuiltData->PreviousScaleVariableOffset, i, Scales[1] * CurveScale * PrevMeshScale);
 		}
 	}
 }
diff --git a/Engine/Plugins/FX/Niagara/Source/Niagara/Private/Stateless/Modules/NiagaraStatelessModule_ScaleRibbonWidth.cpp b/Engine/Plugins/FX/Niagara/Source/Niagara/Private/Stateless/Modules/NiagaraStatelessModule_ScaleRibbonWidth.cpp
--- a/Engine/Plugins/FX/Niagara/Source/Niagara/Private/Stateless/Modules/NiagaraStatelessModule_ScaleRibbonWidth.cpp
+++ b/Engine/Plugins/FX/Niagara/Source/Niagara/Private/Stateless/Modules/NiagaraStatelessModule_ScaleRibbonWidth.cpp
@@ -27,24 +27,20 @@ namespace NSMScaleRibbonWidthPrivate
 
 		const FModuleBuiltData* ModuleBuiltData = ParticleSimulationContext.ReadBuiltData<FModuleBuiltData>();
 
-		const float ScaleFactor = ParticleSimulationContext.GetParameterBufferFloat(ModuleBuiltData->CurveScaleOffset, ModuleBuiltData->CurveScale);
+		const float ScaleFactor = ParticleSimulationContext.GetParameterBufferFloatDefaulted(ModuleBuiltData->CurveScaleOffset, ModuleBuiltData->CurveScale);
 		const float* NormalizedAgeData = ParticleSimulationContext.GetParticleNormalizedAge();
 		const float* PreviousNormalizedAgeData = ParticleSimulationContext.GetParticlePreviousNormalizedAge();
 
 		for (uint32 i = 0; i < ParticleSimulationContext.GetNumInstances(); ++i)
 		{
-			const FStatelessDistributionSampler<float> SpriteScaleSampler(ParticleSimulationContext, ModuleBuiltData->DistributionParameters, i, 0);
-
-			const float NormalizedAge = NormalizedAgeData[i];
-			const float PreviousNormalizedAge = PreviousNormalizedAgeData[i];
-
 			const float RibbonWidth = ParticleSimulationContext.ReadParticleVariable(ModuleBuiltData->RibbonWidthVariableOffset, i, 0.0f);
 			const float PreviousRibbonWidth = ParticleSimulationContext.ReadParticleVariable(ModuleBuiltData->PreviousRibbonWidthVariableOffset, i, 0.0f);
-			const float Scale = SpriteScaleSampler.GetValue(ParticleSimulationContext, NormalizedAge) * ScaleFactor;
-			const float PreviousScale = SpriteScaleSampler.GetValue(ParticleSimulationContext, PreviousNormalizedAge) * ScaleFactor;
 
-			ParticleSimulationContext.WriteParticleVariable(ModuleBuiltData->RibbonWidthVariableOffset, i, RibbonWidth * Scale);
-			ParticleSimulationContext.WriteParticleVariable(ModuleBuiltData->PreviousRibbonWidthVariableOffset, i, PreviousRibbonWidth * PreviousScale);
+			float Scales[2];
+			ParticleSimulationContext.SampleDistributionValues(ModuleBuiltData->DistributionParameters, i, 0, NormalizedAgeData[i], PreviousNormalizedAgeData[i], Scales[0], Scales[1]);
+
+			ParticleSimulationContext.WriteParticleVariable(ModuleBuiltData->RibbonWidthVariableOffset, i, RibbonWidth * ScaleFactor * Scales[0]);
+			ParticleSimulationContext.WriteParticleVariable(ModuleBuiltData->PreviousRibbonWidthVariableOffset, i, PreviousRibbonWidth * ScaleFactor * Scales[1]);
 		}
 	}
 }
diff --git a/Engine/Plugins/FX/Niagara/Source/Niagara/Private/Stateless/Modules/NiagaraStatelessModule_ScaleSpriteSize.cpp b/Engine/Plugins/FX/Niagara/Source/Niagara/Private/Stateless/Modules/NiagaraStatelessModule_ScaleSpriteSize.cpp
--- a/Engine/Plugins/FX/Niagara/Source/Niagara/Private/Stateless/Modules/NiagaraStatelessModule_ScaleSpriteSize.cpp
+++ b/Engine/Plugins/FX/Niagara/Source/Niagara/Private/Stateless/Modules/NiagaraStatelessModule_ScaleSpriteSize.cpp
@@ -28,24 +28,20 @@ namespace NSMScaleSpriteSizePrivate
 
 		const FModuleBuiltData* ModuleBuiltData = ParticleSimulationContext.ReadBuiltData<FModuleBuiltData>();
 
-		const FVector2f ScaleFactor = ParticleSimulationContext.GetParameterBufferFloat(ModuleBuiltData->CurveScaleOffset, ModuleBuiltData->CurveScale);
+		const FVector2f ScaleFactor = ParticleSimulationContext.GetParameterBufferFloatDefaulted(ModuleBuiltData->CurveScaleOffset, ModuleBuiltData->CurveScale);
 		const float* NormalizedAgeData = ParticleSimulationContext.GetParticleNormalizedAge();
 		const float* PreviousNormalizedAgeData = ParticleSimulationContext.GetParticlePreviousNormalizedAge();
 
 		for (uint32 i = 0; i < ParticleSimulationContext.GetNumInstances(); ++i)
 		{
-			const FStatelessDistributionSampler<FVector2f> SpriteScaleSampler(ParticleSimulationContext, ModuleBuiltData->DistributionParameters, i, 0);
-
-			const float NormalizedAge = NormalizedAgeData[i];
-			const float PreviousNormalizedAge = PreviousNormalizedAgeData[i];
-
 			const FVector2f SpriteSize = ParticleSimulationContext.ReadParticleVariable(ModuleBuiltData->SpriteSizeVariableOffset, i, FVector2f::ZeroVector);
 			const FVector2f PreviousSpriteSize = ParticleSimulationContext.ReadParticleVariable(ModuleBuiltData->PreviousSpriteSizeVariableOffset, i, FVector2f::ZeroVector);
-			const FVector2f Scale = SpriteScaleSampler.GetValue(ParticleSimulationContext, NormalizedAge) * ScaleFactor;
-			const FVector2f PreviousScale = SpriteScaleSampler.GetValue(ParticleSimulationContext, PreviousNormalizedAge) * ScaleFactor;
 
-			ParticleSimulationContext.WriteParticleVariable(ModuleBuiltData->SpriteSizeVariableOffset, i, SpriteSize * Scale);
-			ParticleSimulationContext.WriteParticleVariable(ModuleBuiltData->PreviousSpriteSizeVariableOffset, i, PreviousSpriteSize * PreviousScale);
+			FVector2f Scales[2];
+			ParticleSimulationContext.SampleDistributionValues(ModuleBuiltData->DistributionParameters, i, 0, NormalizedAgeData[i], PreviousNormalizedAgeData[i], Scales[0], Scales[1]);
+
+			ParticleSimulationContext.WriteParticleVariable(ModuleBuiltData->SpriteSizeVariableOffset, i, SpriteSize * Scales[0]);
+			ParticleSimulationContext.WriteParticleVariable(ModuleBuiltData->PreviousSpriteSizeVariableOffset, i, PreviousSpriteSize * Scales[1]);
 		}
 	}
 }
diff --git a/Engine/Plugins/FX/Niagara/Source/Niagara/Private/Stateless/NiagaraStatelessComputeManager.cpp b/Engine/Plugins/FX/Niagara/Source/Niagara/Private/Stateless/NiagaraStatelessComputeManager.cpp
--- a/Engine/Plugins/FX/Niagara/Source/Niagara/Private/Stateless/NiagaraStatelessComputeManager.cpp
+++ b/Engine/Plugins/FX/Niagara/Source/Niagara/Private/Stateless/NiagaraStatelessComputeManager.cpp
@@ -143,7 +143,7 @@ namespace NiagaraStatelessComputeManagerPrivate
 
 					if (EmitterInstance->BindingBufferData.Num())
 					{
-						EmitterInstance->BindingBuffer.Initialize(RHICmdList, TEXT("FNiagaraStatelessEmitterInstance::BindingBuffer"), sizeof(uint32), EmitterInstance->BindingBufferData.Num() / sizeof(uint32), EPixelFormat::PF_R32_UINT, EBufferUsageFlags::Static);
+						EmitterInstance->BindingBuffer.Initialize(RHICmdList, TEXT("FNiagaraStatelessEmitterInstance::BindingBuffer"), EmitterInstance->BindingBufferData.Num(), EBufferUsageFlags::Static);
 						void* LockedBuffer = RHICmdList.LockBuffer(EmitterInstance->BindingBuffer.Buffer, 0, EmitterInstance->BindingBuffer.NumBytes, RLM_WriteOnly);
 						FMemory::Memcpy(LockedBuffer, EmitterInstance->BindingBufferData.GetData(), EmitterInstance->BindingBuffer.NumBytes);
 						RHICmdList.UnlockBuffer(EmitterInstance->BindingBuffer.Buffer);
@@ -176,8 +176,8 @@ namespace NiagaraStatelessComputeManagerPrivate
 			//ShaderParameters->Common_HalfOutputBuffer		= DestinationData->GetGPUBufferHalf().UAV;
 			ShaderParameters->Common_IntOutputBuffer = DestinationData->GetGPUBufferInt().UAV.IsValid() ? DestinationData->GetGPUBufferInt().UAV.GetReference() : EmptyIntBufferUAV;
 			ShaderParameters->Common_GPUCountBuffer = CountBufferUAV;
-			ShaderParameters->Common_StaticFloatBuffer = EmitterData->StaticFloatBuffer.SRV;
-			ShaderParameters->Common_ParameterBuffer = FNiagaraRenderer::GetSrvOrDefaultUInt(EmitterInstance->BindingBuffer.SRV);
+			ShaderParameters->Common_StaticDataBuffer = EmitterData->StaticDataBufferGpu.SRV;
+			ShaderParameters->Common_DynamicDataBuffer = FNiagaraRenderer::GetSrvOrDefaultByteAddress(EmitterInstance->BindingBuffer.SRV);
 
 			// Execute the simulation
 			TShaderRef<NiagaraStateless::FSimulationShader> ComputeShader = EmitterData->GetShader();
diff --git a/Engine/Plugins/FX/Niagara/Source/Niagara/Private/Stateless/NiagaraStatelessDistribution.cpp b/Engine/Plugins/FX/Niagara/Source/Niagara/Private/Stateless/NiagaraStatelessDistribution.cpp
--- a/Engine/Plugins/FX/Niagara/Source/Niagara/Private/Stateless/NiagaraStatelessDistribution.cpp
+++ b/Engine/Plugins/FX/Niagara/Source/Niagara/Private/Stateless/NiagaraStatelessDistribution.cpp
@@ -230,6 +230,22 @@ namespace NiagaraStatelessDistributionPrivate
 		const ENiagaraDistributionCurveLUTMode LUTMode = ENiagaraDistributionCurveLUTMode::Sample;
 		switch(InMode)
 		{
+			case ENiagaraDistributionMode::Array:
+				if ((InChannelConstantsAndRanges.Num() % InChannelCount) == 0)
+				{
+					const int32 NumEntries = InChannelConstantsAndRanges.Num() / InChannelCount;
+
+					InContainerNum(NumEntries);
+					OutTimeRange = FVector2f(0.0f, 1.0f);
+					for (int32 iEntry=0; iEntry < NumEntries; ++iEntry)
+					{
+						for (int32 iChannel=0; iChannel < InChannelCount; ++iChannel)
+						{
+							InValueAndChannelAccessor(iEntry, iChannel) = InChannelConstantsAndRanges[(iEntry * InChannelCount) + iChannel];
+						}
+					}
+				}
+				break;
 			case ENiagaraDistributionMode::UniformConstant:
 				if (InChannelConstantsAndRanges.Num() >= 1)
 				{
@@ -329,6 +345,37 @@ namespace NiagaraStatelessDistributionPrivate
 			}
 		}
 	}
+
+	void UpdateDistributionTimeRange(ENiagaraDistributionTimeRangeMode TimeRangeMode, const FVector2f& CustomTimeRange, int32 NumEntries, FVector2f& OutTimeRange)
+	{
+		switch (TimeRangeMode)
+		{
+			// The 'time' value passed in maps directly into the array
+			// i.e. Index = (Time * 1) + 0
+			case ENiagaraDistributionTimeRangeMode::Direct:
+				OutTimeRange.X = 1.0f;
+				OutTimeRange.Y = 0.0f;
+				break;
+
+			// The 'time' value passed in needs to expand to the array range
+			// i.e. Index = (Time * (NumEntries - 1) + 0
+			case ENiagaraDistributionTimeRangeMode::Normalized:
+				OutTimeRange.X = float(NumEntries - 1);
+				OutTimeRange.Y = 0.0f;
+				break;
+
+			// The 'time' value passed in covers a user defined range, copy directly
+			case ENiagaraDistributionTimeRangeMode::Custom:
+				OutTimeRange = CustomTimeRange;
+				break;
+
+			default:
+				checkNoEntry();
+				OutTimeRange = FVector2f::ZeroVector;
+				break;
+		}
+	}
+
 #endif
 } //NiagaraStatelessCommon
 
@@ -684,17 +731,35 @@ void FNiagaraDistributionBase::ForEachParameterBinding(TFunction<void(const FNia
 	}
 }
 
+const UEnum* FNiagaraDistributionBase::GetLookupValueModeEnum(FProperty* Property)
+{
+	const UEnum* LookupEnum = nullptr;
+	if (ensure(Property))
+	{
+		const FName NAME_DistributionLookupValueEnumPath("DistributionLookupValueEnumPath");
+		const FString& EnumPath = Property->GetMetaData(NAME_DistributionLookupValueEnumPath);
+		if (!EnumPath.IsEmpty())
+		{
+			LookupEnum = FindObject<const UEnum>(FTopLevelAssetPath(EnumPath));
+			checkf(LookupEnum != nullptr, TEXT("Failed to find enum '%s' from property"), *EnumPath);
+		}
+	}
+	return LookupEnum ? LookupEnum : StaticEnum<ENiagaraDistributionLookupValueMode>();
+}
+
 void FNiagaraDistributionBase::PostEditChangeProperty(UObject* OwnerObject, FPropertyChangedEvent& PropertyChangedEvent)
 {
 	FStructProperty* StructProperty = CastField<FStructProperty>(PropertyChangedEvent.Property);
 	if (StructProperty && StructProperty->Struct && StructProperty->Struct->IsChildOf(FNiagaraDistributionBase::StaticStruct()))
 	{
+		FProperty* ValueProperty = StructProperty;
 		FNiagaraDistributionBase* ValuePtr = nullptr;
 		if (PropertyChangedEvent.Property != PropertyChangedEvent.MemberProperty)
 		{
 			// Properties stored in a UStruct inside a UObject need to first offset from UObject -> UStruct then UStruct -> Property
 			FStructProperty* MemberStructProperty = CastField<FStructProperty>(PropertyChangedEvent.MemberProperty);
 			void* StructPtr = MemberStructProperty->ContainerPtrToValuePtr<void>(OwnerObject);
+			ValueProperty = MemberStructProperty;
 			ValuePtr = StructProperty->ContainerPtrToValuePtr<FNiagaraDistributionBase>(StructPtr);
 		}
 		else
@@ -705,6 +770,36 @@ void FNiagaraDistributionBase::PostEditChangeProperty(UObject* OwnerObject, FPro
 		if (ValuePtr != nullptr)
 		{
 			ValuePtr->UpdateValuesFromDistribution();
+
+			// Arrays can modify the time range packed inside the distribution, ensure we update that information here.
+			// Some of this will be extended to curves in the future.
+			if (ValuePtr->Mode == ENiagaraDistributionMode::Array)
+			{
+				ENiagaraDistributionTimeRangeMode TimeRangeMode = ENiagaraDistributionTimeRangeMode::Normalized;
+
+				if (const UEnum* LookupEnum = FNiagaraDistributionBase::GetLookupValueModeEnum(ValueProperty))
+				{
+					const int32 ValueIndex = LookupEnum->GetIndexByValue(ValuePtr->GetLookupValueMode());
+					const FString TimeRangeModeStr = ValueIndex != INDEX_NONE ? LookupEnum->GetMetaData(TEXT("TimeRangeMode"), ValueIndex) : FString();
+					if (!TimeRangeModeStr.IsEmpty())
+					{
+						const int64 TimeRangeModeValue = StaticEnum<ENiagaraDistributionTimeRangeMode>()->GetValueByNameString(TimeRangeModeStr);
+						if (ensureMsgf(TimeRangeModeValue != INDEX_NONE, TEXT("Time Range metadata(%s) on Enum(%s) is invalid."), *TimeRangeModeStr, *LookupEnum->GetName()))
+						{
+							TimeRangeMode = ENiagaraDistributionTimeRangeMode(TimeRangeModeValue);
+						}
+					}
+				}
+				ValuePtr->UpdateTimeRangeFromDistribution(TimeRangeMode);
+			}
+			else
+			{
+				// Set back to defaults if the mode is not an array
+				// If/when we support some of this with curves, i.e. address mode / lookup value mode, we can remove this but we don't want value to leak from one mode to the next
+				ValuePtr->SetInterpolationMode(ENiagaraDistributionInterpolationMode::Linear);
+				ValuePtr->SetAddressMode(ENiagaraDistributionAddressMode::Clamp);
+				ValuePtr->SetLookupValueMode(uint8(ENiagaraDistributionLookupValueMode::ParticlesNormalizedAge));
+			}
 		}
 	}
 }
@@ -808,6 +903,16 @@ void FNiagaraDistributionFloat::UpdateValuesFromDistribution()
 	);
 }
 
+void FNiagaraDistributionFloat::UpdateTimeRangeFromDistribution(ENiagaraDistributionTimeRangeMode TimeRangeMode, const FVector2f& CustomTimeRange)
+{
+	NiagaraStatelessDistributionPrivate::UpdateDistributionTimeRange(
+		TimeRangeMode,
+		CustomTimeRange,
+		Values.Num(),
+		ValuesTimeRange
+	);
+}
+
 void FNiagaraDistributionVector2::UpdateValuesFromDistribution()
 {
 	NiagaraStatelessDistributionPrivate::UpdateDistributionValues(
@@ -822,6 +927,16 @@ void FNiagaraDistributionVector2::UpdateValuesFromDistribution()
 	);
 }
 
+void FNiagaraDistributionVector2::UpdateTimeRangeFromDistribution(ENiagaraDistributionTimeRangeMode TimeRangeMode, const FVector2f& CustomTimeRange)
+{
+	NiagaraStatelessDistributionPrivate::UpdateDistributionTimeRange(
+		TimeRangeMode,
+		CustomTimeRange,
+		Values.Num(),
+		ValuesTimeRange
+	);
+}
+
 void FNiagaraDistributionVector3::UpdateValuesFromDistribution()
 {
 	NiagaraStatelessDistributionPrivate::UpdateDistributionValues(
@@ -836,6 +951,16 @@ void FNiagaraDistributionVector3::UpdateValuesFromDistribution()
 	);
 }
 
+void FNiagaraDistributionVector3::UpdateTimeRangeFromDistribution(ENiagaraDistributionTimeRangeMode TimeRangeMode, const FVector2f& CustomTimeRange)
+{
+	NiagaraStatelessDistributionPrivate::UpdateDistributionTimeRange(
+		TimeRangeMode,
+		CustomTimeRange,
+		Values.Num(),
+		ValuesTimeRange
+	);
+}
+
 void FNiagaraDistributionColor::UpdateValuesFromDistribution()
 {
 	NiagaraStatelessDistributionPrivate::UpdateDistributionValues(
@@ -850,6 +975,16 @@ void FNiagaraDistributionColor::UpdateValuesFromDistribution()
 	);
 }
 
+void FNiagaraDistributionColor::UpdateTimeRangeFromDistribution(ENiagaraDistributionTimeRangeMode TimeRangeMode, const FVector2f& CustomTimeRange)
+{
+	NiagaraStatelessDistributionPrivate::UpdateDistributionTimeRange(
+		TimeRangeMode,
+		CustomTimeRange,
+		Values.Num(),
+		ValuesTimeRange
+	);
+}
+
 void FNiagaraDistributionCurveFloat::UpdateValuesFromDistribution()
 {
 	NiagaraStatelessDistributionPrivate::UpdateDistributionCurveValues(
diff --git a/Engine/Plugins/FX/Niagara/Source/Niagara/Private/Stateless/NiagaraStatelessEmitter.cpp b/Engine/Plugins/FX/Niagara/Source/Niagara/Private/Stateless/NiagaraStatelessEmitter.cpp
--- a/Engine/Plugins/FX/Niagara/Source/Niagara/Private/Stateless/NiagaraStatelessEmitter.cpp
+++ b/Engine/Plugins/FX/Niagara/Source/Niagara/Private/Stateless/NiagaraStatelessEmitter.cpp
@@ -445,7 +445,7 @@ void UNiagaraStatelessEmitter::CacheFromCompiledData()
 			StatelessEmitterData->RendererBindings,
 			StatelessEmitterData->Expressions,
 			StatelessEmitterData->BuiltData,
-			StatelessEmitterData->StaticFloatData,
+			StatelessEmitterData->StaticDataBufferCpu,
 			StatelessEmitterData->ParticleSimExecData
 		);
 
@@ -460,9 +460,11 @@ void UNiagaraStatelessEmitter::CacheFromCompiledData()
 		}
 		StatelessEmitterData->bModulesHaveRendererBindings = StatelessEmitterData->RendererBindings.Num() > 0;
 
-		if (StatelessEmitterData->StaticFloatData.Num() == 0)
+		// Make sure our static data buffer allows us to sample the max size in one chunk without having to worry about overflow
+		constexpr int MinSafeBufferSize = 16;
+		if (StatelessEmitterData->StaticDataBufferCpu.Num() < MinSafeBufferSize)
 		{
-			StatelessEmitterData->StaticFloatData.Emplace(0.0f);
+			StatelessEmitterData->StaticDataBufferCpu.AddZeroed(MinSafeBufferSize - StatelessEmitterData->StaticDataBufferCpu.Num());
 		}
 		StatelessEmitterData->InitRenderResources();
 
diff --git a/Engine/Plugins/FX/Niagara/Source/Niagara/Private/Stateless/NiagaraStatelessEmitterData.cpp b/Engine/Plugins/FX/Niagara/Source/Niagara/Private/Stateless/NiagaraStatelessEmitterData.cpp
--- a/Engine/Plugins/FX/Niagara/Source/Niagara/Private/Stateless/NiagaraStatelessEmitterData.cpp
+++ b/Engine/Plugins/FX/Niagara/Source/Niagara/Private/Stateless/NiagaraStatelessEmitterData.cpp
@@ -10,7 +10,7 @@
 FNiagaraStatelessEmitterData::~FNiagaraStatelessEmitterData()
 {
 	check(IsInRenderingThread());
-	StaticFloatBuffer.Release();
+	StaticDataBufferGpu.Release();
 
 	if (ParticleSimExecData)
 	{
@@ -25,10 +25,10 @@ void FNiagaraStatelessEmitterData::InitRenderResources()
 	(
 		[this](FRHICommandListImmediate& RHICmdList)
 		{
-			StaticFloatBuffer.Initialize(RHICmdList, TEXT("NiagaraStatelessEmitterData_StaticFloatBuffer"), sizeof(float), StaticFloatData.Num(), PF_R32_FLOAT, EBufferUsageFlags::Static);
-			void* BufferData = RHICmdList.LockBuffer(StaticFloatBuffer.Buffer, 0, StaticFloatBuffer.NumBytes, RLM_WriteOnly);
-			FMemory::Memcpy(BufferData, StaticFloatData.GetData(), StaticFloatBuffer.NumBytes);
-			RHICmdList.UnlockBuffer(StaticFloatBuffer.Buffer);
+			StaticDataBufferGpu.Initialize(RHICmdList, TEXT("NiagaraStatelessEmitterData_StaticDataBuffer"), StaticDataBufferCpu.Num(), EBufferUsageFlags::Static);
+			void* BufferData = RHICmdList.LockBuffer(StaticDataBufferGpu.Buffer, 0, StaticDataBufferGpu.NumBytes, RLM_WriteOnly);
+			FMemory::Memcpy(BufferData, StaticDataBufferCpu.GetData(), StaticDataBufferGpu.NumBytes);
+			RHICmdList.UnlockBuffer(StaticDataBufferGpu.Buffer);
 		}
 	);
 }
diff --git a/Engine/Plugins/FX/Niagara/Source/Niagara/Private/Stateless/NiagaraStatelessEmitterDataBuildContext.cpp b/Engine/Plugins/FX/Niagara/Source/Niagara/Private/Stateless/NiagaraStatelessEmitterDataBuildContext.cpp
--- a/Engine/Plugins/FX/Niagara/Source/Niagara/Private/Stateless/NiagaraStatelessEmitterDataBuildContext.cpp
+++ b/Engine/Plugins/FX/Niagara/Source/Niagara/Private/Stateless/NiagaraStatelessEmitterDataBuildContext.cpp
@@ -9,6 +9,37 @@
 #include "NiagaraParameterStore.h"
 #include "NiagaraParameterStore.h"
 
+namespace NiagaraStatelessEmitterDataBuildContextPrivate
+{
+	uint32 AddStaticData(TArray<uint8>& Buffer, const void* Data, uint32 DataSize)
+	{
+		if (DataSize == 0)
+		{
+			return 0;
+		}
+
+		constexpr uint32 BlockSize = 4;
+		check((DataSize % BlockSize) == 0);
+
+		constexpr bool bDeduplicateData = true;
+		if (bDeduplicateData)
+		{
+			const uint32 iEndOffset = Buffer.Num();
+			for (uint32 iOffset = 0; iOffset + DataSize <= iEndOffset; iOffset += BlockSize)
+			{
+				if (FMemory::Memcmp(&Buffer[iOffset], Data, DataSize) == 0)
+				{
+					return iOffset / BlockSize;
+				}
+			}
+		}
+
+		const uint32 OutIndex = Buffer.AddUninitialized(DataSize);
+		FMemory::Memcpy(&Buffer[OutIndex], Data, DataSize);
+		return OutIndex / BlockSize;
+	}
+};
+
 void FNiagaraStatelessEmitterDataBuildContext::PreModuleBuild(int32 InShaderParameterOffset)
 {
 	ModuleBuiltDataOffset = BuiltData.Num();
@@ -16,40 +47,39 @@ void FNiagaraStatelessEmitterDataBuildContext::PreModuleBuild(int32 InShaderPara
 	++RandomSeedOffest;
 }
 
-uint32 FNiagaraStatelessEmitterDataBuildContext::AddStaticData(TConstArrayView<float> FloatData) const
+uint32 FNiagaraStatelessEmitterDataBuildContext::AddStaticData(TConstArrayView<uint32> StaticData) const
 {
-	for (int32 i=0; i <= StaticFloatData.Num() - FloatData.Num(); ++i)
-	{
-		if (FMemory::Memcmp(StaticFloatData.GetData() + i, FloatData.GetData(), FloatData.Num() * sizeof(float)) == 0)
-		{
-			return i;
-		}
-	}
+	return NiagaraStatelessEmitterDataBuildContextPrivate::AddStaticData(StaticDataBuffer, StaticData.GetData(), StaticData.GetTypeSize() * StaticData.Num());
+}
 
-	// Add new
-	const uint32 OutIndex = StaticFloatData.AddUninitialized(FloatData.Num());
-	FMemory::Memcpy(StaticFloatData.GetData() + OutIndex, FloatData.GetData(), FloatData.Num() * sizeof(float));
-	return OutIndex;
+uint32 FNiagaraStatelessEmitterDataBuildContext::AddStaticData(TConstArrayView<int32> StaticData) const
+{
+	return NiagaraStatelessEmitterDataBuildContextPrivate::AddStaticData(StaticDataBuffer, StaticData.GetData(), StaticData.GetTypeSize() * StaticData.Num());
+}
+
+uint32 FNiagaraStatelessEmitterDataBuildContext::AddStaticData(TConstArrayView<float> StaticData) const
+{
+	return NiagaraStatelessEmitterDataBuildContextPrivate::AddStaticData(StaticDataBuffer, StaticData.GetData(), StaticData.GetTypeSize() * StaticData.Num());
 }
 
-uint32 FNiagaraStatelessEmitterDataBuildContext::AddStaticData(TConstArrayView<FVector2f> FloatData) const
+uint32 FNiagaraStatelessEmitterDataBuildContext::AddStaticData(TConstArrayView<FVector2f> StaticData) const
 {
-	return AddStaticData(MakeArrayView(reinterpret_cast<const float*>(FloatData.GetData()), FloatData.Num() * 2));
+	return NiagaraStatelessEmitterDataBuildContextPrivate::AddStaticData(StaticDataBuffer, StaticData.GetData(), StaticData.GetTypeSize() * StaticData.Num());
 }
 
-uint32 FNiagaraStatelessEmitterDataBuildContext::AddStaticData(TConstArrayView<FVector3f> FloatData) const
+uint32 FNiagaraStatelessEmitterDataBuildContext::AddStaticData(TConstArrayView<FVector3f> StaticData) const
 {
-	return AddStaticData(MakeArrayView(reinterpret_cast<const float*>(FloatData.GetData()), FloatData.Num() * 3));
+	return NiagaraStatelessEmitterDataBuildContextPrivate::AddStaticData(StaticDataBuffer, StaticData.GetData(), StaticData.GetTypeSize() * StaticData.Num());
 }
 
-uint32 FNiagaraStatelessEmitterDataBuildContext::AddStaticData(TConstArrayView<FVector4f> FloatData) const
+uint32 FNiagaraStatelessEmitterDataBuildContext::AddStaticData(TConstArrayView<FVector4f> StaticData) const
 {
-	return AddStaticData(MakeArrayView(reinterpret_cast<const float*>(FloatData.GetData()), FloatData.Num() * 4));
+	return NiagaraStatelessEmitterDataBuildContextPrivate::AddStaticData(StaticDataBuffer, StaticData.GetData(), StaticData.GetTypeSize() * StaticData.Num());
 }
 
-uint32 FNiagaraStatelessEmitterDataBuildContext::AddStaticData(TConstArrayView<FLinearColor> FloatData) const
+uint32 FNiagaraStatelessEmitterDataBuildContext::AddStaticData(TConstArrayView<FLinearColor> StaticData) const
 {
-	return AddStaticData(MakeArrayView(reinterpret_cast<const float*>(FloatData.GetData()), FloatData.Num() * 4));
+	return NiagaraStatelessEmitterDataBuildContextPrivate::AddStaticData(StaticDataBuffer, StaticData.GetData(), StaticData.GetTypeSize() * StaticData.Num());
 }
 
 int32 FNiagaraStatelessEmitterDataBuildContext::AddRendererBinding(const FNiagaraVariableBase& Variable) const
diff --git a/Engine/Plugins/FX/Niagara/Source/Niagara/Private/Stateless/NiagaraStatelessParticleSimContext.cpp b/Engine/Plugins/FX/Niagara/Source/Niagara/Private/Stateless/NiagaraStatelessParticleSimContext.cpp
--- a/Engine/Plugins/FX/Niagara/Source/Niagara/Private/Stateless/NiagaraStatelessParticleSimContext.cpp
+++ b/Engine/Plugins/FX/Niagara/Source/Niagara/Private/Stateless/NiagaraStatelessParticleSimContext.cpp
@@ -41,7 +41,7 @@ namespace NiagaraStateless
 		: EmitterData(InEmitterData)
 		, BuiltData(EmitterData->BuiltData)
 		, ShaderParametersData(static_cast<const uint8*>(InShaderParametersData))
-		, StaticFloatData(EmitterData->StaticFloatData)
+		, StaticFloatData(MakeConstArrayView(reinterpret_cast<const float*>(EmitterData->StaticDataBufferCpu.GetData()), EmitterData->StaticDataBufferCpu.Num() / sizeof(float)))
 		, DynamicBufferData(InDynamicBufferData)
 	{
 		check(EmitterData->ParticleSimExecData);
diff --git a/Engine/Plugins/FX/Niagara/Source/Niagara/Public/NiagaraRenderer.h b/Engine/Plugins/FX/Niagara/Source/Niagara/Public/NiagaraRenderer.h
--- a/Engine/Plugins/FX/Niagara/Source/Niagara/Public/NiagaraRenderer.h
+++ b/Engine/Plugins/FX/Niagara/Source/Niagara/Public/NiagaraRenderer.h
@@ -130,6 +130,7 @@ public:
 	inline static FRHIShaderResourceView* GetSrvOrDefaultFloat4(FRHIShaderResourceView* InSRV) { return InSRV ? InSRV : GetDummyFloat4Buffer(); }
 	inline static FRHIShaderResourceView* GetSrvOrDefaultUInt(FRHIShaderResourceView* InSRV) { return InSRV ? InSRV : GetDummyUIntBuffer(); }
 	inline static FRHIShaderResourceView* GetSrvOrDefaultInt(FRHIShaderResourceView* InSRV) { return InSRV ? InSRV : GetDummyIntBuffer(); }
+	inline static FRHIShaderResourceView* GetSrvOrDefaultByteAddress(FRHIShaderResourceView* InSRV) { return InSRV ? InSRV : GetDummyByteAddressBuffer(); }
 
 	static NIAGARA_API FRHIShaderResourceView* GetDummyFloatBuffer();
 	static NIAGARA_API FRHIShaderResourceView* GetDummyFloat2Buffer();
@@ -140,6 +141,7 @@ public:
 	static NIAGARA_API FRHIShaderResourceView* GetDummyUInt2Buffer();
 	static NIAGARA_API FRHIShaderResourceView* GetDummyUInt4Buffer();
 	static NIAGARA_API FRHIShaderResourceView* GetDummyHalfBuffer();
+	static NIAGARA_API FRHIShaderResourceView* GetDummyByteAddressBuffer();
 
 	static inline int SelectVFComponent(const TConstArrayView<FNiagaraRendererVariableInfo>& Variables, int Variable, int Fallback)
 	{
diff --git a/Engine/Plugins/FX/Niagara/Source/Niagara/Public/Stateless/NiagaraStatelessDistribution.h b/Engine/Plugins/FX/Niagara/Source/Niagara/Public/Stateless/NiagaraStatelessDistribution.h
--- a/Engine/Plugins/FX/Niagara/Source/Niagara/Public/Stateless/NiagaraStatelessDistribution.h
+++ b/Engine/Plugins/FX/Niagara/Source/Niagara/Public/Stateless/NiagaraStatelessDistribution.h
@@ -10,8 +10,9 @@
 #include "NiagaraStatelessDistribution.generated.h"
 
 UENUM()
-enum class ENiagaraDistributionMode
+enum class ENiagaraDistributionMode : uint8
 {
+	Array,
 	Binding,
 	Expression,
 	UniformConstant,
@@ -23,6 +24,44 @@ enum class ENiagaraDistributionMode
 	ColorGradient,
 };
 
+UENUM()
+enum class ENiagaraDistributionInterpolationMode : uint8
+{
+	// Values sampled will be linearly interpolated
+	Linear,
+	// No interpolation is applied, fractional indexes will be rounded to nearest
+	None,
+};
+
+UENUM()
+enum class ENiagaraDistributionAddressMode : uint8
+{
+	// Index will be clamped, i.e. Clamp(Index, 0, Num)
+	Clamp,
+	// Index will wrap around, i.e. Modulo(Index, Num)
+	Wrap,
+};
+
+UENUM()
+enum class ENiagaraDistributionTimeRangeMode : uint8
+{
+	Direct,			// The lookup value is used directly to look into the data
+	Normalized,		// The lookup value is normalized so should expand over the data
+	Custom,			// A custom time range will be applied to the lookup value
+};
+
+// Default distribution LookupValueMode the enumeration used can be overridden using metadata on the property, i.e. DistributionLookupValueEnumPath="/Script/Niagara.ENiagaraDistributionLookupValueMode"
+// Random is a special value discovered when packing the distribution data, the max number of potential values is based on EFlags::LookupValueModeMask - 1
+// Meta data TimeRangeMode is used to drive how the time range will be calculated when not overriden by the user, see ENiagaraDistributionTimeRangeMode
+UENUM()
+enum class ENiagaraDistributionLookupValueMode : uint8
+{
+	// Use a randomly generated value, generally evaluated at spawn
+	Random = 0xff UMETA(TimeRangeMode = "Normalized"),
+	// Bound to the variable Particles.NormalizedAge
+	ParticlesNormalizedAge = 0 UMETA(TimeRangeMode = "Normalized", DisplayName = "Particles.NormalizedAge"),
+};
+
 enum class ENiagaraDistributionCurveLUTMode
 {
 	Sample,		// Each sample in the LUT represents the curve evaulation
@@ -39,12 +78,25 @@ struct FNiagaraDistributionBase
 	UPROPERTY(EditAnywhere, Category = "Parameters")
 	ENiagaraDistributionMode Mode = ENiagaraDistributionMode::UniformConstant;
 
+	// Maps to ENiagaraDistributionInterpolationMode
+	UPROPERTY(EditAnywhere, Category = "Parameters")
+	uint8 InterpolationMode : 1 = uint8(ENiagaraDistributionInterpolationMode::Linear);
+
+	// Maps to ENiagaraDistributionAddressMode
+	UPROPERTY(EditAnywhere, Category = "Parameters")
+	uint8 AddressMode : 1 = uint8(ENiagaraDistributionAddressMode::Clamp);
+
+	// Custom per property, the default is to use ENiagaraDistributionLookupValueMode, check the metadata for exact binding
+	UPROPERTY(EditAnywhere, Category = "Parameters")
+	uint8 LookupValueMode = uint8(ENiagaraDistributionLookupValueMode::ParticlesNormalizedAge);
+
 	UPROPERTY(EditAnywhere, Category = "Parameters")
 	FNiagaraVariableBase ParameterBinding;
 
 	UPROPERTY(EditAnywhere, Category = "Parameters")
 	FInstancedStruct ParameterExpression;
 
+	bool IsArray() const { return Mode == ENiagaraDistributionMode::Array; }
 	bool IsBinding() const { return Mode == ENiagaraDistributionMode::Binding; }
 	bool IsExpression() const { return Mode == ENiagaraDistributionMode::Expression; }
 	bool IsConstant() const { return Mode == ENiagaraDistributionMode::UniformConstant || Mode == ENiagaraDistributionMode::NonUniformConstant; }
@@ -53,7 +105,15 @@ struct FNiagaraDistributionBase
 	bool IsGradient() const { return Mode == ENiagaraDistributionMode::ColorGradient; }
 	bool IsRange() const { return Mode == ENiagaraDistributionMode::UniformRange || Mode == ENiagaraDistributionMode::NonUniformRange; }
 
+	ENiagaraDistributionInterpolationMode GetInterpolationMode() const { return ENiagaraDistributionInterpolationMode(InterpolationMode); }
+	ENiagaraDistributionAddressMode GetAddressMode() const { return ENiagaraDistributionAddressMode(AddressMode); }
+	uint8 GetLookupValueMode() const { return LookupValueMode; }
+
 #if WITH_EDITORONLY_DATA
+	void SetInterpolationMode(ENiagaraDistributionInterpolationMode InMode) { check(uint8(InMode) >= 0 && uint8(InMode) <= 1); InterpolationMode = uint8(InMode); }
+	void SetAddressMode(ENiagaraDistributionAddressMode InMode) { check(uint8(InMode) >= 0 && uint8(InMode) <= 1); AddressMode = uint8(InMode); }
+	void SetLookupValueMode(uint8 InMode) { LookupValueMode = InMode; }
+
 	UPROPERTY(EditAnywhere, Category = "Parameters")
 	TArray<float> ChannelConstantsAndRanges;
 
@@ -67,16 +127,20 @@ struct FNiagaraDistributionBase
 
 	void ForEachParameterBinding(TFunction<void(const FNiagaraVariableBase&)> Delegate) const;
 
+	virtual bool AllowArray() const { return true; }
 	virtual bool AllowBinding() const { return true; }
 	virtual bool AllowConstant() const { return true; }
 	virtual bool AllowCurves() const { return true; }
 	virtual bool DisplayAsColor() const { return false; }
 	virtual int32 GetBaseNumberOfChannels() const { return 0; }
 	virtual void UpdateValuesFromDistribution() { }
+	virtual void UpdateTimeRangeFromDistribution(ENiagaraDistributionTimeRangeMode TimeRangeMode, const FVector2f& CustomTimeRange = FVector2f::ZeroVector) {}
 
 	virtual FNiagaraTypeDefinition GetBindingTypeDef() const { return FNiagaraTypeDefinition(); }
 
 	static void PostEditChangeProperty(UObject* OwnerObject, FPropertyChangedEvent& PropertyChangedEvent);
+
+	NIAGARA_API static const UEnum* GetLookupValueModeEnum(FProperty* Property);
 #endif
 };
 
@@ -108,6 +172,7 @@ struct FNiagaraDistributionRangeInt
 	NIAGARA_API void InitConstant(int32 Value);
 	NIAGARA_API FNiagaraStatelessRangeInt CalculateRange(const int32 Default = GetDefaultValue()) const;
 
+	bool IsArray() const { return Mode == ENiagaraDistributionMode::Array; }
 	bool IsBinding() const { return Mode == ENiagaraDistributionMode::Binding; }
 	bool IsExpression() const { return Mode == ENiagaraDistributionMode::Expression; }
 	bool IsConstant() const { return Mode == ENiagaraDistributionMode::UniformConstant || Mode == ENiagaraDistributionMode::NonUniformConstant; }
@@ -143,6 +208,7 @@ struct FNiagaraDistributionRangeFloat : public FNiagaraDistributionBase
 	NIAGARA_API FNiagaraStatelessRangeFloat CalculateRange(const float Default = GetDefaultValue()) const;
 
 #if WITH_EDITORONLY_DATA
+	virtual bool AllowArray() const override { return false; }
 	virtual bool AllowCurves() const override { return false; }
 	virtual int32 GetBaseNumberOfChannels() const override { return 1; }
 	NIAGARA_API virtual void UpdateValuesFromDistribution() override;
@@ -171,6 +237,7 @@ struct FNiagaraDistributionRangeVector2 : public FNiagaraDistributionBase
 	NIAGARA_API FNiagaraStatelessRangeVector2 CalculateRange(const FVector2f& Default = GetDefaultValue()) const;
 
 #if WITH_EDITORONLY_DATA
+	virtual bool AllowArray() const override { return false; }
 	virtual bool AllowCurves() const override { return false; }
 	virtual int32 GetBaseNumberOfChannels() const override { return 2; }
 	virtual void UpdateValuesFromDistribution() override;
@@ -199,6 +266,7 @@ struct FNiagaraDistributionRangeVector3 : public FNiagaraDistributionBase
 	NIAGARA_API FNiagaraStatelessRangeVector3 CalculateRange(const FVector3f& Default = GetDefaultValue()) const;
 
 #if WITH_EDITORONLY_DATA
+	virtual bool AllowArray() const override { return false; }
 	virtual bool AllowCurves() const override { return false; }
 	virtual int32 GetBaseNumberOfChannels() const override { return 3; }
 	virtual void UpdateValuesFromDistribution() override;
@@ -227,6 +295,7 @@ struct FNiagaraDistributionRangeColor : public FNiagaraDistributionBase
 	NIAGARA_API FNiagaraStatelessRangeColor CalculateRange(const FLinearColor& Default = GetDefaultValue()) const;
 
 #if WITH_EDITORONLY_DATA
+	virtual bool AllowArray() const override { return false; }
 	virtual bool AllowCurves() const override { return false; }
 	virtual bool DisplayAsColor() const { return true; }
 	virtual int32 GetBaseNumberOfChannels() const override { return 4; }
@@ -253,6 +322,7 @@ struct FNiagaraDistributionRangeRotator : public FNiagaraDistributionBase
 	NIAGARA_API FNiagaraStatelessRangeRotator CalculateRange(const FRotator3f& Default = FRotator3f::ZeroRotator) const;
 
 #if WITH_EDITORONLY_DATA
+	virtual bool AllowArray() const override { return false; }
 	virtual bool AllowCurves() const override { return false; }
 	virtual int32 GetBaseNumberOfChannels() const override { return 3; }
 	virtual void UpdateValuesFromDistribution() override;
@@ -292,6 +362,7 @@ struct FNiagaraDistributionFloat : public FNiagaraDistributionBase
 	}
 	virtual int32 GetBaseNumberOfChannels() const override { return 1; }
 	NIAGARA_API virtual void UpdateValuesFromDistribution() override;
+	NIAGARA_API virtual void UpdateTimeRangeFromDistribution(ENiagaraDistributionTimeRangeMode TimeRangeMode, const FVector2f& CustomTimeRange) override;
 	virtual FNiagaraTypeDefinition GetBindingTypeDef() const { return FNiagaraTypeDefinition::GetFloatDef(); }
 #endif
 };
@@ -320,6 +391,7 @@ struct FNiagaraDistributionVector2 : public FNiagaraDistributionBase
 #if WITH_EDITORONLY_DATA
 	virtual int32 GetBaseNumberOfChannels() const override { return 2; }
 	NIAGARA_API virtual void UpdateValuesFromDistribution() override;
+	NIAGARA_API virtual void UpdateTimeRangeFromDistribution(ENiagaraDistributionTimeRangeMode TimeRangeMode, const FVector2f& CustomTimeRange) override;
 	virtual FNiagaraTypeDefinition GetBindingTypeDef() const { return FNiagaraTypeDefinition::GetVec2Def(); }
 #endif
 };
@@ -352,6 +424,7 @@ struct FNiagaraDistributionVector3 : public FNiagaraDistributionBase
 #if WITH_EDITORONLY_DATA
 	virtual int32 GetBaseNumberOfChannels() const override { return 3; }
 	NIAGARA_API virtual void UpdateValuesFromDistribution() override;
+	NIAGARA_API virtual void UpdateTimeRangeFromDistribution(ENiagaraDistributionTimeRangeMode TimeRangeMode, const FVector2f& CustomTimeRange) override;
 	virtual FNiagaraTypeDefinition GetBindingTypeDef() const { return FNiagaraTypeDefinition::GetVec3Def(); }
 #endif
 };
@@ -393,6 +466,7 @@ struct FNiagaraDistributionColor : public FNiagaraDistributionBase
 	virtual bool DisplayAsColor() const override { return true; }
 	virtual int32 GetBaseNumberOfChannels() const override { return 4; }
 	NIAGARA_API virtual void UpdateValuesFromDistribution() override;
+	NIAGARA_API virtual void UpdateTimeRangeFromDistribution(ENiagaraDistributionTimeRangeMode TimeRangeMode, const FVector2f& CustomTimeRange) override;
 	virtual FNiagaraTypeDefinition GetBindingTypeDef() const { return FNiagaraTypeDefinition::GetColorDef(); }
 #endif
 };
@@ -417,6 +491,7 @@ struct FNiagaraDistributionCurveFloat : public FNiagaraDistributionBase
 		return (FNiagaraDistributionBase)*this == (FNiagaraDistributionBase)Other && Values == Other.Values;
 	}
 
+	virtual bool AllowArray() const override { return false; }
 	virtual bool AllowBinding() const override { return false; }
 	virtual bool AllowConstant() const override { return false; }
 	virtual int32 GetBaseNumberOfChannels() const override { return 1; }
@@ -447,6 +522,7 @@ struct FNiagaraDistributionCurveVector3 : public FNiagaraDistributionBase
 		return (FNiagaraDistributionBase)*this == (FNiagaraDistributionBase)Other && Values == Other.Values;
 	}
 
+	virtual bool AllowArray() const override { return false; }
 	virtual bool AllowBinding() const override { return false; }
 	virtual bool AllowConstant() const override { return false; }
 	virtual int32 GetBaseNumberOfChannels() const override { return 3; }
diff --git a/Engine/Plugins/FX/Niagara/Source/NiagaraEditor/Private/NiagaraEditorStyle.cpp b/Engine/Plugins/FX/Niagara/Source/NiagaraEditor/Private/NiagaraEditorStyle.cpp
--- a/Engine/Plugins/FX/Niagara/Source/NiagaraEditor/Private/NiagaraEditorStyle.cpp
+++ b/Engine/Plugins/FX/Niagara/Source/NiagaraEditor/Private/NiagaraEditorStyle.cpp
@@ -441,6 +441,7 @@ void FNiagaraEditorStyle::InitIcons()
 	SlateBrush->TintColor = FLinearColor(0.01f, 0.01f, 0.01f);
 	Set("NiagaraEditor.Module.InputTypeBorder", SlateBrush);
 
+	Set("NiagaraEditor.DistributionEditor.Array", new IMAGE_BRUSH("Icons/DistributionEditor/Array", Icon20x20));
 	Set("NiagaraEditor.DistributionEditor.Binding", new IMAGE_BRUSH_SVG("Icons/DistributionEditor/Binding", Icon20x20));
 	Set("NiagaraEditor.DistributionEditor.UniformConstant", new IMAGE_BRUSH_SVG("Icons/DistributionEditor/UniformConstant", Icon20x20));
 	Set("NiagaraEditor.DistributionEditor.NonUniformConstant", new IMAGE_BRUSH_SVG("Icons/DistributionEditor/NonUniformConstant", Icon20x20));
diff --git a/Engine/Plugins/FX/Niagara/Source/NiagaraEditor/Private/Stateless/NiagaraDistributionIntPropertyCustomization.cpp b/Engine/Plugins/FX/Niagara/Source/NiagaraEditor/Private/Stateless/NiagaraDistributionIntPropertyCustomization.cpp
--- a/Engine/Plugins/FX/Niagara/Source/NiagaraEditor/Private/Stateless/NiagaraDistributionIntPropertyCustomization.cpp
+++ b/Engine/Plugins/FX/Niagara/Source/NiagaraEditor/Private/Stateless/NiagaraDistributionIntPropertyCustomization.cpp
@@ -92,6 +92,7 @@ class SNiagaraDistributionIntPropertyWidget : public SCompoundWidget
 		{
 			switch (Distribution->Mode)
 			{
+				case ENiagaraDistributionMode::Array:			return ENiagaraDistributionEditorMode::Array;
 				case ENiagaraDistributionMode::Binding:			return ENiagaraDistributionEditorMode::Binding;
 				//case ENiagaraDistributionMode::Expression:		return ENiagaraDistributionEditorMode::Expression;
 				case ENiagaraDistributionMode::UniformConstant: return ENiagaraDistributionEditorMode::UniformConstant;
@@ -114,6 +115,7 @@ class SNiagaraDistributionIntPropertyWidget : public SCompoundWidget
 
 			switch (NewMode)
 			{
+				case ENiagaraDistributionEditorMode::Array:				Distribution->Mode = ENiagaraDistributionMode::Array; break;
 				case ENiagaraDistributionEditorMode::Binding:			Distribution->Mode = ENiagaraDistributionMode::Binding; break;
 				//case ENiagaraDistributionEditorMode::Expression:		Distribution->Mode = ENiagaraDistributionMode::Expression; break;
 				case ENiagaraDistributionEditorMode::UniformConstant:	Distribution->Mode = ENiagaraDistributionMode::UniformConstant; break;
@@ -150,6 +152,12 @@ class SNiagaraDistributionIntPropertyWidget : public SCompoundWidget
 				//	SupportedModes.Add(ENiagaraDistributionEditorMode::Expression);
 				//}
 			}
+			const FName ArrayBindingDistributionName("DisableArrayDistribution");
+			if (!PropertyHandle->HasMetaData(ArrayBindingDistributionName))
+			{
+				//-TODO: Array support for ints
+				//SupportedModes.Add(ENiagaraDistributionEditorMode::Array);
+			}
 			SupportedModes.Add(ENiagaraDistributionEditorMode::UniformConstant);
 			SupportedModes.Add(ENiagaraDistributionEditorMode::UniformRange);
 		}
@@ -194,6 +202,15 @@ class SNiagaraDistributionIntPropertyWidget : public SCompoundWidget
 					ConstructBindingWidget()
 				];
 		}
+		//-TODO: Support ints
+		//else if (Mode == ENiagaraDistributionEditorMode::Array)
+		//{
+			//return SNew(SHorizontalBox)
+			//	+ SHorizontalBox::Slot()
+			//	[
+			//		ConstructBindingWidget()
+			//	];
+		//}
 		//else if (Mode == ENiagaraDistributionEditorMode::Expression)
 		//{
 		//	return SNew(SHorizontalBox)
diff --git a/Engine/Plugins/FX/Niagara/Source/NiagaraEditor/Private/Stateless/NiagaraDistributionPropertyCustomization.cpp b/Engine/Plugins/FX/Niagara/Source/NiagaraEditor/Private/Stateless/NiagaraDistributionPropertyCustomization.cpp
--- a/Engine/Plugins/FX/Niagara/Source/NiagaraEditor/Private/Stateless/NiagaraDistributionPropertyCustomization.cpp
+++ b/Engine/Plugins/FX/Niagara/Source/NiagaraEditor/Private/Stateless/NiagaraDistributionPropertyCustomization.cpp
@@ -44,11 +44,13 @@ public:
 		SourceNumChannels = InNumChannels;
 		if (IsValid())
 		{
+			const FName DisableArrayDistributionName("DisableArrayDistribution");
 			const FName DisableBindingDistributionName("DisableBindingDistribution");
 			const FName DisableCurveDistributionName("DisableCurveDistribution");
 			const FName DisableUniformDistributionName("DisableUniformDistribution");
 			const FName DisableNonUniformDistributionName("DisableNonUniformDistribution");
 			const FName DisableRangeDistributionName("DisableRangeDistribution");
+			bAllowArray = InDistribution->AllowArray() && InDistribution->GetBindingTypeDef().IsValid() && (PropertyHandle ? PropertyHandle->HasMetaData(DisableArrayDistributionName) == false : false);
 			bAllowBinding = InDistribution->AllowBinding() && InDistribution->GetBindingTypeDef().IsValid() && (PropertyHandle ? PropertyHandle->HasMetaData(DisableBindingDistributionName) == false : false);
 			bAllowUniform = InDistribution->AllowConstant() && PropertyHandle->HasMetaData(DisableUniformDistributionName) == false;
 			bAllowNonUniform = InDistribution->AllowConstant() && PropertyHandle->HasMetaData(DisableNonUniformDistributionName) == false;
@@ -72,6 +74,8 @@ public:
 
 	virtual bool IsValid() const override { return PropertyHandle.IsValid() && SourceDistribution != nullptr; }
 
+	virtual TSharedPtr<IPropertyHandle> GetPropertyHandle() const override { return PropertyHandle; }
+
 	virtual int32 GetNumChannels() const override { return EditorNumChannels; }
 
 	virtual FText GetChannelDisplayName(int32 ChannelIndex) const override 
@@ -150,6 +154,10 @@ public:
 
 	virtual void GetSupportedDistributionModes(TArray<ENiagaraDistributionEditorMode>& OutSupportedModes) const override
 	{
+		if (bAllowArray)
+		{
+			OutSupportedModes.Add(ENiagaraDistributionEditorMode::Array);
+		}
 		if (bAllowBinding)
 		{
 			OutSupportedModes.Add(ENiagaraDistributionEditorMode::Binding);
@@ -217,6 +225,8 @@ public:
 		{
 			switch (InSourceMode)
 			{
+			case ENiagaraDistributionMode::Array:
+				return ENiagaraDistributionEditorMode::Array;
 			case ENiagaraDistributionMode::Binding:
 				return ENiagaraDistributionEditorMode::Binding;
 			case ENiagaraDistributionMode::Expression:
@@ -237,6 +247,8 @@ public:
 		{
 			switch (InSourceMode)
 			{
+			case ENiagaraDistributionMode::Array:
+				return ENiagaraDistributionEditorMode::Array;
 			case ENiagaraDistributionMode::Binding:
 				return ENiagaraDistributionEditorMode::Binding;
 			case ENiagaraDistributionMode::Expression:
@@ -255,6 +267,8 @@ public:
 		{
 			switch (InSourceMode)
 			{
+			case ENiagaraDistributionMode::Array:
+				return ENiagaraDistributionEditorMode::Array;
 			case ENiagaraDistributionMode::Binding:
 				return ENiagaraDistributionEditorMode::Binding;
 			case ENiagaraDistributionMode::Expression:
@@ -283,6 +297,8 @@ public:
 	{
 		switch (InEditorMode)
 		{
+		case ENiagaraDistributionEditorMode::Array:
+			return ENiagaraDistributionMode::Array;
 		case ENiagaraDistributionEditorMode::Binding:
 			return ENiagaraDistributionMode::Binding;
 		case ENiagaraDistributionEditorMode::Expression:
@@ -495,6 +511,25 @@ public:
 		PropertyHandle->NotifyPostChange(bContinuousChangeActive ? EPropertyChangeType::Interactive : EPropertyChangeType::ValueSet);
 	}
 
+	virtual TConstArrayView<float> GetArrayValues() const override
+	{
+		return MakeArrayView(SourceDistribution->ChannelConstantsAndRanges);
+	}
+
+	virtual void SetArrayValues(TConstArrayView<float> ArrayData) override
+	{
+		check((ArrayData.Num() % GetNumChannels()) == 0);
+
+		FText TransactionText = LOCTEXT("SetArrayValues", "Set Array Value");
+		ExecuteTransaction(
+			LOCTEXT("SetArrayValues", "Set Array Value"),
+			[&]()
+			{
+				SourceDistribution->ChannelConstantsAndRanges = ArrayData;
+			}
+		);
+	}
+
 	virtual FNiagaraTypeDefinition GetExpressionTypeDef() const override
 	{
 		return SourceDistribution->GetBindingTypeDef();
@@ -638,6 +673,54 @@ public:
 		}
 	}
 
+	virtual ENiagaraDistributionInterpolationMode GetInterpolationMode() const override
+	{
+		return SourceDistribution->GetInterpolationMode();
+	}
+
+	virtual void SetInterpolationMode(ENiagaraDistributionInterpolationMode Mode) override
+	{
+		ExecuteTransaction(
+			LOCTEXT("SetInterpolationMode", "Set Interpolation Mode"),
+			[&]()
+			{
+				SourceDistribution->SetInterpolationMode(Mode);
+			}
+		);
+	}
+
+	virtual ENiagaraDistributionAddressMode GetAddressMode() const override
+	{
+		return SourceDistribution->GetAddressMode();
+	}
+
+	virtual void SetAddressMode(ENiagaraDistributionAddressMode Mode) override
+	{
+		ExecuteTransaction(
+			LOCTEXT("SetAddressMode", "Set Address Mode"),
+			[&]()
+			{
+				SourceDistribution->SetAddressMode(Mode);
+			}
+		);
+	}
+
+	virtual uint8 GetLookupValueMode() const override
+	{
+		return SourceDistribution->GetLookupValueMode();
+	}
+
+	virtual void SetLookupValueMode(uint8 Mode) override
+	{
+		ExecuteTransaction(
+			LOCTEXT("SetLookupValueMode", "Set Lookup Value Mode"),
+			[&]()
+			{
+				SourceDistribution->SetLookupValueMode(Mode);
+			}
+		);
+	}
+
 	virtual void BeginContinuousChange() override
 	{
 		bContinuousTransactionPending = true;
@@ -929,6 +1012,7 @@ private:
 	TOptional<int32> ContinuousTransactionIndex;
 	bool bContinuousChangeActive = false;
 
+	bool bAllowArray = true;
 	bool bAllowBinding = true;
 	bool bAllowUniform = true;
 	bool bAllowNonUniform = true;
diff --git a/Engine/Plugins/FX/Niagara/Source/NiagaraEditor/Private/TypeEditorUtilities/NiagaraDistributionPropertyEditorUtilities.cpp b/Engine/Plugins/FX/Niagara/Source/NiagaraEditor/Private/TypeEditorUtilities/NiagaraDistributionPropertyEditorUtilities.cpp
--- a/Engine/Plugins/FX/Niagara/Source/NiagaraEditor/Private/TypeEditorUtilities/NiagaraDistributionPropertyEditorUtilities.cpp
+++ b/Engine/Plugins/FX/Niagara/Source/NiagaraEditor/Private/TypeEditorUtilities/NiagaraDistributionPropertyEditorUtilities.cpp
@@ -72,6 +72,7 @@ bool FNiagaraDistributionPropertyEditorUtilities::TryUpdateClipboardPortableValu
 				InTargetClipboardPortableValue = FNiagaraClipboardPortableValue::CreateFromStructValue(*CurveCollectionStruct, (uint8*)&CurveCollectionValue);
 				return true;
 			}
+			//-TODO: Support for Array & Range
 		}
 	}
 	return false;
diff --git a/Engine/Plugins/FX/Niagara/Source/NiagaraEditor/Private/Widgets/INiagaraDistributionAdapter.h b/Engine/Plugins/FX/Niagara/Source/NiagaraEditor/Private/Widgets/INiagaraDistributionAdapter.h
--- a/Engine/Plugins/FX/Niagara/Source/NiagaraEditor/Private/Widgets/INiagaraDistributionAdapter.h
+++ b/Engine/Plugins/FX/Niagara/Source/NiagaraEditor/Private/Widgets/INiagaraDistributionAdapter.h
@@ -6,6 +6,7 @@
 #include "HAL/Platform.h"
 #include "Internationalization/Text.h"
 #include "Styling/SlateColor.h"
+#include "Stateless/NiagaraStatelessDistribution.h"
 #include "NiagaraCommon.h"
 #include "NiagaraVariableMetaData.h"
 
@@ -14,6 +15,8 @@ struct FRichCurve;
 /** Defines different edit modes for a niagara distribution */
 enum class ENiagaraDistributionEditorMode
 {
+	/* A multi-channel value which is mapped to an array of values. */
+	Array,
 	/* Bound to an script attribute. */
 	Binding,
 	/** Bound to an expression chain */
@@ -31,6 +34,7 @@ enum class ENiagaraDistributionEditorMode
 	/** A multi-channel value which is distributed between a single range of values. */
 	UniformRange,
 	/** A multi-channel value which is distributed between ranges of values, one range per channel. */
+	/** A multi-channel value which is distributed between ranges of values, one range per channel. */
 	NonUniformRange,
 	/** A multi-channel value which is distributed between a single range of colors. */
 	ColorRange,
@@ -53,6 +57,9 @@ public:
 	/** Returns whether or not this adapter valid for use */
 	virtual bool IsValid() const = 0;
 
+	/** Get the underlying property we are modifying. */
+	virtual TSharedPtr<class IPropertyHandle> GetPropertyHandle() const = 0;
+
 	/** Gets the number of channels currently supported by this adapter. This can change with different distribution modes. */
 	virtual int32 GetNumChannels() const = 0;
 
@@ -102,6 +109,11 @@ public:
 	*/
 	virtual void SetConstantOrRangeValues(int32 ValueIndex, const TArray<float>& InValues) = 0;
 
+	/** Returns a view into the flat float array data. */
+	virtual TConstArrayView<float> GetArrayValues() const = 0;
+	/** Sets the array values. */
+	virtual void SetArrayValues(TConstArrayView<float> EntryData) = 0;
+
 	/** Get the type used when generating an expression to run */
 	virtual FNiagaraTypeDefinition GetExpressionTypeDef() const = 0;
 	/** Get the root expression struct. */
@@ -122,6 +134,21 @@ public:
 	/** Set a curve value for the specified channel */
 	virtual void SetCurveValue(int32 ChannelIndex, const FRichCurve& InValue) = 0;
 
+	/** Get the interpolation mode from the distribution */
+	virtual ENiagaraDistributionInterpolationMode GetInterpolationMode() const = 0;
+	/** Set the interpolation mode on the distribution */
+	virtual void SetInterpolationMode(ENiagaraDistributionInterpolationMode Mode) = 0;
+
+	/** Get the address mode from the distribution */
+	virtual ENiagaraDistributionAddressMode GetAddressMode() const = 0;
+	/** Set the address mode on the distribution */
+	virtual void SetAddressMode(ENiagaraDistributionAddressMode Mode) = 0;
+
+	/** Get the lookup value mode from the distribution */
+	virtual uint8 GetLookupValueMode() const = 0;
+	/** Set the lookup value mode on the distribution */
+	virtual void SetLookupValueMode(uint8 Mode) = 0;
+
 	/** Called to notify the distribution that a continuous change has begun. */
 	virtual void BeginContinuousChange() = 0;
 
diff --git a/Engine/Plugins/FX/Niagara/Source/NiagaraEditor/Private/Widgets/NiagaraDistributionEditorUtilities.cpp b/Engine/Plugins/FX/Niagara/Source/NiagaraEditor/Private/Widgets/NiagaraDistributionEditorUtilities.cpp
--- a/Engine/Plugins/FX/Niagara/Source/NiagaraEditor/Private/Widgets/NiagaraDistributionEditorUtilities.cpp
+++ b/Engine/Plugins/FX/Niagara/Source/NiagaraEditor/Private/Widgets/NiagaraDistributionEditorUtilities.cpp
@@ -12,6 +12,8 @@ FText FNiagaraDistributionEditorUtilities::DistributionModeToDisplayName(ENiagar
 {
 	switch (InMode)
 	{
+	case ENiagaraDistributionEditorMode::Array:
+		return LOCTEXT("ArrayDisplayName", "Array");
 	case ENiagaraDistributionEditorMode::Binding:
 		return LOCTEXT("BindingDisplayName", "Binding");
 	case ENiagaraDistributionEditorMode::Expression:
@@ -52,6 +54,8 @@ FText FNiagaraDistributionEditorUtilities::DistributionModeToToolTipText(ENiagar
 {
 	switch (InMode)
 	{
+	case ENiagaraDistributionEditorMode::Array:
+		return LOCTEXT("ArrayToolTip", "Value driven by an array of values.");
 	case ENiagaraDistributionEditorMode::Binding:
 		return LOCTEXT("BindingToolTip", "Value bound to an attribute.");
 	case ENiagaraDistributionEditorMode::Expression:
@@ -92,6 +96,8 @@ FName FNiagaraDistributionEditorUtilities::DistributionModeToIconBrushName(ENiag
 {
 	switch (InMode)
 	{
+	case ENiagaraDistributionEditorMode::Array:
+		return "NiagaraEditor.DistributionEditor.Array";
 	case ENiagaraDistributionEditorMode::Binding:
 		return "NiagaraEditor.DistributionEditor.Binding";
 	case ENiagaraDistributionEditorMode::Expression:
@@ -135,6 +141,12 @@ FSlateIcon FNiagaraDistributionEditorUtilities::DistributionModeToIcon(ENiagaraD
 	return FSlateIcon(FNiagaraEditorStyle::Get().GetStyleSetName(), DistributionModeToIconBrushName(InMode));
 }
 
+bool FNiagaraDistributionEditorUtilities::IsArray(ENiagaraDistributionEditorMode InMode)
+{
+	return
+		InMode == ENiagaraDistributionEditorMode::Array;
+}
+
 bool FNiagaraDistributionEditorUtilities::IsBinding(ENiagaraDistributionEditorMode InMode)
 {
 	return
diff --git a/Engine/Plugins/FX/Niagara/Source/NiagaraEditor/Private/Widgets/NiagaraDistributionEditorUtilities.h b/Engine/Plugins/FX/Niagara/Source/NiagaraEditor/Private/Widgets/NiagaraDistributionEditorUtilities.h
--- a/Engine/Plugins/FX/Niagara/Source/NiagaraEditor/Private/Widgets/NiagaraDistributionEditorUtilities.h
+++ b/Engine/Plugins/FX/Niagara/Source/NiagaraEditor/Private/Widgets/NiagaraDistributionEditorUtilities.h
@@ -19,6 +19,7 @@ namespace FNiagaraDistributionEditorUtilities
 	const FSlateBrush* DistributionModeToIconBrush(ENiagaraDistributionEditorMode InMode);
 	FSlateIcon DistributionModeToIcon(ENiagaraDistributionEditorMode InMode);
 
+	bool IsArray(ENiagaraDistributionEditorMode InMode);
 	bool IsBinding(ENiagaraDistributionEditorMode InMode);
 	bool IsExpression(ENiagaraDistributionEditorMode InMode);
 	bool IsUniform(ENiagaraDistributionEditorMode InMode);
diff --git a/Engine/Plugins/FX/Niagara/Source/NiagaraEditor/Private/Widgets/SNiagaraDistributionArrayEditor.cpp b/Engine/Plugins/FX/Niagara/Source/NiagaraEditor/Private/Widgets/SNiagaraDistributionArrayEditor.cpp
new file mode 100644
--- /dev/null
+++ b/Engine/Plugins/FX/Niagara/Source/NiagaraEditor/Private/Widgets/SNiagaraDistributionArrayEditor.cpp
@@ -0,0 +1,768 @@
+// Copyright Epic Games, Inc. All Rights Reserved.
+
+#include "Widgets/SNiagaraDistributionArrayEditor.h"
+
+#include "Widgets/INiagaraDistributionAdapter.h"
+#include "Widgets/SNiagaraParameterName.h"
+#include "SNiagaraParameterEditor.h"
+#include "NiagaraVariableMetaData.h"
+#include "NiagaraEditorModule.h"
+#include "INiagaraEditorTypeUtilities.h"
+
+#include "Framework/MultiBox/MultiBoxBuilder.h"
+#include "Widgets/Layout/SBox.h"
+#include "Widgets/Layout/SGridPanel.h"
+#include "Widgets/Input/SButton.h"
+#include "Widgets/Input/SComboButton.h"
+#include "UObject/StructOnScope.h"
+#include "PropertyCustomizationHelpers.h"
+#include "PropertyEditorModule.h"
+#include "IDetailCustomization.h"
+#include "DetailCategoryBuilder.h"
+#include "DetailLayoutBuilder.h"
+#include "DetailWidgetRow.h"
+#include "IStructureDetailsView.h"
+
+#define LOCTEXT_NAMESPACE "NiagaraDistributionArrayEditor"
+
+//-TODO: Copy & Paste handling :grimacing:
+
+namespace SNiagaraDistributionArrayEditorPrivate
+{
+	template<typename TEnumType>
+	FText GetEnumToolTipTextByValue(TEnumType EnumValue)
+	{
+		const UEnum* Enum = StaticEnum<TEnumType>();
+		const int64 EnumIndex = Enum->GetIndexByValue(int64(EnumValue));
+		return Enum->GetToolTipTextByIndex(EnumIndex);
+	}
+
+	TSharedRef<SWidget> GetEnumMenuOptions(const UEnum* Enum, TFunction<FExecuteAction(int64)> CreateAction)
+	{
+		FMenuBuilder MenuBuilder(true, nullptr);
+		MenuBuilder.BeginSection("Actions");
+
+		for (int64 iValue=0; iValue < Enum->GetMaxEnumValue(); ++iValue)
+		{
+			if (Enum->IsValidEnumValue(iValue) == false)
+			{
+				continue;
+			}
+
+			const int64 iIndex = Enum->GetIndexByValue(iValue);
+
+			MenuBuilder.AddMenuEntry(
+				Enum->GetDisplayNameTextByIndex(iIndex),
+				Enum->GetToolTipTextByIndex(iIndex),
+				FSlateIcon(),
+				FUIAction(CreateAction(iValue))
+			);
+		}
+		MenuBuilder.EndSection();
+		return MenuBuilder.MakeWidget();
+	}
+
+	int32 GetArrayNum(const TSharedPtr<INiagaraDistributionAdapter>& Distribution)
+	{
+		return Distribution ? Distribution->GetArrayValues().Num() / Distribution->GetNumChannels() : 0;
+	}
+
+	TArray<float> GetArrayElement(const TSharedPtr<INiagaraDistributionAdapter>& Distribution, int32 Index)
+	{
+		TArray<float> Value;
+		if (Distribution && Index >= 0 && Index < GetArrayNum(Distribution))
+		{
+			const int32 NumChannels = Distribution->GetNumChannels();
+			Value.AddDefaulted(NumChannels);
+
+			const int32 ArrayOffset = NumChannels * Index;
+			TConstArrayView<float> ArrayData = Distribution->GetArrayValues();
+			for (int32 i = 0; i < NumChannels; ++i)
+			{
+				Value[i] = ArrayData[ArrayOffset + i];
+			}
+		}
+		return Value;
+	}
+
+	void SetArrayElement(const TSharedPtr<INiagaraDistributionAdapter>& Distribution, int32 Index, TConstArrayView<float> Value)
+	{
+		if (Distribution && Index >= 0 && Index < GetArrayNum(Distribution))
+		{
+			TArray<float> ArrayData(Distribution->GetArrayValues());
+
+			const int32 NumChannels = Distribution->GetNumChannels();
+			const int32 ArrayOffset = NumChannels * Index;
+			for (int32 i = 0; i < NumChannels; ++i)
+			{
+				ArrayData[ArrayOffset + i] = Value[i];
+			}
+
+			Distribution->SetArrayValues(ArrayData);
+		}
+	}
+
+	void SwapArrayElement(const TSharedPtr<INiagaraDistributionAdapter>& Distribution, int32 IndexA, int32 IndexB)
+	{
+		if (Distribution && IndexA >= 0 && IndexA < GetArrayNum(Distribution) && IndexB >= 0 && IndexB < GetArrayNum(Distribution))
+		{
+			TArray<float> ArrayData(Distribution->GetArrayValues());
+
+			const int32 NumChannels = Distribution->GetNumChannels();
+			const int32 ArrayOffsetA = NumChannels * IndexA;
+			const int32 ArrayOffsetB = NumChannels * IndexB;
+			for (int32 i = 0; i < NumChannels; ++i)
+			{
+				Swap(ArrayData[ArrayOffsetA + i], ArrayData[ArrayOffsetB + i]);
+			}
+
+			Distribution->SetArrayValues(ArrayData);
+		}
+	}
+
+	void InsertArrayElement(const TSharedPtr<INiagaraDistributionAdapter>& Distribution, int32 Index)
+	{
+		if (Distribution && Index >= 0 && Index <= GetArrayNum(Distribution))
+		{
+			const int32 NumChannels = Distribution->GetNumChannels();
+			const int32 ArrayOffset = NumChannels * Index;
+
+			TArray<float> ArrayData(Distribution->GetArrayValues());
+			if (ArrayOffset == ArrayData.Num())
+			{
+				for (int32 i = 0; i < NumChannels; ++i)
+				{
+					ArrayData.Emplace(1.0f);
+				}
+			}
+			else
+			{
+				for (int32 i = 0; i < NumChannels; ++i)
+				{
+					ArrayData.Insert(1.0f, ArrayOffset);
+				}
+			}
+
+			Distribution->SetArrayValues(ArrayData);
+		}
+	}
+
+	void DuplicateArrayElement(const TSharedPtr<INiagaraDistributionAdapter>& Distribution, int32 Index)
+	{
+		if (Distribution && Index >= 0 && Index < GetArrayNum(Distribution))
+		{
+			const int32 NumChannels = Distribution->GetNumChannels();
+			const int32 ArrayOffset = NumChannels * Index;
+
+			TArray<float> ElementToDuplicate = GetArrayElement(Distribution, Index);
+
+			TArray<float> ArrayData(Distribution->GetArrayValues());
+			ArrayData.InsertZeroed(ArrayOffset, NumChannels);
+			for (int32 i = 0; i < NumChannels; ++i)
+			{
+				ArrayData[ArrayOffset + i] = ElementToDuplicate[i];
+			}
+
+			Distribution->SetArrayValues(ArrayData);
+		}
+	}
+
+	void RemoveArrayElement(const TSharedPtr<INiagaraDistributionAdapter>& Distribution, int32 Index)
+	{
+		if (Distribution && Index >= 0 && Index < GetArrayNum(Distribution))
+		{
+			const int32 NumChannels = Distribution->GetNumChannels();
+			const int32 ArrayOffset = NumChannels * Index;
+
+			TArray<float> ArrayData(Distribution->GetArrayValues());
+			ArrayData.RemoveAt(ArrayOffset, NumChannels);
+
+			Distribution->SetArrayValues(ArrayData);
+		}
+	}
+
+	class SArrayButton : public SCompoundWidget
+	{
+	public:
+		SLATE_BEGIN_ARGS(SArrayButton) {}
+			SLATE_ARGUMENT(FText, ToolTip)
+			SLATE_ARGUMENT_DEFAULT( const FSlateBrush*, Image ) = FAppStyle::GetBrush("Default");
+			SLATE_ARGUMENT_DEFAULT( bool, IsFocusable ) = false;
+			SLATE_ATTRIBUTE( bool, IsEnabled )
+			SLATE_EVENT(FSimpleDelegate, OnClicked)
+		SLATE_END_ARGS()
+
+		void Construct( const FArguments& InArgs )
+		{
+			OnClickedDelegate = InArgs._OnClicked;
+			ChildSlot
+			[
+				SNew(SBox)
+				.HAlign(HAlign_Center)
+				.VAlign(VAlign_Center)
+				.WidthOverride(22.0f)
+				.HeightOverride(22.0f)
+				.ToolTipText(InArgs._ToolTip)
+				[
+					SNew(SButton)
+					.ButtonStyle( FAppStyle::Get(), "SimpleButton" )
+					.OnClicked(this, &SArrayButton::OnClicked)
+					.ContentPadding(0.0f)
+					.IsFocusable(InArgs._IsFocusable)
+					.IsEnabled(InArgs._IsEnabled)
+					[ 
+						SNew( SImage )
+						.Image( InArgs._Image )
+						.ColorAndOpacity( FSlateColor::UseForeground() )
+					]
+				]
+			]; 
+		}
+
+		FReply OnClicked() const
+		{
+			OnClickedDelegate.ExecuteIfBound();
+			return FReply::Handled();
+		}
+
+		FSimpleDelegate OnClickedDelegate;
+	};
+
+	class SArrayToggleButton : public SCompoundWidget
+	{
+	public:
+		SLATE_BEGIN_ARGS(SArrayToggleButton) {}
+			SLATE_ARGUMENT(FText, ToolTip)
+			SLATE_ARGUMENT_DEFAULT( const FSlateBrush*, Image ) = FAppStyle::GetBrush("Default");
+			SLATE_ARGUMENT_DEFAULT( bool, IsFocusable ) = false;
+			SLATE_ATTRIBUTE( bool, IsEnabled )
+			SLATE_EVENT(FSimpleDelegate, OnClicked)
+		SLATE_END_ARGS()
+
+		void Construct( const FArguments& InArgs )
+		{
+			OnClickedDelegate = InArgs._OnClicked;
+			ChildSlot
+			[
+				SNew(SBox)
+				.HAlign(HAlign_Center)
+				.VAlign(VAlign_Center)
+				.WidthOverride(22.0f)
+				.HeightOverride(22.0f)
+				.ToolTipText(InArgs._ToolTip)
+				[
+					SNew(SButton)
+					.ButtonStyle( FAppStyle::Get(), "ToggleButton" )
+					.OnClicked(this, &SArrayToggleButton::OnClicked)
+					.ContentPadding(0.0f)
+					.IsFocusable(InArgs._IsFocusable)
+					.IsEnabled(InArgs._IsEnabled)
+					[ 
+						SNew( SImage )
+						.Image( InArgs._Image )
+						.ColorAndOpacity( FSlateColor::UseForeground() )
+					]
+				]
+			]; 
+		}
+
+		FReply OnClicked() const
+		{
+			OnClickedDelegate.ExecuteIfBound();
+			return FReply::Handled();
+		}
+
+		FSimpleDelegate OnClickedDelegate;
+	};
+
+	class SArrayEntryWidget : public STableRow<TSharedPtr<int32>>
+	{
+	public:
+		SLATE_BEGIN_ARGS(SArrayEntryWidget) {}
+			SLATE_ARGUMENT(TSharedPtr<INiagaraDistributionAdapter>, DistributionAdapter)
+			SLATE_ARGUMENT(TSharedPtr<int32>, RowIndexPtr)
+			SLATE_EVENT(FSimpleDelegate, OnValueChanged)
+		SLATE_END_ARGS()
+
+		void Construct(const FArguments& InArgs, const TSharedRef<STableViewBase>& InOwnerTableView)
+		{
+			DistributionAdapter = InArgs._DistributionAdapter;
+			RowIndexPtr = InArgs._RowIndexPtr;
+			OnValueChangedDelegate = InArgs._OnValueChanged;
+
+			STableRow::Construct(
+				STableRow<TSharedPtr<int32>>::FArguments()
+				.Style(FAppStyle::Get(), "DataTableEditor.CellListViewRow")
+				, InOwnerTableView
+			);
+
+			if (DistributionAdapter.IsValid() && RowIndexPtr.IsValid())
+			{
+				FNiagaraEditorModule& NiagaraEditorModule = FModuleManager::GetModuleChecked<FNiagaraEditorModule>("NiagaraEditor");
+				const FNiagaraTypeDefinition TypeDef = DistributionAdapter->GetExpressionTypeDef();
+				TSharedPtr<INiagaraEditorTypeUtilities, ESPMode::ThreadSafe> TypeEditorUtilities = NiagaraEditorModule.GetTypeUtilities(TypeDef);
+				if (TypeEditorUtilities.IsValid())
+				{
+					TSharedPtr<IPropertyHandle> PropertyHandle = DistributionAdapter->GetPropertyHandle();
+					PropertyHandle->SetOnPropertyValueChanged(FSimpleDelegate::CreateSP(this, &SArrayEntryWidget::OnExternalChange));
+
+					ElementData = GetArrayElement(DistributionAdapter, *RowIndexPtr);
+
+					ValueStructOnScope = MakeShared<FStructOnScope>(TypeDef.GetStruct(), reinterpret_cast<uint8*>(ElementData.GetData()));
+
+					FNiagaraInputParameterCustomization CustomizationOptions;
+					CustomizationOptions.bBroadcastValueChangesOnCommitOnly = true;
+					ValueParameterEditor = TypeEditorUtilities->CreateParameterEditor(TypeDef, EUnit::Unspecified, CustomizationOptions);
+					ValueParameterEditor->UpdateInternalValueFromStruct(ValueStructOnScope.ToSharedRef());
+					ValueParameterEditor->SetOnValueChanged(
+						SNiagaraParameterEditor::FOnValueChange::CreateSP(this, &SArrayEntryWidget::OnValueChanged)
+					);
+
+					ChildSlot
+					[
+						SNew(SVerticalBox)
+						+ SVerticalBox::Slot()
+						. AutoHeight()
+						[
+							SNew(SHorizontalBox)
+							+ SHorizontalBox::Slot()
+							. AutoWidth()
+							[
+								SNew(SBox)
+								.HAlign(HAlign_Left)
+								.VAlign(VAlign_Center)
+								.Padding(FMargin(0, 0, 5, 0))
+								[
+									SNew(STextBlock)
+									.Font(IDetailLayoutBuilder::GetDetailFont())
+									.Text(FText::Format(LOCTEXT("ArrayEntryFormat", "[{0}]"), FText::AsNumber(*RowIndexPtr)))
+								]
+							]
+							+ SHorizontalBox::Slot()
+							. FillWidth(1.0f)
+							[
+								ValueParameterEditor.ToSharedRef()
+							]
+							+ SHorizontalBox::Slot()
+							. AutoWidth()
+							[
+								SNew(SNiagaraDistributionArrayEditorPrivate::SArrayButton)
+								.Image(FAppStyle::GetBrush(TEXT("Icons.ArrowUp")))
+								.ToolTip(LOCTEXT("MoveElementUp", "Move this element up one"))
+								.IsEnabled(this, &SArrayEntryWidget::CanMoveElementUp)
+								.OnClicked(this, &SArrayEntryWidget::MoveElementUp)
+							]
+							+ SHorizontalBox::Slot()
+							. AutoWidth()
+							[
+								SNew(SNiagaraDistributionArrayEditorPrivate::SArrayButton)
+								.Image(FAppStyle::GetBrush(TEXT("Icons.ArrowDown")))
+								.ToolTip(LOCTEXT("MoveElementDown", "Move this element down one"))
+								.IsEnabled(this, &SArrayEntryWidget::CanMoveElementDown)
+								.OnClicked(this, &SArrayEntryWidget::MoveElementDown)
+							]
+							//+ SHorizontalBox::Slot()
+							//. AutoWidth()
+							//[
+								//SNew(SNiagaraDistributionArrayEditorPrivate::SArrayButton)
+								//.Image(FAppStyle::GetBrush(TEXT("Icons.PlusCircle")))
+								//.ToolTip(LOCTEXT("InsertElement", "Inset a new element above this one"))
+								//.OnClicked(this, &SArrayEntryWidget::InsertElementAbove)
+							//]
+							//+ SHorizontalBox::Slot()
+							//.AutoWidth()
+							//[
+							//	SNew(SNiagaraDistributionArrayEditorPrivate::SArrayButton)
+							//	.Image(FAppStyle::GetBrush(TEXT("Icons.Minus")))
+							//	.ToolTip(LOCTEXT("RemoveElement", "Remove this element"))
+							//	.OnClicked(this, &SArrayEntryWidget::RemoveElement)
+							//]
+							+ SHorizontalBox::Slot()
+							. AutoWidth()
+							[
+								PropertyCustomizationHelpers::MakeInsertDeleteDuplicateButton(
+									FExecuteAction::CreateSP(this, &SArrayEntryWidget::InsertElementAbove),
+									FExecuteAction::CreateSP(this, &SArrayEntryWidget::RemoveElement),
+									FExecuteAction::CreateSP(this, &SArrayEntryWidget::DuplicateElementAbove)
+								)
+							]
+						]
+					];
+				}
+			}
+		}
+
+		bool CanMoveElementUp() const
+		{
+			return RowIndexPtr && *RowIndexPtr > 0;
+		}
+		
+		bool CanMoveElementDown() const
+		{
+			return RowIndexPtr && *RowIndexPtr < GetArrayNum(DistributionAdapter) - 1;
+		}
+
+		void MoveElementUp()
+		{
+			const int32 Index = *RowIndexPtr;
+			SwapArrayElement(DistributionAdapter, Index, Index - 1);
+			OnValueChangedDelegate.ExecuteIfBound();
+		}
+
+		void MoveElementDown()
+		{
+			const int32 Index = *RowIndexPtr;
+			SwapArrayElement(DistributionAdapter, Index, Index + 1);
+			OnValueChangedDelegate.ExecuteIfBound();
+		}
+
+		void InsertElementAbove()
+		{
+			const int32 Index = *RowIndexPtr;
+			InsertArrayElement(DistributionAdapter, Index);
+			OnValueChangedDelegate.ExecuteIfBound();
+		}
+
+		void DuplicateElementAbove()
+		{
+			const int32 Index = *RowIndexPtr;
+			DuplicateArrayElement(DistributionAdapter, Index);
+			OnValueChangedDelegate.ExecuteIfBound();
+		}
+
+		void RemoveElement()
+		{
+			const int32 Index = *RowIndexPtr;
+			RemoveArrayElement(DistributionAdapter, Index);
+			OnValueChangedDelegate.ExecuteIfBound();
+		}
+
+		void OnValueChanged() const
+		{
+			ValueParameterEditor->UpdateStructFromInternalValue(ValueStructOnScope.ToSharedRef());
+			SetArrayElement(DistributionAdapter, *RowIndexPtr, ElementData);
+		}
+
+		void OnExternalChange()
+		{
+			if (*RowIndexPtr >= GetArrayNum(DistributionAdapter))
+			{
+				// This can happen when we remove an element
+				return;
+			}
+
+			TArray<float> NewElementData = GetArrayElement(DistributionAdapter, *RowIndexPtr);
+			for (int32 i = 0; i < ElementData.Num(); ++i)
+			{
+				ElementData[i] = NewElementData[i];
+			}
+			ValueParameterEditor->UpdateInternalValueFromStruct(ValueStructOnScope.ToSharedRef());
+		}
+
+		TSharedPtr<INiagaraDistributionAdapter>		DistributionAdapter;
+		TSharedPtr<int32>							RowIndexPtr;
+		TArray<float>								ElementData;
+		FSimpleDelegate								OnValueChangedDelegate;
+
+		TSharedPtr<FStructOnScope>					ValueStructOnScope;
+		TSharedPtr<SNiagaraParameterEditor>			ValueParameterEditor;
+	};
+}
+
+void SNiagaraDistributionArrayEditor::Construct(const FArguments& InArgs, TSharedRef<INiagaraDistributionAdapter> InDistributionAdapter)
+{
+	DistributionAdapter = InDistributionAdapter;
+
+	TSharedPtr<IPropertyHandle> PropertyHandle = DistributionAdapter->GetPropertyHandle();
+	LookupValueModeEnum = FNiagaraDistributionBase::GetLookupValueModeEnum(PropertyHandle ? PropertyHandle->GetProperty() : nullptr);
+
+	RefreshRows();
+
+	SAssignNew(ArrayListViewWidget, SListView<TSharedPtr<int32>>)
+		.ListItemsSource(&ArrayEntryItems)
+		.OnGenerateRow(this, &SNiagaraDistributionArrayEditor::MakeArrayEntryWidget)
+		.Visibility(EVisibility::Visible)
+		.SelectionMode(ESelectionMode::Single)
+		.ConsumeMouseWheel(EConsumeMouseWheel::Never)
+		.ScrollbarVisibility(EVisibility::Collapsed)
+		;
+
+	FSlimHorizontalToolBarBuilder ToolBarBuilder(TSharedPtr<const FUICommandList>(), FMultiBoxCustomization::None);
+	{
+		ToolBarBuilder.SetStyle(&FAppStyle::Get(), "DetailsView.ExtensionToolBar");
+		ToolBarBuilder.SetIsFocusable(false);
+		ToolBarBuilder.BeginSection("Array");
+
+		ToolBarBuilder.AddToolBarButton(
+			FUIAction(FExecuteAction::CreateSP(this, &SNiagaraDistributionArrayEditor::AddArrayElement)),
+			NAME_None,
+			FText(),
+			LOCTEXT("AddElementToEnd", "Add a new element to the end of the array"),
+			FSlateIcon(FAppStyle::GetAppStyleSetName(), "Icons.PlusCircle"),
+			EUserInterfaceActionType::Button
+		);
+
+		ToolBarBuilder.AddToolBarButton(
+			FUIAction(FExecuteAction::CreateSP(this, &SNiagaraDistributionArrayEditor::RemoveAllArrayElements)),
+			NAME_None,
+			FText(),
+			LOCTEXT("DeleteAllElements", "Delete all elements in the array"),
+			FSlateIcon(FAppStyle::GetAppStyleSetName(), "Icons.Delete"),
+			EUserInterfaceActionType::Button
+		);
+
+		ToolBarBuilder.EndSection();
+	}
+
+	ChildSlot
+	[
+		SNew(SVerticalBox)
+		+ SVerticalBox::Slot()
+		. AutoHeight()
+		[
+			SNew(SGridPanel)
+			+ SGridPanel::Slot(0, 0)
+			.HAlign(HAlign_Left)
+			.VAlign(VAlign_Center)
+			[
+				SNew(STextBlock)
+				.Font(IDetailLayoutBuilder::GetDetailFont())
+				.Text(LOCTEXT("LookupValueMode", "Lookup Value: "))
+			]
+			+ SGridPanel::Slot(1, 0)
+			.HAlign(HAlign_Left)
+			.VAlign(VAlign_Center)
+			[
+				SNew(SComboButton)
+				.OnGetMenuContent(this, &SNiagaraDistributionArrayEditor::OnGetLookupValueModeOptions)
+				.ContentPadding(1)
+				.ButtonContent()
+				[
+					SNew(SNiagaraParameterName)
+					.ParameterName(this, &SNiagaraDistributionArrayEditor::GetLookupValueModeName)
+					.IsReadOnly(true)
+				]
+			]
+
+			+ SGridPanel::Slot(0, 1)
+			.HAlign(HAlign_Left)
+			.VAlign(VAlign_Center)
+			[
+				SNew(STextBlock)
+				.Font(IDetailLayoutBuilder::GetDetailFont())
+				.Text(LOCTEXT("InterpolationMode", "Interpolation: "))
+			]
+			+ SGridPanel::Slot(1, 1)
+			.HAlign(HAlign_Left)
+			.VAlign(VAlign_Center)
+			[
+				SNew(SComboButton)
+				.OnGetMenuContent(this, &SNiagaraDistributionArrayEditor::OnGetInterpolationModeOptions)
+				.ContentPadding(1)
+				//-TODO:.ToolTip(this, &SNiagaraDistributionArrayEditor::GetInterpolationModeToolTip)
+				.ButtonContent()
+				[
+					SNew(STextBlock)
+					.Text(this, &SNiagaraDistributionArrayEditor::GetInterpolationModeText)
+				]
+			]
+
+			+ SGridPanel::Slot(0, 2)
+			.HAlign(HAlign_Left)
+			.VAlign(VAlign_Center)
+			[
+				SNew(STextBlock)
+				.Font(IDetailLayoutBuilder::GetDetailFont())
+				.Text(LOCTEXT("AddessMode", "Address Mode: "))
+			]
+			+ SGridPanel::Slot(1, 2)
+			.HAlign(HAlign_Left)
+			.VAlign(VAlign_Center)
+			[
+				SNew(SComboButton)
+				.OnGetMenuContent(this, &SNiagaraDistributionArrayEditor::OnGetAddressModeOptions)
+				.ContentPadding(1)
+				//-TODO:.ToolTip(this, &SNiagaraDistributionArrayEditor::GetAddressModeToolTip)
+				.ButtonContent()
+				[
+					SNew(STextBlock)
+					.Text(this, &SNiagaraDistributionArrayEditor::GetAddressModeText)
+				]
+			]
+		]
+		//////////////////////////////////////////////////////
+		+ SVerticalBox::Slot()
+		. AutoHeight()
+		[
+			SNew(SHorizontalBox)
+			+ SHorizontalBox::Slot()
+			. AutoWidth()
+			[
+				SNew(SBox)
+				.HAlign(HAlign_Left)
+				.VAlign(VAlign_Center)
+				.Padding(0,0,5,0)
+				[
+					SNew(STextBlock)
+					.Font(IDetailLayoutBuilder::GetDetailFont())
+					.Text(this, &SNiagaraDistributionArrayEditor::GetArrayHeaderText)
+				]
+			]
+			+ SHorizontalBox::Slot()
+			.AutoWidth()
+			[
+				ToolBarBuilder.MakeWidget()
+			]
+		]
+		+ SVerticalBox::Slot()
+		. AutoHeight()
+		[
+			ArrayListViewWidget.ToSharedRef()
+		]
+	];
+}
+
+int32 SNiagaraDistributionArrayEditor::GetArrayNum() const
+{
+	return SNiagaraDistributionArrayEditorPrivate::GetArrayNum(DistributionAdapter);
+}
+
+FText SNiagaraDistributionArrayEditor::GetArrayHeaderText() const
+{
+	const int32 NumElements = GetArrayNum();
+	return FText::Format(NumElements > 1 ? LOCTEXT("NumArrayElements", "{0} array elements") : LOCTEXT("NumArrayElement", "{0} array element"), FText::AsNumber(NumElements));
+}
+
+void SNiagaraDistributionArrayEditor::AddArrayElement()
+{
+	const int32 NumEntries = SNiagaraDistributionArrayEditorPrivate::GetArrayNum(DistributionAdapter);
+	if (NumEntries == 0)
+	{
+		SNiagaraDistributionArrayEditorPrivate::InsertArrayElement(DistributionAdapter, 0);
+	}
+	else
+	{
+		SNiagaraDistributionArrayEditorPrivate::DuplicateArrayElement(DistributionAdapter, NumEntries - 1);
+	}
+	RefreshRows();
+}
+
+void SNiagaraDistributionArrayEditor::RemoveAllArrayElements()
+{
+	DistributionAdapter->SetArrayValues({});
+
+	RefreshRows();
+}
+
+void SNiagaraDistributionArrayEditor::SetLookupValueMode(uint8 Mode)
+{
+	DistributionAdapter->SetLookupValueMode(Mode);
+}
+
+FName SNiagaraDistributionArrayEditor::GetLookupValueModeName() const
+{
+	const uint8 Mode = DistributionAdapter->GetLookupValueMode();
+	return FName(LookupValueModeEnum->GetDisplayNameTextByValue(Mode).ToString());
+}
+
+FText SNiagaraDistributionArrayEditor::GetLookupValueModeToolTip() const
+{
+	const uint8 Mode = DistributionAdapter->GetLookupValueMode();
+	const int64 Index = LookupValueModeEnum->GetIndexByValue(Mode);
+	return Index != INDEX_NONE ? LookupValueModeEnum->GetToolTipTextByIndex(Index) : FText();
+}
+
+void SNiagaraDistributionArrayEditor::SetInterpolationMode(ENiagaraDistributionInterpolationMode Mode)
+{
+	DistributionAdapter->SetInterpolationMode(Mode);
+}
+
+FText SNiagaraDistributionArrayEditor::GetInterpolationModeText() const
+{
+	return UEnum::GetDisplayValueAsText(DistributionAdapter->GetInterpolationMode());
+}
+
+FText SNiagaraDistributionArrayEditor::GetInterpolationModeToolTip() const
+{
+	return SNiagaraDistributionArrayEditorPrivate::GetEnumToolTipTextByValue(DistributionAdapter->GetInterpolationMode());
+}
+
+void SNiagaraDistributionArrayEditor::SetAddressMode(ENiagaraDistributionAddressMode Mode)
+{
+	DistributionAdapter->SetAddressMode(Mode);
+}
+
+FText SNiagaraDistributionArrayEditor::GetAddressModeText() const
+{
+	return UEnum::GetDisplayValueAsText(DistributionAdapter->GetAddressMode());
+}
+
+FText SNiagaraDistributionArrayEditor::GetAddressModeToolTip() const
+{
+	return SNiagaraDistributionArrayEditorPrivate::GetEnumToolTipTextByValue(DistributionAdapter->GetAddressMode());
+}
+
+void SNiagaraDistributionArrayEditor::RefreshRows()
+{
+	const int32 NumElements = GetArrayNum();
+	if (ArrayEntryItems.Num() > NumElements)
+	{
+		ArrayEntryItems.SetNum(NumElements);
+	}
+	else
+	{
+		while (ArrayEntryItems.Num() < NumElements)
+		{
+			ArrayEntryItems.Emplace(MakeShared<int32>(ArrayEntryItems.Num()));
+		}
+	}
+	if (ArrayListViewWidget.IsValid())
+	{
+		ArrayListViewWidget->RequestListRefresh();
+	}
+}
+
+TSharedRef<ITableRow> SNiagaraDistributionArrayEditor::MakeArrayEntryWidget(const TSharedPtr<int32> RowIndexPtr, const TSharedRef<STableViewBase>& OwnerTable)
+{
+	return
+		SNew(SNiagaraDistributionArrayEditorPrivate::SArrayEntryWidget, OwnerTable)
+		.DistributionAdapter(DistributionAdapter)
+		.RowIndexPtr(RowIndexPtr)
+		.OnValueChanged(this, &SNiagaraDistributionArrayEditor::RefreshRows)
+		;
+}
+
+TSharedRef<SWidget> SNiagaraDistributionArrayEditor::OnGetLookupValueModeOptions()
+{
+	return SNiagaraDistributionArrayEditorPrivate::GetEnumMenuOptions(
+		LookupValueModeEnum,
+		[&](int64 iValue)
+		{
+			return FExecuteAction::CreateSP(this, &SNiagaraDistributionArrayEditor::SetLookupValueMode, uint8(iValue));
+		}
+	);
+}
+
+TSharedRef<SWidget> SNiagaraDistributionArrayEditor::OnGetInterpolationModeOptions()
+{
+	return SNiagaraDistributionArrayEditorPrivate::GetEnumMenuOptions(
+		StaticEnum<ENiagaraDistributionInterpolationMode>(),
+		[&](int64 iValue)
+		{
+			return FExecuteAction::CreateSP(this, &SNiagaraDistributionArrayEditor::SetInterpolationMode, ENiagaraDistributionInterpolationMode(iValue));
+		}
+	);
+}
+
+TSharedRef<SWidget> SNiagaraDistributionArrayEditor::OnGetAddressModeOptions()
+{
+	return SNiagaraDistributionArrayEditorPrivate::GetEnumMenuOptions(
+		StaticEnum<ENiagaraDistributionAddressMode>(),
+		[&](int64 iValue)
+		{
+			return FExecuteAction::CreateSP(this, &SNiagaraDistributionArrayEditor::SetAddressMode, ENiagaraDistributionAddressMode(iValue));
+		}
+	);
+}
+
+#undef LOCTEXT_NAMESPACE
diff --git a/Engine/Plugins/FX/Niagara/Source/NiagaraEditor/Private/Widgets/SNiagaraDistributionArrayEditor.h b/Engine/Plugins/FX/Niagara/Source/NiagaraEditor/Private/Widgets/SNiagaraDistributionArrayEditor.h
new file mode 100644
--- /dev/null
+++ b/Engine/Plugins/FX/Niagara/Source/NiagaraEditor/Private/Widgets/SNiagaraDistributionArrayEditor.h
@@ -0,0 +1,77 @@
+// Copyright Epic Games, Inc. All Rights Reserved.
+
+#pragma once
+
+#include "Stateless/NiagaraStatelessDistribution.h"
+
+#include "Misc/NotifyHook.h"
+#include "Widgets/Views/SListView.h"
+#include "Widgets/SCompoundWidget.h"
+
+#include "SNiagaraDistributionArrayEditor.generated.h"
+
+class INiagaraDistributionAdapter;
+class ITableRow;
+
+USTRUCT()
+struct FNiagaraDistributionArray
+{
+	GENERATED_BODY()
+
+	UPROPERTY(EditAnywhere, Category="Array")
+	TArray<float> FloatData;
+
+	UPROPERTY(EditAnywhere, Category="Array")
+	TArray<FVector2f> Vector2Data;
+
+	UPROPERTY(EditAnywhere, Category="Array")
+	TArray<FVector3f> Vector3Data;
+
+	UPROPERTY(EditAnywhere, Category="Array")
+	TArray<FVector4f> Vector4Data;
+};
+
+class SNiagaraDistributionArrayEditor : public SCompoundWidget
+{
+public:
+	SLATE_BEGIN_ARGS(SNiagaraDistributionArrayEditor) { }
+	SLATE_END_ARGS();
+
+	void Construct(const FArguments& InArgs, TSharedRef<INiagaraDistributionAdapter> InDistributionAdapter);
+
+private:
+	int32 GetArrayNum() const;
+	FText GetArrayHeaderText() const;
+
+	void AddArrayElement();
+	void RemoveAllArrayElements();
+
+	void RefreshRows();
+
+	TSharedRef<ITableRow> MakeArrayEntryWidget(const TSharedPtr<int32> RowIndexPtr, const TSharedRef<STableViewBase>& OwnerTable);
+
+	TSharedRef<SWidget> OnGetLookupValueModeOptions();
+	TSharedRef<SWidget> OnGetInterpolationModeOptions();
+	TSharedRef<SWidget> OnGetAddressModeOptions();
+
+	void SetLookupValueMode(uint8 Mode);
+	FName GetLookupValueModeName() const;
+	FText GetLookupValueModeToolTip() const;
+
+	void SetInterpolationMode(ENiagaraDistributionInterpolationMode Mode);
+	FText GetInterpolationModeText() const;
+	FText GetInterpolationModeToolTip() const;
+
+	void SetAddressMode(ENiagaraDistributionAddressMode Mode);
+	FText GetAddressModeText() const;
+	FText GetAddressModeToolTip() const;
+
+private:
+	TSharedPtr<INiagaraDistributionAdapter>		DistributionAdapter;
+	FNiagaraDistributionArray					ArrayStruct;
+
+	TSharedPtr<SListView<TSharedPtr<int32>>>	ArrayListViewWidget;
+	TArray<TSharedPtr<int32>>					ArrayEntryItems;
+
+	const UEnum*								LookupValueModeEnum = nullptr;
+};
diff --git a/Engine/Plugins/FX/Niagara/Source/NiagaraEditor/Private/Widgets/SNiagaraDistributionEditor.cpp b/Engine/Plugins/FX/Niagara/Source/NiagaraEditor/Private/Widgets/SNiagaraDistributionEditor.cpp
--- a/Engine/Plugins/FX/Niagara/Source/NiagaraEditor/Private/Widgets/SNiagaraDistributionEditor.cpp
+++ b/Engine/Plugins/FX/Niagara/Source/NiagaraEditor/Private/Widgets/SNiagaraDistributionEditor.cpp
@@ -16,6 +16,7 @@
 #include "Widgets/Layout/SGridPanel.h"
 #include "Widgets/NiagaraDistributionEditorUtilities.h"
 #include "Widgets/SNiagaraColorEditor.h"
+#include "Widgets/SNiagaraDistributionArrayEditor.h"
 #include "Widgets/SNiagaraDistributionCurveEditor.h"
 #include "Widgets/SNiagaraParameterName.h"
 
@@ -719,7 +720,11 @@ void  SNiagaraDistributionEditor::OnDistributionModeChanged()
 TSharedRef<SWidget> SNiagaraDistributionEditor::ConstructContentForMode()
 {
 	ENiagaraDistributionEditorMode Mode = DistributionAdapter->GetDistributionMode();
-	if (FNiagaraDistributionEditorUtilities::IsBinding(Mode))
+	if (FNiagaraDistributionEditorUtilities::IsArray(Mode))
+	{
+		return SNew(SNiagaraDistributionArrayEditor, DistributionAdapter.ToSharedRef());
+	}
+	else if (FNiagaraDistributionEditorUtilities::IsBinding(Mode))
 	{
 		return SNew(SNiagaraDistributionBindingEditor, DistributionAdapter.ToSharedRef());
 	}
diff --git a/Engine/Plugins/FX/Niagara/Source/NiagaraShader/Internal/Stateless/NiagaraStatelessSimulationShader.h b/Engine/Plugins/FX/Niagara/Source/NiagaraShader/Internal/Stateless/NiagaraStatelessSimulationShader.h
--- a/Engine/Plugins/FX/Niagara/Source/NiagaraShader/Internal/Stateless/NiagaraStatelessSimulationShader.h
+++ b/Engine/Plugins/FX/Niagara/Source/NiagaraShader/Internal/Stateless/NiagaraStatelessSimulationShader.h
@@ -35,8 +35,8 @@ namespace NiagaraStateless
 		//SHADER_PARAMETER_UAV(RWBuffer<half>,	Common_HalfOutputBuffer)
 		SHADER_PARAMETER_UAV(RWBuffer<int>,		Common_IntOutputBuffer)
 		SHADER_PARAMETER_UAV(RWBuffer<int>,		Common_GPUCountBuffer)
-		SHADER_PARAMETER_SRV(Buffer<float>,		Common_StaticFloatBuffer)
-		SHADER_PARAMETER_SRV(Buffer<uint32>,	Common_ParameterBuffer)
+		SHADER_PARAMETER_SRV(ByteAddressBuffer,	Common_StaticDataBuffer)
+		SHADER_PARAMETER_SRV(ByteAddressBuffer,	Common_DynamicDataBuffer)
 
 		SHADER_PARAMETER_STRUCT_INCLUDE(FSpawnInfoShaderParameters,	SpawnParameters)
 	END_SHADER_PARAMETER_STRUCT()
