- Update scale by modules to use a distribution

#rnx

diff --git a/Engine/Plugins/FX/Niagara/Shaders/Private/Stateless/Modules/NiagaraStatelessModule_ScaleMeshSizeBySpeed.ush b/Engine/Plugins/FX/Niagara/Shaders/Private/Stateless/Modules/NiagaraStatelessModule_ScaleMeshSizeBySpeed.ush
--- a/Engine/Plugins/FX/Niagara/Shaders/Private/Stateless/Modules/NiagaraStatelessModule_ScaleMeshSizeBySpeed.ush
+++ b/Engine/Plugins/FX/Niagara/Shaders/Private/Stateless/Modules/NiagaraStatelessModule_ScaleMeshSizeBySpeed.ush
@@ -2,17 +2,17 @@
 
 float3	ScaleMeshSizeBySpeed_ScaleFactorBias;
 float3	ScaleMeshSizeBySpeed_ScaleFactorScale;
-uint2	ScaleMeshSizeBySpeed_ScaleDistribution;
+FNiagaraStatelessBuiltDistributionType	ScaleMeshSizeBySpeed_ScaleDistribution;
 float	ScaleMeshSizeBySpeed_VelocityNorm;
 
 void ScaleMeshSizeBySpeed_Simulate(inout FStatelessParticle Particle)
 {
-	if (ScaleMeshSizeBySpeed_ScaleDistribution.y > 0)
+	if (FNiagaraStatelessBuiltDistribution::IsValid(ScaleMeshSizeBySpeed_ScaleDistribution))
 	{
 		const float3 Velocity	= (Particle.Position - Particle.PreviousPosition) * Particle.InvDeltaTime;
 		const float Speed		= length2(Velocity);
 		const float NormSpeed	= saturate(Speed * ScaleMeshSizeBySpeed_VelocityNorm);
-		const float Interp		= saturate(LerpStaticFloat(ScaleMeshSizeBySpeed_ScaleDistribution, NormSpeed));
+		const float Interp		= saturate(SampleDistributionValueFloat(ScaleMeshSizeBySpeed_ScaleDistribution, 0, NormSpeed));
 		const float3 Scale		= ScaleMeshSizeBySpeed_ScaleFactorBias + (ScaleMeshSizeBySpeed_ScaleFactorScale * Interp);
 			
 		Particle.Scale			*= Scale;
diff --git a/Engine/Plugins/FX/Niagara/Shaders/Private/Stateless/Modules/NiagaraStatelessModule_ScaleSpriteSizeBySpeed.ush b/Engine/Plugins/FX/Niagara/Shaders/Private/Stateless/Modules/NiagaraStatelessModule_ScaleSpriteSizeBySpeed.ush
--- a/Engine/Plugins/FX/Niagara/Shaders/Private/Stateless/Modules/NiagaraStatelessModule_ScaleSpriteSizeBySpeed.ush
+++ b/Engine/Plugins/FX/Niagara/Shaders/Private/Stateless/Modules/NiagaraStatelessModule_ScaleSpriteSizeBySpeed.ush
@@ -2,17 +2,17 @@
 
 float2	ScaleSpriteSizeBySpeed_ScaleFactorBias;
 float2	ScaleSpriteSizeBySpeed_ScaleFactorScale;
-uint2	ScaleSpriteSizeBySpeed_ScaleDistribution;
+FNiagaraStatelessBuiltDistributionType	ScaleSpriteSizeBySpeed_ScaleDistribution;
 float	ScaleSpriteSizeBySpeed_VelocityNorm;
 
 void ScaleSpriteSizeBySpeed_Simulate(inout FStatelessParticle Particle)
 {
-	if (ScaleSpriteSizeBySpeed_ScaleDistribution.y > 0)
+	if (FNiagaraStatelessBuiltDistribution::IsValid(ScaleSpriteSizeBySpeed_ScaleDistribution))
 	{
 		const float3 Velocity	= (Particle.Position - Particle.PreviousPosition) * Particle.InvDeltaTime;
 		const float Speed		= length2(Velocity);
 		const float NormSpeed	= saturate(Speed * ScaleSpriteSizeBySpeed_VelocityNorm);
-		const float Interp		= saturate(LerpStaticFloat(ScaleSpriteSizeBySpeed_ScaleDistribution, NormSpeed));
+		const float Interp		= saturate(SampleDistributionValueFloat(ScaleSpriteSizeBySpeed_ScaleDistribution, 0, NormSpeed));
 		const float2 Scale		= ScaleSpriteSizeBySpeed_ScaleFactorBias + (ScaleSpriteSizeBySpeed_ScaleFactorScale * Interp);
 			
 		Particle.SpriteSize			*= Scale;
diff --git a/Engine/Plugins/FX/Niagara/Shaders/Private/Stateless/NiagaraStatelessCommon.ush b/Engine/Plugins/FX/Niagara/Shaders/Private/Stateless/NiagaraStatelessCommon.ush
--- a/Engine/Plugins/FX/Niagara/Shaders/Private/Stateless/NiagaraStatelessCommon.ush
+++ b/Engine/Plugins/FX/Niagara/Shaders/Private/Stateless/NiagaraStatelessCommon.ush
@@ -129,30 +129,6 @@ float2 GetParameterBufferFloat2(int BaseOffset, int Element, float2 DefaultValue
 float3 GetParameterBufferFloat3(int BaseOffset, int Element, float3 DefaultValue) { return BaseOffset >= 0 ? asfloat(Common_DynamicDataBuffer.Load3((BaseOffset + Element*3) * 4)) : DefaultValue; }
 float4 GetParameterBufferFloat4(int BaseOffset, int Element, float4 DefaultValue) { return BaseOffset >= 0 ? asfloat(Common_DynamicDataBuffer.Load4((BaseOffset + Element*4) * 4)) : DefaultValue; }
 	
-float LerpStaticFloat (uint2 Parameters, float U)
-{
-	const float Offset = U * float(Parameters.y);
-	const float Value0 = GetStaticFloat(Parameters.x, floor(Offset));
-	const float Value1 = GetStaticFloat(Parameters.x, ceil(Offset));
-	return lerp(Value0, Value1, frac(Offset));
-}
-
-float2 LerpStaticFloat2(uint2 Parameters, float U)
-{
-	const float  Offset = U * float(Parameters.y);
-	const float2 Value0 = GetStaticFloat2(Parameters.x, floor(Offset));
-	const float2 Value1 = GetStaticFloat2(Parameters.x, ceil(Offset));
-	return lerp(Value0, Value1, frac(Offset));
-}
-
-float3 LerpStaticFloat3(uint2 Parameters, float U)
-{
-	const float  Offset = U * float(Parameters.y);
-	const float3 Value0 = GetStaticFloat3(Parameters.x, floor(Offset));
-	const float3 Value1 = GetStaticFloat3(Parameters.x, ceil(Offset));
-	return lerp(Value0, Value1, frac(Offset));
-}
-
 float  StatelessLerp(float  lhs, float  rhs, float  u, bool bUniform) { return lerp(lhs, rhs, u); }
 float2 StatelessLerp(float2 lhs, float2 rhs, float2 u, bool bUniform) { return lerp(lhs, rhs, bUniform ? u.xx : u); }
 float3 StatelessLerp(float3 lhs, float3 rhs, float3 u, bool bUniform) { return lerp(lhs, rhs, bUniform ? u.xxx : u); }
diff --git a/Engine/Plugins/FX/Niagara/Source/Niagara/Internal/Stateless/Modules/NiagaraStatelessModule_ScaleSpriteSizeBySpeed.h b/Engine/Plugins/FX/Niagara/Source/Niagara/Internal/Stateless/Modules/NiagaraStatelessModule_ScaleSpriteSizeBySpeed.h
--- a/Engine/Plugins/FX/Niagara/Source/Niagara/Internal/Stateless/Modules/NiagaraStatelessModule_ScaleSpriteSizeBySpeed.h
+++ b/Engine/Plugins/FX/Niagara/Source/Niagara/Internal/Stateless/Modules/NiagaraStatelessModule_ScaleSpriteSizeBySpeed.h
@@ -14,19 +14,6 @@ class UNiagaraStatelessModule_ScaleSpriteSizeBySpeed : public UNiagaraStatelessM
 {
 	GENERATED_BODY()
 
-	struct FModuleBuiltData
-	{
-		FNiagaraStatelessRangeVector2	MinScaleFactor = FNiagaraStatelessRangeVector2(FVector2f::One());
-		FNiagaraStatelessRangeVector2	MaxScaleFactor = FNiagaraStatelessRangeVector2(FVector2f::One());
-		FNiagaraStatelessRangeFloat		VelocityNorm = FNiagaraStatelessRangeFloat(0.0f);
-		FUintVector2					ScaleDistribution = FUintVector2::ZeroValue;
-
-		int32							PositionVariableOffset = INDEX_NONE;
-		int32							PreviousPositionVariableOffset = INDEX_NONE;
-		int32							SpriteSizeVariableOffset = INDEX_NONE;
-		int32							PreviousSpriteSizeVariableOffset = INDEX_NONE;
-	};
-
 	static constexpr float DefaultVelocity = 1000.0f;
 
 public:
diff --git a/Engine/Plugins/FX/Niagara/Source/Niagara/Internal/Stateless/NiagaraStatelessParticleSimContext.h b/Engine/Plugins/FX/Niagara/Source/Niagara/Internal/Stateless/NiagaraStatelessParticleSimContext.h
--- a/Engine/Plugins/FX/Niagara/Source/Niagara/Internal/Stateless/NiagaraStatelessParticleSimContext.h
+++ b/Engine/Plugins/FX/Niagara/Source/Niagara/Internal/Stateless/NiagaraStatelessParticleSimContext.h
@@ -120,15 +120,6 @@ public:
 	template<> FVector4f GetStaticFloat<FVector4f>(uint32 Offset, uint32 Element) const { Offset += Element * 4; return FVector4f(StaticFloatData[Offset], StaticFloatData[Offset + 1], StaticFloatData[Offset + 2], StaticFloatData[Offset + 3]); }
 	template<> FLinearColor GetStaticFloat<FLinearColor>(uint32 Offset, uint32 Element) const { Offset += Element * 4; return FLinearColor(StaticFloatData[Offset], StaticFloatData[Offset + 1], StaticFloatData[Offset + 2], StaticFloatData[Offset + 3]); }
 
-	template<typename TType>
-	TType LerpStaticFloat(FUintVector2 Parameters, float U) const
-	{
-		const float Offset = U * float(Parameters.Y);
-		const TType Value0 = GetStaticFloat<TType>(Parameters.X, FMath::FloorToInt(Offset));
-		const TType Value1 = GetStaticFloat<TType>(Parameters.X, FMath::CeilToInt(Offset));
-		return FMath::Lerp(Value0, Value1, FMath::Fractional(Offset));
-	}
-
 	float		Lerp(const float& Lhs, const float& Rhs, const float& U, bool bUniform) const { return FMath::Lerp(Lhs, Rhs, U); }
 	FVector2f	Lerp(const FVector2f& Lhs, const FVector2f& Rhs, const FVector2f& U, bool bUniform) const { return FMath::Lerp(Lhs, Rhs, bUniform ? FVector2f(U.X, U.X) : U); }
 	FVector3f	Lerp(const FVector3f& Lhs, const FVector3f& Rhs, const FVector3f& U, bool bUniform) const { return FMath::Lerp(Lhs, Rhs, bUniform ? FVector3f(U.X, U.X, U.X) : U); }
diff --git a/Engine/Plugins/FX/Niagara/Source/Niagara/Private/Stateless/Modules/NiagaraStatelessModule_ScaleMeshSizeBySpeed.cpp b/Engine/Plugins/FX/Niagara/Source/Niagara/Private/Stateless/Modules/NiagaraStatelessModule_ScaleMeshSizeBySpeed.cpp
--- a/Engine/Plugins/FX/Niagara/Source/Niagara/Private/Stateless/Modules/NiagaraStatelessModule_ScaleMeshSizeBySpeed.cpp
+++ b/Engine/Plugins/FX/Niagara/Source/Niagara/Private/Stateless/Modules/NiagaraStatelessModule_ScaleMeshSizeBySpeed.cpp
@@ -13,7 +13,7 @@ namespace NSMScaleMeshSizeBySpeedPrivate
 		FNiagaraStatelessRangeVector3	MinScaleFactor = FNiagaraStatelessRangeVector3(FVector3f::OneVector);
 		FNiagaraStatelessRangeVector3	MaxScaleFactor = FNiagaraStatelessRangeVector3(FVector3f::OneVector);
 		FNiagaraStatelessRangeFloat		VelocityNorm = FNiagaraStatelessRangeFloat(0.0f);
-		FUintVector2					ScaleDistribution = FUintVector2::ZeroValue;
+		FUintVector3					ScaleDistribution = FUintVector3::ZeroValue;
 
 		int32							PositionVariableOffset = INDEX_NONE;
 		int32							PreviousPositionVariableOffset = INDEX_NONE;
@@ -36,7 +36,7 @@ namespace NSMScaleMeshSizeBySpeedPrivate
 			const FVector3f Velocity = (Position - PreviousPosition) * ParticleSimulationContext.GetInvDeltaTime();
 			const float Speed = Velocity.SquaredLength();
 			const float NormSpeed = FMath::Clamp(Speed * Parameters->ScaleMeshSizeBySpeed_VelocityNorm, 0.0f, 1.0f);
-			const float Interp = ParticleSimulationContext.LerpStaticFloat<float>(Parameters->ScaleMeshSizeBySpeed_ScaleDistribution, NormSpeed);
+			const float Interp = ParticleSimulationContext.SampleDistributionValue<float>(Parameters->ScaleMeshSizeBySpeed_ScaleDistribution, i, 0, NormSpeed);
 			const FVector3f	Scale = Parameters->ScaleMeshSizeBySpeed_ScaleFactorBias + (Parameters->ScaleMeshSizeBySpeed_ScaleFactorScale * FMath::Clamp(Interp, 0.0f, 1.0f));
 
 			FVector3f ParticleScale = ParticleSimulationContext.ReadParticleVariable(ModuleBuiltData->ScaleVariableVariableOffset, i, FVector3f::OneVector);
@@ -77,24 +77,12 @@ void UNiagaraStatelessModule_ScaleMeshSizeBySpeed::BuildEmitterData(const FNiaga
 	BuiltData->MaxScaleFactor = BuildContext.ConvertDistributionToRange(MaxScaleFactor, FVector3f::OneVector);
 	if ( bSampleScaleFactorByCurve )
 	{
-		if (SampleFactorCurve.IsCurve() && SampleFactorCurve.Values.Num() > 1)
-		{
-			BuiltData->ScaleDistribution.X = BuildContext.AddStaticData(SampleFactorCurve.Values);
-			BuiltData->ScaleDistribution.Y = SampleFactorCurve.Values.Num() - 1;
-		}
-		else
-		{
-			const float Value = SampleFactorCurve.Values.Num() > 0 ? SampleFactorCurve.Values[0] : 1.0f;
-			const float Values[] = { Value, Value };
-			BuiltData->ScaleDistribution.X = BuildContext.AddStaticData(Values);
-			BuiltData->ScaleDistribution.Y = 1;
-		}
+		BuiltData->ScaleDistribution = BuildContext.AddDistribution(SampleFactorCurve);
 	}
 	else
 	{
-		const float Values[] = { 0.0f, 1.0f };
-		BuiltData->ScaleDistribution.X = BuildContext.AddStaticData(Values);
-		BuiltData->ScaleDistribution.Y = 1;
+		const FNiagaraDistributionFloat DefaultDistribution({ 0.0f, 1.0f });
+		BuiltData->ScaleDistribution = BuildContext.AddDistribution(DefaultDistribution);
 	}
 	BuildContext.AddParticleSimulationExecSimulate(&NSMScaleMeshSizeBySpeedPrivate::ParticleSimulate);
 }
diff --git a/Engine/Plugins/FX/Niagara/Source/Niagara/Private/Stateless/Modules/NiagaraStatelessModule_ScaleSpriteSizeBySpeed.cpp b/Engine/Plugins/FX/Niagara/Source/Niagara/Private/Stateless/Modules/NiagaraStatelessModule_ScaleSpriteSizeBySpeed.cpp
--- a/Engine/Plugins/FX/Niagara/Source/Niagara/Private/Stateless/Modules/NiagaraStatelessModule_ScaleSpriteSizeBySpeed.cpp
+++ b/Engine/Plugins/FX/Niagara/Source/Niagara/Private/Stateless/Modules/NiagaraStatelessModule_ScaleSpriteSizeBySpeed.cpp
@@ -12,7 +12,7 @@ namespace NSMScaleSpriteSizeBySpeedPrivate
 		FNiagaraStatelessRangeVector2	MinScaleFactor = FNiagaraStatelessRangeVector2(FVector2f::One());
 		FNiagaraStatelessRangeVector2	MaxScaleFactor = FNiagaraStatelessRangeVector2(FVector2f::One());
 		FNiagaraStatelessRangeFloat		VelocityNorm = FNiagaraStatelessRangeFloat(0.0f);
-		FUintVector2					ScaleDistribution = FUintVector2::ZeroValue;
+		FUintVector3					ScaleDistribution = FUintVector3::ZeroValue;
 
 		int32							PositionVariableOffset = INDEX_NONE;
 		int32							PreviousPositionVariableOffset = INDEX_NONE;
@@ -34,7 +34,7 @@ namespace NSMScaleSpriteSizeBySpeedPrivate
 			const FVector3f Velocity			= (Position - PreviousPosition) * ParticleSimulationContext.GetInvDeltaTime();
 			const float Speed					= Velocity.SquaredLength();
 			const float NormSpeed				= FMath::Clamp(Speed * Parameters->ScaleSpriteSizeBySpeed_VelocityNorm, 0.0f, 1.0f);
-			const float Interp					= ParticleSimulationContext.LerpStaticFloat<float>(Parameters->ScaleSpriteSizeBySpeed_ScaleDistribution, NormSpeed);
+			const float Interp					= ParticleSimulationContext.SampleDistributionValue<float>(Parameters->ScaleSpriteSizeBySpeed_ScaleDistribution, i, 0, NormSpeed);
 			const FVector2f	Scale				= Parameters->ScaleSpriteSizeBySpeed_ScaleFactorBias + (Parameters->ScaleSpriteSizeBySpeed_ScaleFactorScale * FMath::Clamp(Interp, 0.0f, 1.0f));
 
 			FVector2f SpriteSize			= ParticleSimulationContext.ReadParticleVariable(ModuleBuiltData->SpriteSizeVariableOffset, i, FVector2f::One());
@@ -75,24 +75,12 @@ void UNiagaraStatelessModule_ScaleSpriteSizeBySpeed::BuildEmitterData(const FNia
 	BuiltData->MaxScaleFactor = BuildContext.ConvertDistributionToRange(MaxScaleFactor, FVector2f::One());
 	if (bSampleScaleFactorByCurve)
 	{
-		if (SampleFactorCurve.IsCurve() && SampleFactorCurve.Values.Num() > 1)
-		{
-			BuiltData->ScaleDistribution.X = BuildContext.AddStaticData(SampleFactorCurve.Values);
-			BuiltData->ScaleDistribution.Y = SampleFactorCurve.Values.Num() - 1;
-		}
-		else
-		{
-			const float Value = SampleFactorCurve.Values.Num() > 0 ? SampleFactorCurve.Values[0] : 1.0f;
-			const float Values[] = { Value, Value };
-			BuiltData->ScaleDistribution.X = BuildContext.AddStaticData(Values);
-			BuiltData->ScaleDistribution.Y = 1;
-		}
+		BuiltData->ScaleDistribution = BuildContext.AddDistribution(SampleFactorCurve);
 	}
 	else
 	{
-		const float Values[] = { 0.0f, 1.0f };
-		BuiltData->ScaleDistribution.X = BuildContext.AddStaticData(Values);
-		BuiltData->ScaleDistribution.Y = 1;
+		const FNiagaraDistributionFloat DefaultDistribution({ 0.0f, 1.0f });
+		BuiltData->ScaleDistribution = BuildContext.AddDistribution(DefaultDistribution);
 	}
 	BuildContext.AddParticleSimulationExecSimulate(&ParticleSimulate);
 }
diff --git a/Engine/Plugins/FX/Niagara/Source/NiagaraShader/Internal/Stateless/NiagaraStatelessModuleShaderParameters.h b/Engine/Plugins/FX/Niagara/Source/NiagaraShader/Internal/Stateless/NiagaraStatelessModuleShaderParameters.h
--- a/Engine/Plugins/FX/Niagara/Source/NiagaraShader/Internal/Stateless/NiagaraStatelessModuleShaderParameters.h
+++ b/Engine/Plugins/FX/Niagara/Source/NiagaraShader/Internal/Stateless/NiagaraStatelessModuleShaderParameters.h
@@ -84,10 +84,10 @@ namespace NiagaraStateless
 	END_SHADER_PARAMETER_STRUCT()
 
 	BEGIN_SHADER_PARAMETER_STRUCT(FScaleMeshSizeBySpeedModule_ShaderParameters, NIAGARASHADER_API)
-		SHADER_PARAMETER(FVector3f,		ScaleMeshSizeBySpeed_ScaleFactorBias)
-		SHADER_PARAMETER(FVector3f,		ScaleMeshSizeBySpeed_ScaleFactorScale)
-		SHADER_PARAMETER(FUintVector2,	ScaleMeshSizeBySpeed_ScaleDistribution)
-		SHADER_PARAMETER(float,			ScaleMeshSizeBySpeed_VelocityNorm)
+		SHADER_PARAMETER(FVector3f,									ScaleMeshSizeBySpeed_ScaleFactorBias)
+		SHADER_PARAMETER(FVector3f,									ScaleMeshSizeBySpeed_ScaleFactorScale)
+		SHADER_PARAMETER(FNiagaraStatelessBuiltDistributionType,	ScaleMeshSizeBySpeed_ScaleDistribution)
+		SHADER_PARAMETER(float,										ScaleMeshSizeBySpeed_VelocityNorm)
 	END_SHADER_PARAMETER_STRUCT()
 		
 	BEGIN_SHADER_PARAMETER_STRUCT(FMeshIndexModule_ShaderParameters, NIAGARASHADER_API)
@@ -114,10 +114,10 @@ namespace NiagaraStateless
 	END_SHADER_PARAMETER_STRUCT()
 
 	BEGIN_SHADER_PARAMETER_STRUCT(FScaleSpriteSizeBySpeedModule_ShaderParameters, NIAGARASHADER_API)
-		SHADER_PARAMETER(FVector2f,		ScaleSpriteSizeBySpeed_ScaleFactorBias)
-		SHADER_PARAMETER(FVector2f,		ScaleSpriteSizeBySpeed_ScaleFactorScale)
-		SHADER_PARAMETER(FUintVector2,	ScaleSpriteSizeBySpeed_ScaleDistribution)
-		SHADER_PARAMETER(float,			ScaleSpriteSizeBySpeed_VelocityNorm)
+		SHADER_PARAMETER(FVector2f,									ScaleSpriteSizeBySpeed_ScaleFactorBias)
+		SHADER_PARAMETER(FVector2f,									ScaleSpriteSizeBySpeed_ScaleFactorScale)
+		SHADER_PARAMETER(FNiagaraStatelessBuiltDistributionType,	ScaleSpriteSizeBySpeed_ScaleDistribution)
+		SHADER_PARAMETER(float,										ScaleSpriteSizeBySpeed_VelocityNorm)
 	END_SHADER_PARAMETER_STRUCT()
 
 	BEGIN_SHADER_PARAMETER_STRUCT(FSpriteFacingAndAlignmentModule_ShaderParameters, NIAGARASHADER_API)
