- Add enum for notify progress, none / forward / reverse
- Fixes issues around using negative rate scaling and allows the user to still use 0-1 when playing backwards

diff --git a/Engine/Plugins/FX/Niagara/Source/NiagaraAnimNotifies/Private/AnimNotifyState_TimedNiagaraEffect.cpp b/Engine/Plugins/FX/Niagara/Source/NiagaraAnimNotifies/Private/AnimNotifyState_TimedNiagaraEffect.cpp
--- a/Engine/Plugins/FX/Niagara/Source/NiagaraAnimNotifies/Private/AnimNotifyState_TimedNiagaraEffect.cpp
+++ b/Engine/Plugins/FX/Niagara/Source/NiagaraAnimNotifies/Private/AnimNotifyState_TimedNiagaraEffect.cpp
@@ -171,6 +171,17 @@ void UAnimNotifyState_TimedNiagaraEffectAdvanced::Serialize(FArchive& Ar)
 	}
 }
 
+#if WITH_EDITORONLY_DATA
+void UAnimNotifyState_TimedNiagaraEffectAdvanced::PostLoad()
+{
+	Super::PostLoad();
+	if (!bEnableNormalizedNotifyProgress_DEPRECATED)
+	{
+		NotifyProgressType = ENiagaraAnimNotifyProgressType::None;
+	}
+}
+#endif
+
 void UAnimNotifyState_TimedNiagaraEffectAdvanced::NotifyBegin(USkeletalMeshComponent* MeshComp, class UAnimSequenceBase* Animation, float TotalDuration)
 {
 }
@@ -206,17 +217,23 @@ void UAnimNotifyState_TimedNiagaraEffectAdvanced::NotifyTick(USkeletalMeshCompon
 	//TODO: There must be some way to avoid this map and lookup. The information about the current elapsed time and a mapping onto the notify range should be available somewhere in the mesh comp and notify.
 	if (FInstanceProgressInfo* ProgressInfo = ProgressInfoMap.Find(MeshComp))
 	{
-		const float RateScale = bApplyRateScaleToProgress ? Animation->RateScale : 1.0f;
+		const float RateScale = bApplyRateScaleToProgress ? FMath::Abs(Animation->RateScale) : 1.0f;
 		ProgressInfo->Elapsed += FrameDeltaTime * RateScale;
 	}
 
 	if (UFXSystemComponent* FXComponent = GetSpawnedEffect(MeshComp))
 	{
 		//send the notify progress to the FX Component
-		if (bEnableNormalizedNotifyProgress && !NotifyProgressUserParameter.IsNone())
+		if (NotifyProgressType != ENiagaraAnimNotifyProgressType::None && !NotifyProgressUserParameter.IsNone())
 		{
-			FXComponent->SetFloatParameter(NotifyProgressUserParameter, GetNotifyProgress(MeshComp));
+			float NotifyProgress = GetNotifyProgress(MeshComp);
+			if (NotifyProgressType == ENiagaraAnimNotifyProgressType::Reverse)
+			{
+				NotifyProgress = 1.0f - NotifyProgress;
+			}
+			FXComponent->SetFloatParameter(NotifyProgressUserParameter, NotifyProgress);
 		}
+
 		//Send anim curve data to the FX Component
 		if (AnimCurves.Num() != 0)
 		{
diff --git a/Engine/Plugins/FX/Niagara/Source/NiagaraAnimNotifies/Public/AnimNotifyState_TimedNiagaraEffect.h b/Engine/Plugins/FX/Niagara/Source/NiagaraAnimNotifies/Public/AnimNotifyState_TimedNiagaraEffect.h
--- a/Engine/Plugins/FX/Niagara/Source/NiagaraAnimNotifies/Public/AnimNotifyState_TimedNiagaraEffect.h
+++ b/Engine/Plugins/FX/Niagara/Source/NiagaraAnimNotifies/Public/AnimNotifyState_TimedNiagaraEffect.h
@@ -13,6 +13,17 @@ class UFXSystemAsset;
 class UFXSystemComponent;
 class UAnimInstance;
 
+UENUM()
+enum class ENiagaraAnimNotifyProgressType : uint8
+{
+	/** Progress information will not be send to the FX component */
+	None,
+	/** Send normalized progress (0.0 to 1.0) to the FX component as a float user parameter */
+	Forward,
+	/** Send reversed normalized progress (1.0 to 0.0) to the FX component as a float user parameter */
+	Reverse,
+};
+
 // Timed Niagara Effect Notify
 // Allows a looping Niagara effect to be played in an animation that will activate
 // at the beginning of the notify and deactivate at the end.
@@ -96,6 +107,9 @@ class UAnimNotifyState_TimedNiagaraEffectAdvanced : public UAnimNotifyState_Time
 
 public:
 	NIAGARAANIMNOTIFIES_API virtual void Serialize(class FArchive& Ar) override;
+#if WITH_EDITORONLY_DATA
+	NIAGARAANIMNOTIFIES_API virtual void PostLoad() override;
+#endif
 
 	UE_DEPRECATED(5.0, "Please use the other NotifyBegin function instead")
 	NIAGARAANIMNOTIFIES_API virtual void NotifyBegin(class USkeletalMeshComponent* MeshComp, class UAnimSequenceBase* Animation, float TotalDuration) override;
@@ -114,11 +128,12 @@ public:
 	UFUNCTION(BlueprintCallable, Category = "AnimNotify")
 	NIAGARAANIMNOTIFIES_API float GetNotifyProgress(UMeshComponent* MeshComp);
 
-	UPROPERTY(EditAnywhere, Category = NotifyProgress, meta = (DisplayName = "Enable Normalized Notify Progress", ToolTip = "This send a 0-1 value of the normalized progress to the FX Component to the float User Parameter."))
-	bool bEnableNormalizedNotifyProgress = true;
+	// How/if we should send progress to the FX component
+	UPROPERTY(EditAnywhere, Category = NotifyProgress)
+	ENiagaraAnimNotifyProgressType NotifyProgressType = ENiagaraAnimNotifyProgressType::Forward;
 
 	// Should we apply the animation rate scale when calculating the elasped time.
-	UPROPERTY(EditAnywhere, BlueprintReadWrite, Category = NiagaraSystem)
+	UPROPERTY(EditAnywhere, BlueprintReadWrite, Category = NotifyProgress, meta = (DisplayName = "Notify Progress Apply Rate Scale"))
 	bool bApplyRateScaleToProgress = true;
 
 	UPROPERTY(EditAnywhere, Category = NotifyProgress, meta = (DisplayName = "User Parameter", ToolTip = "The name of your niagara user variable you would like to send the normalized notify progress to."))
@@ -127,6 +142,11 @@ public:
 	UPROPERTY(EditAnywhere, Category = AnimCurves, meta = (DisplayName = "Anim Curve Parameters", ToolTip = "Array of fnames to map Anim Curve names to Niagara Parameters."))
 	TArray<FCurveParameterPair> AnimCurves;
 
+#if WITH_EDITORONLY_DATA
+	UPROPERTY()
+	bool bEnableNormalizedNotifyProgress_DEPRECATED = true;
+#endif
+
 protected:
 
 	struct FInstanceProgressInfo
@@ -136,4 +156,3 @@ protected:
 	};
 	TMap<UMeshComponent*, FInstanceProgressInfo> ProgressInfoMap;
 };
-
