- Move log about clothing assets into an information on the data interface, and include more details

#tests Sample a SKM with clothing and observe the info

diff --git a/Engine/Plugins/FX/Niagara/Source/Niagara/Private/NiagaraDataInterfaceSkeletalMesh.cpp b/Engine/Plugins/FX/Niagara/Source/Niagara/Private/NiagaraDataInterfaceSkeletalMesh.cpp
--- a/Engine/Plugins/FX/Niagara/Source/Niagara/Private/NiagaraDataInterfaceSkeletalMesh.cpp
+++ b/Engine/Plugins/FX/Niagara/Source/Niagara/Private/NiagaraDataInterfaceSkeletalMesh.cpp
@@ -2091,11 +2091,6 @@ bool FNDISkeletalMesh_InstanceData::Init(UNiagaraDataInterfaceSkeletalMesh* Inte
 			FSkeletalMeshLODInfo* LODInfo = Mesh->GetLODInfo(CachedLODIdx);
 			bIsGpuUniformlyDistributedSampling = LODInfo->bSupportUniformlyDistributedSampling && bAllRegionsAreAreaWeighting;
 
-			if (Mesh->HasActiveClothingAssets())
-			{
-				UE_LOG(LogNiagara, Warning, TEXT("Skeletal Mesh %s has cloth asset on it: spawning from it might not work properly."), *Mesh->GetName());
-			}
-
 			const auto MaxInfluenceType = GetDefault<UNiagaraSettings>()->NDISkelMesh_GpuMaxInfluences;
 			int32 MaxInfluenceCount = -1;
 			if (MaxInfluenceType == ENDISkelMesh_GpuMaxInfluences::Type::AllowMax4)
@@ -2727,6 +2722,27 @@ void UNiagaraDataInterfaceSkeletalMesh::GetFeedback(UNiagaraSystem* Asset, UNiag
 
 		OutWarnings.Add(CPUAccessNotAllowedWarning);
 	}
+
+	if (SkelMesh && SkelMesh->HasActiveClothingAssets())
+	{
+		const FText ClothSummary = LOCTEXT("UsingClothSummary", "Active clothing assets found.  Results may be incorrect.");
+		if (GetDefault<UNiagaraSettings>()->NDISkelMesh_SupportReadingDeformedGeometry == false)
+		{
+			OutInfo.Emplace(
+				FText::Format(LOCTEXT("UsingClothNeedsProjectSetting", "Sampling from clothing regions on mesh '{0} will be incorrect.  Please enable reading from deformed geometry is disabled in project settings."), FText::FromString(SkelMesh->GetName())),
+				ClothSummary,
+				FNiagaraDataInterfaceFix()
+			);
+		}
+		else
+		{
+			OutInfo.Emplace(
+				FText::Format(LOCTEXT("UsingClothCheckValid", "Sampling from clothing regions on mesh '{0}' may be incorrect.  Ensure you are using a GPU emitter and the skin cache or deformer graph are enabled."), FText::FromString(SkelMesh->GetName())),
+				ClothSummary,
+				FNiagaraDataInterfaceFix()
+			);
+		}
+	}
 #endif
 
 	if (SoftSourceActor.Get() == nullptr && bHasNoMeshAssignedWarning)
