Add default value in NDC for SystemBoundsPadding.

diff --git a/Engine/Plugins/FX/Niagara/Source/Niagara/Private/NiagaraDataChannel_GameplayBurst.cpp b/Engine/Plugins/FX/Niagara/Source/Niagara/Private/NiagaraDataChannel_GameplayBurst.cpp
--- a/Engine/Plugins/FX/Niagara/Source/Niagara/Private/NiagaraDataChannel_GameplayBurst.cpp
+++ b/Engine/Plugins/FX/Niagara/Source/Niagara/Private/NiagaraDataChannel_GameplayBurst.cpp
@@ -41,6 +41,7 @@ void FNDCMapEntry_GameplayBurst::Init(FNDCAccessContextInst& AccessContext, UNia
 
 	UObject* SysToSpawn = TypedOwner->GetSystemToSpawn(AccessContext);
 
+	int32 BoundsPadding = TypedContext.bOverrideBoundsPadding ? TypedContext.SystemBoundsPadding : TypedOwner->DefaultSystemBoundsPadding;
 	bool bIsAttached = TypedContext.bAttachedToOwningComponent;
 	if (bIsAttached)
 	{
@@ -57,7 +58,7 @@ void FNDCMapEntry_GameplayBurst::Init(FNDCAccessContextInst& AccessContext, UNia
 		if (SysToSpawn)
 		{
 			FBoxSphereBounds Bounds = Comp->CalcBounds(Comp->GetComponentToWorld());
-			FVector Extents = Bounds.BoxExtent + TypedContext.SystemBoundsPadding;
+			FVector Extents = Bounds.BoxExtent + BoundsPadding;
 			if (UNiagaraComponent* NewComp = SpawnSystem(AccessContext.GetChecked<FNDCAccessContext_MapBase>(), SysToSpawn, Comp, Location, Extents, true))
 			{
 				//Add the attached state tag so we can be sure of how to handle FindData calls from it in future without having to re-gen attachment settings.
@@ -78,7 +79,7 @@ void FNDCMapEntry_GameplayBurst::Init(FNDCAccessContextInst& AccessContext, UNia
 			FVector CellExtents = FVector(CellSize * 0.5);
 			FInt64Vector3 CellCoord = NDCGameplayBurst_Local::GetCellCoords(Location, 1.0 / CellSize);
 			FVector CellCentre = FVector(CellCoord * CellSize) + CellExtents;
-			CellExtents += FVector(TypedContext.SystemBoundsPadding);
+			CellExtents += FVector(BoundsPadding);
 			
 			SpawnSystem(AccessContext.GetChecked<FNDCAccessContext_MapBase>(), SysToSpawn, nullptr, CellCentre, CellExtents, true);
 		}
@@ -126,6 +127,7 @@ void UNiagaraDataChannelHandler_GameplayBurst::Init(const UNiagaraDataChannel* I
 
 	//Cache various data we want to calculate once or modify at runtime.
 	DefaultCellSize = NDCTyped->GetGridCellSize();
+	DefaultSystemBoundsPadding = NDCTyped->GetSystemBoundsPadding();
 	AttachmentSettings = NDCTyped->GetAttachmentSettings();
 }
 
diff --git a/Engine/Plugins/FX/Niagara/Source/Niagara/Public/NiagaraDataChannel_GameplayBurst.h b/Engine/Plugins/FX/Niagara/Source/Niagara/Public/NiagaraDataChannel_GameplayBurst.h
--- a/Engine/Plugins/FX/Niagara/Source/Niagara/Public/NiagaraDataChannel_GameplayBurst.h
+++ b/Engine/Plugins/FX/Niagara/Source/Niagara/Public/NiagaraDataChannel_GameplayBurst.h
@@ -26,8 +26,12 @@ struct FNDCAccessContext_GameplayBurst : public FNDCAccessContext_MapBase
 	Every unique cell size value will create a new data bucket and handler Niagara System. 
 	Use with care. Overriding cell size provides more flexibility for users but can cause problems if not done correctly.
 	*/
-	UPROPERTY(EditAnywhere, BlueprintReadWrite, Category = "Input", meta = (NDCAccessContextInput))
+	UPROPERTY(EditAnywhere, BlueprintReadWrite, Category = "Input", meta = (NDCAccessContextInput, InlineEditConditionToggle))
 	uint8 bOverrideCellSize : 1 = false;
+	
+	/** If true the SystemBoundsPadding in the context will override the bounds padding defind in the NDC. */
+	UPROPERTY(EditAnywhere, BlueprintReadWrite, Category = "Input", meta = (NDCAccessContextInput, InlineEditConditionToggle))
+	uint8 bOverrideBoundsPadding : 1 = false;
 
 	/** 
 	An override for the Cell Size used for internal data routing.
@@ -38,8 +42,12 @@ struct FNDCAccessContext_GameplayBurst : public FNDCAccessContext_MapBase
 	UPROPERTY(EditAnywhere, BlueprintReadWrite, Category = "Input", meta = (NDCAccessContextInput, EditCondition="bOverrideCellSize"))
 	int32 CellSizeOverride = 2500;
 
-	/** Padding applied to the bounds of spawned handler systems. Should be large enough to contain any spawned FX near the edge of a grid cell. */
-	UPROPERTY(EditAnywhere, BlueprintReadWrite, Category = "Input", meta = (NDCAccessContextInput))
+	/** 
+	Padding applied to the bounds of spawned handler systems. 
+	Should be large enough to contain any spawned FX near the edge of a grid cell. 
+	When systems are attached to the owning component, bounds are the component bounds + SystemBoundsPadding.
+	*/
+	UPROPERTY(EditAnywhere, BlueprintReadWrite, Category = "Input", meta = (NDCAccessContextInput, EditCondition="bOverrideBoundsPadding"))
 	int32 SystemBoundsPadding = 250;
 
 	/** Data can optionally be routed by a given Gameplay Tag for additional variation. */
@@ -91,6 +99,7 @@ public:
 	[[nodiscard]] virtual TNDCAccessContextType GetAccessContextType()const override { return TNDCAccessContextType(FNDCAccessContext_GameplayBurst::StaticStruct()); }
 	[[nodiscard]] const FNDCGameplayBurstAttachmentSettings& GetAttachmentSettings()const{ return AttachmentSettings; }
 	[[nodiscard]] float GetGridCellSize()const { return CellSize; }
+	[[nodiscard]] float GetSystemBoundsPadding()const { return SystemBoundsPadding; }
 
 protected:
 	
@@ -101,6 +110,15 @@ protected:
 	UPROPERTY(EditAnywhere, Category = "Grid", meta = (NDCAccessContextInput))
 	int32 CellSize = 2500;
 
+	/** 
+	Default System Bounds Padding. Can be overridden by users via the AccessContext.
+	Niagara Systems spawned by this NDC will have their bounds padded by this amount.
+	For attached systems, the bounds will be the component bounds + padding.
+	For non attached, the bounds will be the CellSize + padding.
+	*/
+	UPROPERTY(EditAnywhere, Category = "Grid", meta = (NDCAccessContextInput))
+	int32 SystemBoundsPadding = 250;
+
 	/** 
 	Settings controlling whether incoming data should route to the main world aligned spatial grids or use a special bucket attached to and owned by a specific component.
 	Useful in cases where we want to attach Burst FX to dynamic moving objects etc for improved visuals.
@@ -148,6 +166,7 @@ protected:
 	FNDCGameplayBurstAttachmentSettings AttachmentSettings;
 
 	int32 DefaultCellSize = 2500;
+	int32 DefaultSystemBoundsPadding = 250;
 
 	/** A tag we use to mark spawned attached components to skip some work re-checking attachment settings. */
 	static const FName AttachedComponentTag;
