- Add getter method for Niagara parameters

#tests Run ET for get user parameters

#virtualized

diff --git a/Engine/Plugins/FX/Niagara/Source/Niagara/Private/NiagaraComponent.cpp b/Engine/Plugins/FX/Niagara/Source/Niagara/Private/NiagaraComponent.cpp
--- a/Engine/Plugins/FX/Niagara/Source/Niagara/Private/NiagaraComponent.cpp
+++ b/Engine/Plugins/FX/Niagara/Source/Niagara/Private/NiagaraComponent.cpp
@@ -2945,6 +2945,85 @@ void UNiagaraComponent::SetVariableTextureRenderTarget(FName InVariableName, UTe
 #endif
 }
 
+bool UNiagaraComponent::GetVariableBool(FName InVariableName, bool& bIsValid) const
+{
+	const FNiagaraVariableBase Variable(FNiagaraTypeDefinition::GetBoolDef(), InVariableName);
+	TOptional<FNiagaraBool> Value = OverrideParameters.GetParameterOptionalValue<FNiagaraBool>(Variable);
+	bIsValid = Value.IsSet();
+	return Value.Get(FNiagaraBool::False).GetValue();
+}
+
+int32 UNiagaraComponent::GetVariableInt(FName InVariableName, bool& bIsValid) const
+{
+	const FNiagaraVariableBase Variable(FNiagaraTypeDefinition::GetIntDef(), InVariableName);
+	TOptional<int32> Value = OverrideParameters.GetParameterOptionalValue<int32>(Variable);
+	bIsValid = Value.IsSet();
+	return Value.Get(0);
+}
+
+float UNiagaraComponent::GetVariableFloat(FName InVariableName, bool& bIsValid) const
+{
+	const FNiagaraVariableBase Variable(FNiagaraTypeDefinition::GetFloatDef(), InVariableName);
+	TOptional<float> Value = OverrideParameters.GetParameterOptionalValue<float>(Variable);
+	bIsValid = Value.IsSet();
+	return Value.Get(0.0f);
+}
+
+FVector2D UNiagaraComponent::GetVariableVec2(FName InVariableName, bool& bIsValid) const
+{
+	const FNiagaraVariableBase Variable(FNiagaraTypeDefinition::GetVec2Def(), InVariableName);
+	TOptional<FVector2f> Value = OverrideParameters.GetParameterOptionalValue<FVector2f>(Variable);
+	bIsValid = Value.IsSet();
+	return FVector2D(Value.Get(FVector2f::ZeroVector));
+}
+
+FVector UNiagaraComponent::GetVariableVec3(FName InVariableName, bool& bIsValid) const
+{
+	const FNiagaraVariableBase Variable(FNiagaraTypeDefinition::GetVec3Def(), InVariableName);
+	TOptional<FVector3f> Value = OverrideParameters.GetParameterOptionalValue<FVector3f>(Variable);
+	bIsValid = Value.IsSet();
+	return FVector(Value.Get(FVector3f::ZeroVector));
+}
+
+FVector4 UNiagaraComponent::GetVariableVec4(FName InVariableName, bool& bIsValid) const
+{
+	const FNiagaraVariableBase Variable(FNiagaraTypeDefinition::GetVec4Def(), InVariableName);
+	TOptional<FVector4f> Value = OverrideParameters.GetParameterOptionalValue<FVector4f>(Variable);
+	bIsValid = Value.IsSet();
+	return FVector4(Value.Get(FVector4f::Zero()));
+}
+
+FVector UNiagaraComponent::GetVariablePosition(FName InVariableName, bool& bIsValid) const
+{
+	const FVector* Position = OverrideParameters.GetPositionParameterValue(InVariableName);
+	bIsValid = Position != nullptr;
+	return Position ? *Position : FVector::ZeroVector;
+}
+
+FLinearColor UNiagaraComponent::GetVariableColor(FName InVariableName, bool& bIsValid) const
+{
+	const FNiagaraVariableBase Variable(FNiagaraTypeDefinition::GetColorDef(), InVariableName);
+	TOptional<FLinearColor> Value = OverrideParameters.GetParameterOptionalValue<FLinearColor>(Variable);
+	bIsValid = Value.IsSet();
+	return Value.Get(FLinearColor::White);
+}
+
+FQuat UNiagaraComponent::GetVariableQuat(FName InVariableName, bool& bIsValid) const
+{
+	const FNiagaraVariableBase Variable(FNiagaraTypeDefinition::GetQuatDef(), InVariableName);
+	TOptional<FQuat4f> Value = OverrideParameters.GetParameterOptionalValue<FQuat4f>(Variable);
+	bIsValid = Value.IsSet();
+	return FQuat(Value.Get(FQuat4f::Identity));
+}
+
+FMatrix UNiagaraComponent::GetVariableMatrix(FName InVariableName, bool& bIsValid) const
+{
+	const FNiagaraVariableBase Variable(FNiagaraTypeDefinition::GetMatrix4Def(), InVariableName);
+	TOptional<FMatrix44f> Value = OverrideParameters.GetParameterOptionalValue<FMatrix44f>(Variable);
+	bIsValid = Value.IsSet();
+	return FMatrix(Value.Get(FMatrix44f::Identity));
+}
+
 void FixInvalidUserParameters(FNiagaraUserRedirectionParameterStore& ParameterStore)
 {
 	static const FString UserPrefix = FNiagaraConstants::UserNamespace.ToString() + TEXT(".");
diff --git a/Engine/Plugins/FX/Niagara/Source/Niagara/Public/NiagaraComponent.h b/Engine/Plugins/FX/Niagara/Source/Niagara/Public/NiagaraComponent.h
--- a/Engine/Plugins/FX/Niagara/Source/Niagara/Public/NiagaraComponent.h
+++ b/Engine/Plugins/FX/Niagara/Source/Niagara/Public/NiagaraComponent.h
@@ -557,6 +557,46 @@ public:
 	UFUNCTION(BlueprintCallable, Category = Niagara, meta = (DisplayName = "Set Niagara Variable (TextureRenderTarget)"))
 	NIAGARA_API void SetVariableTextureRenderTarget(FName InVariableName, class UTextureRenderTarget* TextureRenderTarget);
 
+	/** Gets a Niagara bool parameter by name. */
+	UFUNCTION(BlueprintCallable, Category = Niagara, meta = (DisplayName = "Get Niagara Variable (Bool)", Keywords = "user parameter variable bool"))
+	[[nodiscard]] NIAGARA_API bool GetVariableBool(FName InVariableName, bool& bIsValid) const;
+
+	/** Gets a Niagara int parameter by name. */
+	UFUNCTION(BlueprintCallable, Category = Niagara, meta = (DisplayName = "Get Niagara Variable (Int32)", Keywords = "user parameter variable int"))
+	[[nodiscard]] NIAGARA_API int32 GetVariableInt(FName InVariableName, bool& bIsValid) const;
+
+	/** Gets a Niagara float parameter by name. */
+	UFUNCTION(BlueprintCallable, Category = Niagara, meta = (DisplayName = "Get Niagara Variable (Float)", Keywords = "user parameter variable float"))
+	[[nodiscard]] NIAGARA_API float GetVariableFloat(FName InVariableName, bool& bIsValid) const;
+
+	/** Gets a Niagara Vector2 parameter by name. */
+	UFUNCTION(BlueprintCallable, Category = Niagara, meta = (DisplayName = "Get Niagara Variable (Vector2)", Keywords = "user parameter variable vector"))
+	[[nodiscard]] NIAGARA_API FVector2D GetVariableVec2(FName InVariableName, bool& bIsValid) const;
+
+	/** Gets a Niagara Vector3 parameter by name. */
+	UFUNCTION(BlueprintCallable, Category = Niagara, meta = (DisplayName = "Get Niagara Variable (Vector3)", Keywords = "user parameter variable vector"))
+	[[nodiscard]] NIAGARA_API FVector GetVariableVec3(FName InVariableName, bool& bIsValid) const;
+
+	/** Gets a Niagara Vector4 parameter by name. */
+	UFUNCTION(BlueprintCallable, Category = Niagara, meta = (DisplayName = "Get Niagara Variable (Vector4)", Keywords = "user parameter variable vector"))
+	[[nodiscard]] NIAGARA_API FVector4 GetVariableVec4(FName InVariableName, bool& bIsValid) const;
+
+	/** Gets a Niagara Position parameter by name. */
+	UFUNCTION(BlueprintCallable, Category = Niagara, meta = (DisplayName = "Get Niagara Variable (Position)", Keywords = "user parameter variable vector position lwc"))
+	[[nodiscard]] NIAGARA_API FVector GetVariablePosition(FName InVariableName, bool& bIsValid) const;
+
+	/** Gets a Niagara Linear Color parameter by name. */
+	UFUNCTION(BlueprintCallable, Category = Niagara, meta = (DisplayName = "Get Niagara Variable (LinearColor)", Keywords = "user parameter variable color"))
+	[[nodiscard]] NIAGARA_API FLinearColor GetVariableColor(FName InVariableName, bool& bIsValid) const;
+
+	/** Gets a Niagara quaternion parameter by name. */
+	UFUNCTION(BlueprintCallable, Category = Niagara, meta = (DisplayName = "Get Niagara Variable (Quaternion)", Keywords = "user parameter variable quaternion rotation"))
+	[[nodiscard]] NIAGARA_API FQuat GetVariableQuat(FName InVariableName, bool& bIsValid) const;
+
+	/** Gets a Niagara matrix parameter by name. */
+	UFUNCTION(BlueprintCallable, Category = Niagara, meta = (DisplayName = "Get Niagara Variable (Matrix)", Keywords = "user parameter variable matrix"))
+	[[nodiscard]] NIAGARA_API FMatrix GetVariableMatrix(FName InVariableName, bool& bIsValid) const;
+
 	/** Resets the System to it's initial pre-simulated state. */
 	UFUNCTION(BlueprintCallable, Category = Niagara, meta = (DisplayName = "Reset System"))
 	NIAGARA_API void ResetSystem();
