- Fix for missing compressed quat data accessor

#rnx

diff --git a/Engine/Plugins/FX/Niagara/Source/Niagara/Classes/NiagaraDataSetAccessor.h b/Engine/Plugins/FX/Niagara/Source/Niagara/Classes/NiagaraDataSetAccessor.h
--- a/Engine/Plugins/FX/Niagara/Source/Niagara/Classes/NiagaraDataSetAccessor.h
+++ b/Engine/Plugins/FX/Niagara/Source/Niagara/Classes/NiagaraDataSetAccessor.h
@@ -341,11 +341,11 @@ struct FNiagaraDataSetAccessorTypeInfo<FQuat4f>
 {
 	using TAccessorBaseClass = FNiagaraDataSetAccessorFloat<FQuat4f>;
 
-	static constexpr bool bSupportsHalf = false;
+	static constexpr bool bSupportsHalf = true;
 	static constexpr bool bSupportsAlternateType = false;
 	static constexpr int32 NumElements = 4;
 	static const FNiagaraTypeDefinition& GetFloatType() { return FNiagaraTypeDefinition::GetQuatDef(); }
-	static const FNiagaraTypeDefinition& GetHalfType() { check(false); return FNiagaraTypeDefinition::GetQuatDef(); }
+	static const FNiagaraTypeDefinition& GetHalfType() { return FNiagaraTypeDefinition::GetHalfVec4Def(); }
 	static const FNiagaraTypeDefinition& GetAlternateType() { return FNiagaraTypeDefinition::GetQuatDef(); }
 };
 
