- Allow other distribution types to opt into indexing / address / interpolation

diff --git a/Engine/Plugins/FX/Niagara/Source/Niagara/Internal/Stateless/Modules/NiagaraStatelessModule_ScaleMeshSizeBySpeed.h b/Engine/Plugins/FX/Niagara/Source/Niagara/Internal/Stateless/Modules/NiagaraStatelessModule_ScaleMeshSizeBySpeed.h
--- a/Engine/Plugins/FX/Niagara/Source/Niagara/Internal/Stateless/Modules/NiagaraStatelessModule_ScaleMeshSizeBySpeed.h
+++ b/Engine/Plugins/FX/Niagara/Source/Niagara/Internal/Stateless/Modules/NiagaraStatelessModule_ScaleMeshSizeBySpeed.h
@@ -31,7 +31,7 @@ public:
 	UPROPERTY(EditAnywhere, Category = "Parameters")
 	bool bSampleScaleFactorByCurve = false;
 
-	UPROPERTY(EditAnywhere, Category = "Parameters", meta = (DisableRangeDistribution, DisableBindingDistribution, EditCondition = "bSampleScaleFactorByCurve", EditConditionHides))
+	UPROPERTY(EditAnywhere, Category = "Parameters", meta = (DisableRangeDistribution, DisableDistributionLookupValueMode, EditCondition = "bSampleScaleFactorByCurve", EditConditionHides))
 	FNiagaraDistributionFloat SampleFactorCurve = FNiagaraDistributionFloat({0.0f, 1.0f});
 
 	virtual void BuildEmitterData(const FNiagaraStatelessEmitterDataBuildContext& BuildContext) const override;
diff --git a/Engine/Plugins/FX/Niagara/Source/Niagara/Internal/Stateless/Modules/NiagaraStatelessModule_ScaleSpriteSizeBySpeed.h b/Engine/Plugins/FX/Niagara/Source/Niagara/Internal/Stateless/Modules/NiagaraStatelessModule_ScaleSpriteSizeBySpeed.h
--- a/Engine/Plugins/FX/Niagara/Source/Niagara/Internal/Stateless/Modules/NiagaraStatelessModule_ScaleSpriteSizeBySpeed.h
+++ b/Engine/Plugins/FX/Niagara/Source/Niagara/Internal/Stateless/Modules/NiagaraStatelessModule_ScaleSpriteSizeBySpeed.h
@@ -31,7 +31,7 @@ public:
 	UPROPERTY(EditAnywhere, Category = "Parameters")
 	bool bSampleScaleFactorByCurve = false;
 
-	UPROPERTY(EditAnywhere, Category = "Parameters", meta = (DisableRangeDistribution, DisableBindingDistribution, EditCondition = "bSampleScaleFactorByCurve", EditConditionHides))
+	UPROPERTY(EditAnywhere, Category = "Parameters", meta = (DisableRangeDistribution, DisableDistributionLookupValueMode, EditCondition = "bSampleScaleFactorByCurve", EditConditionHides))
 	FNiagaraDistributionFloat SampleFactorCurve = FNiagaraDistributionFloat({0.0f, 1.0f});
 
 	virtual void BuildEmitterData(const FNiagaraStatelessEmitterDataBuildContext& BuildContext) const override;
diff --git a/Engine/Plugins/FX/Niagara/Source/Niagara/Internal/Stateless/NiagaraStatelessEmitterDataBuildContext.h b/Engine/Plugins/FX/Niagara/Source/Niagara/Internal/Stateless/NiagaraStatelessEmitterDataBuildContext.h
--- a/Engine/Plugins/FX/Niagara/Source/Niagara/Internal/Stateless/NiagaraStatelessEmitterDataBuildContext.h
+++ b/Engine/Plugins/FX/Niagara/Source/Niagara/Internal/Stateless/NiagaraStatelessEmitterDataBuildContext.h
@@ -78,7 +78,7 @@ private:
 
 private:
 	template<typename TType>
-	FNiagaraStatelessBuiltDistributionType InternalAddDistribution(ENiagaraDistributionMode Mode, TConstArrayView<TType> Values, const FVector2f& TimeRange, TType DefaultValue) const
+	FNiagaraStatelessBuiltDistributionType InternalAddDistribution(ENiagaraDistributionMode Mode, TConstArrayView<TType> Values, const FVector2f& TimeRange, const uint8 LookupValueMode, TType DefaultValue) const
 	{
 		using namespace NiagaraStateless;
 		
@@ -126,8 +126,6 @@ private:
 				break;
 		}
 
-		const uint8 LookupValueMode = 0;
-
 		FNiagaraStatelessBuiltDistribution::SetUniformRandom(BuiltDistribution, bUniformRandom);
 		FNiagaraStatelessBuiltDistribution::SetLookupValueMode(BuiltDistribution, bIsRandom ? uint8(ENiagaraDistributionLookupValueMode::Random) : LookupValueMode);
 
@@ -185,7 +183,7 @@ public:
 		{
 			return InternalAddDistribution(AddExpression(Distribution.ParameterExpression), 1, TDistribution::GetDefaultValue(), false);
 		}
-		return InternalAddDistribution(Distribution.Mode, MakeArrayView(Distribution.Values), Distribution.ValuesTimeRange, TDistribution::GetDefaultValue());
+		return InternalAddDistribution(Distribution.Mode, MakeArrayView(Distribution.Values), Distribution.ValuesTimeRange, Distribution.GetLookupValueMode(), TDistribution::GetDefaultValue());
 	}
 
 	// Adds a distribution and forces it to generate as a curve for lookup
@@ -195,12 +193,12 @@ public:
 		FNiagaraStatelessBuiltDistributionType BuiltDistribution = FNiagaraStatelessBuiltDistribution::GetDefault();
 		if (ensure((Distribution.IsCurve() || Distribution.IsGradient()) && Distribution.Values.Num() > 1))
 		{
-			return InternalAddDistribution<TType>(Distribution.Mode, MakeArrayView(Distribution.Values), Distribution.ValuesTimeRange, DefaultValue);
+			return InternalAddDistribution<TType>(Distribution.Mode, MakeArrayView(Distribution.Values), Distribution.ValuesTimeRange, Distribution.GetLookupValueMode(), DefaultValue);
 		}
 		else
 		{
 			const TType DefaultValues[] = {DefaultValue , DefaultValue};
-			return InternalAddDistribution(ENiagaraDistributionMode::NonUniformCurve, MakeArrayView(DefaultValues), Distribution.ValuesTimeRange, DefaultValue);
+			return InternalAddDistribution(ENiagaraDistributionMode::NonUniformCurve, MakeArrayView(DefaultValues), Distribution.ValuesTimeRange, Distribution.GetLookupValueMode(), DefaultValue);
 		}
 	}
 
diff --git a/Engine/Plugins/FX/Niagara/Source/Niagara/Private/Stateless/NiagaraStatelessDistribution.cpp b/Engine/Plugins/FX/Niagara/Source/Niagara/Private/Stateless/NiagaraStatelessDistribution.cpp
--- a/Engine/Plugins/FX/Niagara/Source/Niagara/Private/Stateless/NiagaraStatelessDistribution.cpp
+++ b/Engine/Plugins/FX/Niagara/Source/Niagara/Private/Stateless/NiagaraStatelessDistribution.cpp
@@ -683,6 +683,16 @@ FNiagaraDistributionCurveFloat::FNiagaraDistributionCurveFloat(ENiagaraDistribut
 #endif
 }
 
+FNiagaraDistributionCurveFloat::FNiagaraDistributionCurveFloat(std::initializer_list<float> CurvePoints)
+{
+	Mode = ENiagaraDistributionMode::UniformCurve;
+	Values = CurvePoints;
+
+#if WITH_EDITORONLY_DATA
+	NiagaraStatelessDistributionPrivate::CreateCurveFromValues(Values, ChannelCurves, ChannelConstantsAndRanges);
+#endif
+}
+
 FNiagaraDistributionCurveVector3::FNiagaraDistributionCurveVector3()
 	: FNiagaraDistributionCurveVector3(ENiagaraDistributionCurveLUTMode::Sample)
 {
@@ -731,7 +741,146 @@ void FNiagaraDistributionBase::ForEachParameterBinding(TFunction<void(const FNia
 	}
 }
 
-const UEnum* FNiagaraDistributionBase::GetLookupValueModeEnum(FProperty* Property)
+bool FNiagaraDistributionBase::AllowDistributionMode(ENiagaraDistributionMode InMode, const FProperty* Property) const
+{
+	static const FName DisableArrayDistributionName("DisableArrayDistribution");
+	static const FName DisableBindingDistributionName("DisableBindingDistribution");
+	static const FName DisableCurveDistributionName("DisableCurveDistribution");
+	static const FName DisableUniformDistributionName("DisableUniformDistribution");
+	static const FName DisableNonUniformDistributionName("DisableNonUniformDistribution");
+	static const FName DisableRangeDistributionName("DisableRangeDistribution");
+	static const FName DisableColorGradientDistributionName("DisableColorGradientDistribution");
+
+	bool bAllowed = true;
+	switch (InMode)
+	{
+		case ENiagaraDistributionMode::Array:
+			bAllowed &= !Property || !Property->HasMetaData(DisableArrayDistributionName);
+			break;
+
+		case ENiagaraDistributionMode::Binding:
+		case ENiagaraDistributionMode::Expression:
+			bAllowed &= !Property || !Property->HasMetaData(DisableBindingDistributionName);
+			bAllowed &= GetBindingTypeDef().IsValid();
+			break;
+
+		case ENiagaraDistributionMode::UniformConstant:
+			bAllowed &= !Property || !Property->HasMetaData(DisableUniformDistributionName);
+			break;
+
+		case ENiagaraDistributionMode::NonUniformConstant:
+			bAllowed &= !Property || !Property->HasMetaData(DisableNonUniformDistributionName);
+			break;
+
+		case ENiagaraDistributionMode::UniformRange:
+			bAllowed &= !Property || (!Property->HasMetaData(DisableRangeDistributionName) && !Property->HasMetaData(DisableUniformDistributionName));
+			break;
+
+		case ENiagaraDistributionMode::NonUniformRange:
+			bAllowed &= !Property || (!Property->HasMetaData(DisableRangeDistributionName) && !Property->HasMetaData(DisableNonUniformDistributionName));
+			break;
+
+		case ENiagaraDistributionMode::UniformCurve:
+			bAllowed &= !Property || (!Property->HasMetaData(DisableCurveDistributionName) && !Property->HasMetaData(DisableUniformDistributionName));
+			break;
+
+		case ENiagaraDistributionMode::NonUniformCurve:
+			bAllowed &= !Property || (!Property->HasMetaData(DisableCurveDistributionName) && !Property->HasMetaData(DisableNonUniformDistributionName));
+			break;
+
+		case ENiagaraDistributionMode::ColorGradient:
+			bAllowed &= !Property || !Property->HasMetaData(DisableColorGradientDistributionName);
+			break;
+
+		default:
+			break;
+	}
+	return bAllowed;
+}
+
+bool FNiagaraDistributionBase::InternalAllowRangeDistribution(ENiagaraDistributionMode InMode, const FProperty* Property) const
+{
+	switch (InMode)
+	{
+		case ENiagaraDistributionMode::Binding:
+		case ENiagaraDistributionMode::Expression:
+		case ENiagaraDistributionMode::UniformConstant:
+		case ENiagaraDistributionMode::NonUniformConstant:
+		case ENiagaraDistributionMode::UniformRange:
+		case ENiagaraDistributionMode::NonUniformRange:
+			return FNiagaraDistributionBase::AllowDistributionMode(InMode, Property);
+
+		default:
+			return false;
+	}
+}
+
+bool FNiagaraDistributionBase::InternalAllowCurveDistribution(ENiagaraDistributionMode InMode, const FProperty* Property) const
+{
+	switch (InMode)
+	{
+		case ENiagaraDistributionMode::UniformCurve:
+		case ENiagaraDistributionMode::NonUniformCurve:
+		case ENiagaraDistributionMode::ColorGradient:
+			return FNiagaraDistributionBase::AllowDistributionMode(InMode, Property);
+
+		default:
+			return false;
+	}
+}
+
+bool FNiagaraDistributionBase::AllowAddressMode(const FProperty* Property) const
+{
+	static const FName DisableDistributionAddressModeName("DisableDistributionAddressMode");
+
+	if (Mode != ENiagaraDistributionMode::Array)
+	{
+		return false;
+	}
+
+	if (Property && Property->HasMetaData(DisableDistributionAddressModeName))
+	{
+		return false;
+	}
+
+	return true;
+}
+
+bool FNiagaraDistributionBase::AllowInterpolationMode(const FProperty* Property) const
+{
+	static const FName DisableDistributionInterpolationModeName("DisableDistributionInterpolationMode");
+
+	if (Mode != ENiagaraDistributionMode::Array)
+	{
+		return false;
+	}
+
+	if (Property && Property->HasMetaData(DisableDistributionInterpolationModeName))
+	{
+		return false;
+	}
+
+	return true;
+}
+
+bool FNiagaraDistributionBase::AllowLookupValueMode(const FProperty* Property) const
+{
+	static const FName DisableDistributionLookupValueModeName("DisableDistributionLookupValueMode");
+
+	if (Mode != ENiagaraDistributionMode::Array && Mode != ENiagaraDistributionMode::ColorGradient && Mode != ENiagaraDistributionMode::NonUniformCurve && Mode != ENiagaraDistributionMode::UniformCurve)
+	{
+		return false;
+	}
+
+	if (Property && Property->HasMetaData(DisableDistributionLookupValueModeName))
+	{
+		return false;
+	}
+
+	return true;
+}
+
+const UEnum* FNiagaraDistributionBase::GetLookupValueModeEnum(const FProperty* Property)
 {
 	const UEnum* LookupEnum = nullptr;
 	if (ensure(Property))
@@ -777,28 +926,38 @@ void FNiagaraDistributionBase::PostEditChangeProperty(UObject* OwnerObject, FPro
 			{
 				ENiagaraDistributionTimeRangeMode TimeRangeMode = ENiagaraDistributionTimeRangeMode::Normalized;
 
-				if (const UEnum* LookupEnum = FNiagaraDistributionBase::GetLookupValueModeEnum(ValueProperty))
+				if (ValuePtr->AllowLookupValueMode(ValueProperty))
 				{
-					const int32 ValueIndex = LookupEnum->GetIndexByValue(ValuePtr->GetLookupValueMode());
-					const FString TimeRangeModeStr = ValueIndex != INDEX_NONE ? LookupEnum->GetMetaData(TEXT("TimeRangeMode"), ValueIndex) : FString();
-					if (!TimeRangeModeStr.IsEmpty())
+					if (const UEnum* LookupEnum = FNiagaraDistributionBase::GetLookupValueModeEnum(ValueProperty))
 					{
-						const int64 TimeRangeModeValue = StaticEnum<ENiagaraDistributionTimeRangeMode>()->GetValueByNameString(TimeRangeModeStr);
-						if (ensureMsgf(TimeRangeModeValue != INDEX_NONE, TEXT("Time Range metadata(%s) on Enum(%s) is invalid."), *TimeRangeModeStr, *LookupEnum->GetName()))
+						const int32 ValueIndex = LookupEnum->GetIndexByValue(ValuePtr->GetLookupValueMode());
+						const FString TimeRangeModeStr = ValueIndex != INDEX_NONE ? LookupEnum->GetMetaData(TEXT("TimeRangeMode"), ValueIndex) : FString();
+						if (!TimeRangeModeStr.IsEmpty())
 						{
-							TimeRangeMode = ENiagaraDistributionTimeRangeMode(TimeRangeModeValue);
+							const int64 TimeRangeModeValue = StaticEnum<ENiagaraDistributionTimeRangeMode>()->GetValueByNameString(TimeRangeModeStr);
+							if (ensureMsgf(TimeRangeModeValue != INDEX_NONE, TEXT("Time Range metadata(%s) on Enum(%s) is invalid."), *TimeRangeModeStr, *LookupEnum->GetName()))
+							{
+								TimeRangeMode = ENiagaraDistributionTimeRangeMode(TimeRangeModeValue);
+							}
 						}
 					}
 				}
+
 				ValuePtr->UpdateTimeRangeFromDistribution(TimeRangeMode);
 			}
-			else
+			
+			// Reset any values that are not supported in this mode
+			if (!ValuePtr->AllowLookupValueMode(ValueProperty))
+			{
+				ValuePtr->SetLookupValueMode(uint8(ENiagaraDistributionLookupValueMode::ParticlesNormalizedAge));
+			}
+			if (!ValuePtr->AllowInterpolationMode(ValueProperty))
 			{
-				// Set back to defaults if the mode is not an array
-				// If/when we support some of this with curves, i.e. address mode / lookup value mode, we can remove this but we don't want value to leak from one mode to the next
 				ValuePtr->SetInterpolationMode(ENiagaraDistributionInterpolationMode::Linear);
+			}
+			if (!ValuePtr->AllowAddressMode(ValueProperty))
+			{
 				ValuePtr->SetAddressMode(ENiagaraDistributionAddressMode::Clamp);
-				ValuePtr->SetLookupValueMode(uint8(ENiagaraDistributionLookupValueMode::ParticlesNormalizedAge));
 			}
 		}
 	}
diff --git a/Engine/Plugins/FX/Niagara/Source/Niagara/Public/Stateless/NiagaraStatelessDistribution.h b/Engine/Plugins/FX/Niagara/Source/Niagara/Public/Stateless/NiagaraStatelessDistribution.h
--- a/Engine/Plugins/FX/Niagara/Source/Niagara/Public/Stateless/NiagaraStatelessDistribution.h
+++ b/Engine/Plugins/FX/Niagara/Source/Niagara/Public/Stateless/NiagaraStatelessDistribution.h
@@ -127,10 +127,16 @@ struct FNiagaraDistributionBase
 
 	void ForEachParameterBinding(TFunction<void(const FNiagaraVariableBase&)> Delegate) const;
 
-	virtual bool AllowArray() const { return true; }
-	virtual bool AllowBinding() const { return true; }
-	virtual bool AllowConstant() const { return true; }
-	virtual bool AllowCurves() const { return true; }
+	NIAGARA_API virtual bool AllowDistributionMode(ENiagaraDistributionMode Mode, const FProperty* Property) const;
+protected:
+	NIAGARA_API virtual bool InternalAllowRangeDistribution(ENiagaraDistributionMode Mode, const FProperty* Property) const;
+	NIAGARA_API virtual bool InternalAllowCurveDistribution(ENiagaraDistributionMode Mode, const FProperty* Property) const;
+public:
+
+	NIAGARA_API bool AllowAddressMode(const FProperty* Property) const;
+	NIAGARA_API bool AllowInterpolationMode(const FProperty* Property) const;
+	NIAGARA_API bool AllowLookupValueMode(const FProperty* Property) const;
+
 	virtual bool DisplayAsColor() const { return false; }
 	virtual int32 GetBaseNumberOfChannels() const { return 0; }
 	virtual void UpdateValuesFromDistribution() { }
@@ -140,7 +146,7 @@ struct FNiagaraDistributionBase
 
 	static void PostEditChangeProperty(UObject* OwnerObject, FPropertyChangedEvent& PropertyChangedEvent);
 
-	NIAGARA_API static const UEnum* GetLookupValueModeEnum(FProperty* Property);
+	NIAGARA_API static const UEnum* GetLookupValueModeEnum(const FProperty* Property);
 #endif
 };
 
@@ -208,8 +214,7 @@ struct FNiagaraDistributionRangeFloat : public FNiagaraDistributionBase
 	NIAGARA_API FNiagaraStatelessRangeFloat CalculateRange(const float Default = GetDefaultValue()) const;
 
 #if WITH_EDITORONLY_DATA
-	virtual bool AllowArray() const override { return false; }
-	virtual bool AllowCurves() const override { return false; }
+	virtual bool AllowDistributionMode(ENiagaraDistributionMode InMode, const FProperty* Property) const override { return InternalAllowRangeDistribution(InMode, Property); }
 	virtual int32 GetBaseNumberOfChannels() const override { return 1; }
 	NIAGARA_API virtual void UpdateValuesFromDistribution() override;
 	virtual FNiagaraTypeDefinition GetBindingTypeDef() const { return FNiagaraTypeDefinition::GetFloatDef(); }
@@ -237,8 +242,7 @@ struct FNiagaraDistributionRangeVector2 : public FNiagaraDistributionBase
 	NIAGARA_API FNiagaraStatelessRangeVector2 CalculateRange(const FVector2f& Default = GetDefaultValue()) const;
 
 #if WITH_EDITORONLY_DATA
-	virtual bool AllowArray() const override { return false; }
-	virtual bool AllowCurves() const override { return false; }
+	virtual bool AllowDistributionMode(ENiagaraDistributionMode InMode, const FProperty* Property) const override { return InternalAllowRangeDistribution(InMode, Property); }
 	virtual int32 GetBaseNumberOfChannels() const override { return 2; }
 	virtual void UpdateValuesFromDistribution() override;
 	virtual FNiagaraTypeDefinition GetBindingTypeDef() const { return FNiagaraTypeDefinition::GetVec2Def(); }
@@ -266,8 +270,7 @@ struct FNiagaraDistributionRangeVector3 : public FNiagaraDistributionBase
 	NIAGARA_API FNiagaraStatelessRangeVector3 CalculateRange(const FVector3f& Default = GetDefaultValue()) const;
 
 #if WITH_EDITORONLY_DATA
-	virtual bool AllowArray() const override { return false; }
-	virtual bool AllowCurves() const override { return false; }
+	virtual bool AllowDistributionMode(ENiagaraDistributionMode InMode, const FProperty* Property) const override { return InternalAllowRangeDistribution(InMode, Property); }
 	virtual int32 GetBaseNumberOfChannels() const override { return 3; }
 	virtual void UpdateValuesFromDistribution() override;
 	virtual FNiagaraTypeDefinition GetBindingTypeDef() const { return FNiagaraTypeDefinition::GetVec3Def(); }
@@ -295,8 +298,7 @@ struct FNiagaraDistributionRangeColor : public FNiagaraDistributionBase
 	NIAGARA_API FNiagaraStatelessRangeColor CalculateRange(const FLinearColor& Default = GetDefaultValue()) const;
 
 #if WITH_EDITORONLY_DATA
-	virtual bool AllowArray() const override { return false; }
-	virtual bool AllowCurves() const override { return false; }
+	virtual bool AllowDistributionMode(ENiagaraDistributionMode InMode, const FProperty* Property) const override { return InternalAllowRangeDistribution(InMode, Property); }
 	virtual bool DisplayAsColor() const { return true; }
 	virtual int32 GetBaseNumberOfChannels() const override { return 4; }
 	virtual void UpdateValuesFromDistribution() override;
@@ -322,8 +324,7 @@ struct FNiagaraDistributionRangeRotator : public FNiagaraDistributionBase
 	NIAGARA_API FNiagaraStatelessRangeRotator CalculateRange(const FRotator3f& Default = FRotator3f::ZeroRotator) const;
 
 #if WITH_EDITORONLY_DATA
-	virtual bool AllowArray() const override { return false; }
-	virtual bool AllowCurves() const override { return false; }
+	virtual bool AllowDistributionMode(ENiagaraDistributionMode InMode, const FProperty* Property) const override { return InternalAllowRangeDistribution(InMode, Property); }
 	virtual int32 GetBaseNumberOfChannels() const override { return 3; }
 	virtual void UpdateValuesFromDistribution() override;
 	virtual FNiagaraTypeDefinition GetBindingTypeDef() const { return FNiagaraTypeDefinition::GetQuatDef(); }
@@ -478,6 +479,7 @@ struct FNiagaraDistributionCurveFloat : public FNiagaraDistributionBase
 
 	NIAGARA_API FNiagaraDistributionCurveFloat();
 	NIAGARA_API explicit FNiagaraDistributionCurveFloat(ENiagaraDistributionCurveLUTMode InLUTMode);
+	NIAGARA_API explicit FNiagaraDistributionCurveFloat(std::initializer_list<float> CurvePoints);
 
 	UPROPERTY(EditAnywhere, Category = "Parameters")
 	TArray<float> Values;
@@ -491,9 +493,7 @@ struct FNiagaraDistributionCurveFloat : public FNiagaraDistributionBase
 		return (FNiagaraDistributionBase)*this == (FNiagaraDistributionBase)Other && Values == Other.Values;
 	}
 
-	virtual bool AllowArray() const override { return false; }
-	virtual bool AllowBinding() const override { return false; }
-	virtual bool AllowConstant() const override { return false; }
+	virtual bool AllowDistributionMode(ENiagaraDistributionMode InMode, const FProperty* Property) const override { return InternalAllowCurveDistribution(InMode, Property); }
 	virtual int32 GetBaseNumberOfChannels() const override { return 1; }
 	NIAGARA_API virtual void UpdateValuesFromDistribution() override;
 
@@ -522,9 +522,7 @@ struct FNiagaraDistributionCurveVector3 : public FNiagaraDistributionBase
 		return (FNiagaraDistributionBase)*this == (FNiagaraDistributionBase)Other && Values == Other.Values;
 	}
 
-	virtual bool AllowArray() const override { return false; }
-	virtual bool AllowBinding() const override { return false; }
-	virtual bool AllowConstant() const override { return false; }
+	virtual bool AllowDistributionMode(ENiagaraDistributionMode InMode, const FProperty* Property) const override { return InternalAllowCurveDistribution(InMode, Property); }
 	virtual int32 GetBaseNumberOfChannels() const override { return 3; }
 	NIAGARA_API virtual void UpdateValuesFromDistribution() override;
 
diff --git a/Engine/Plugins/FX/Niagara/Source/NiagaraEditor/Private/NiagaraEditorModule.cpp b/Engine/Plugins/FX/Niagara/Source/NiagaraEditor/Private/NiagaraEditorModule.cpp
--- a/Engine/Plugins/FX/Niagara/Source/NiagaraEditor/Private/NiagaraEditorModule.cpp
+++ b/Engine/Plugins/FX/Niagara/Source/NiagaraEditor/Private/NiagaraEditorModule.cpp
@@ -1437,6 +1437,8 @@ void FNiagaraEditorModule::StartupModule()
 	RegisterPropertyUtilities(FNiagaraDistributionRangeVector2::StaticStruct(), MakeShared<FNiagaraDistributionPropertyEditorUtilities, ESPMode::ThreadSafe>());
 	RegisterPropertyUtilities(FNiagaraDistributionRangeVector3::StaticStruct(), MakeShared<FNiagaraDistributionPropertyEditorUtilities, ESPMode::ThreadSafe>());
 	RegisterPropertyUtilities(FNiagaraDistributionRangeColor::StaticStruct(), MakeShared<FNiagaraDistributionPropertyEditorUtilities, ESPMode::ThreadSafe>());
+	RegisterPropertyUtilities(FNiagaraDistributionCurveFloat::StaticStruct(), MakeShared<FNiagaraDistributionPropertyEditorUtilities, ESPMode::ThreadSafe>());
+	RegisterPropertyUtilities(FNiagaraDistributionCurveVector3::StaticStruct(), MakeShared<FNiagaraDistributionPropertyEditorUtilities, ESPMode::ThreadSafe>());
 
 	FEdGraphUtilities::RegisterVisualPinFactory(GraphPanelPinFactory);
 
diff --git a/Engine/Plugins/FX/Niagara/Source/NiagaraEditor/Private/Stateless/NiagaraDistributionPropertyCustomization.cpp b/Engine/Plugins/FX/Niagara/Source/NiagaraEditor/Private/Stateless/NiagaraDistributionPropertyCustomization.cpp
--- a/Engine/Plugins/FX/Niagara/Source/NiagaraEditor/Private/Stateless/NiagaraDistributionPropertyCustomization.cpp
+++ b/Engine/Plugins/FX/Niagara/Source/NiagaraEditor/Private/Stateless/NiagaraDistributionPropertyCustomization.cpp
@@ -44,18 +44,6 @@ public:
 		SourceNumChannels = InNumChannels;
 		if (IsValid())
 		{
-			const FName DisableArrayDistributionName("DisableArrayDistribution");
-			const FName DisableBindingDistributionName("DisableBindingDistribution");
-			const FName DisableCurveDistributionName("DisableCurveDistribution");
-			const FName DisableUniformDistributionName("DisableUniformDistribution");
-			const FName DisableNonUniformDistributionName("DisableNonUniformDistribution");
-			const FName DisableRangeDistributionName("DisableRangeDistribution");
-			bAllowArray = InDistribution->AllowArray() && InDistribution->GetBindingTypeDef().IsValid() && (PropertyHandle ? PropertyHandle->HasMetaData(DisableArrayDistributionName) == false : false);
-			bAllowBinding = InDistribution->AllowBinding() && InDistribution->GetBindingTypeDef().IsValid() && (PropertyHandle ? PropertyHandle->HasMetaData(DisableBindingDistributionName) == false : false);
-			bAllowUniform = InDistribution->AllowConstant() && PropertyHandle->HasMetaData(DisableUniformDistributionName) == false;
-			bAllowNonUniform = InDistribution->AllowConstant() && PropertyHandle->HasMetaData(DisableNonUniformDistributionName) == false;
-			bAllowRange = InDistribution->AllowConstant() && PropertyHandle->HasMetaData(DisableRangeDistributionName) == false;
-			bAllowCurves = InDistribution->AllowCurves() && PropertyHandle->HasMetaData(DisableCurveDistributionName) == false;
 			bDisplayAsColor = InDistribution->DisplayAsColor();
 			PropertyHandle->SetOnPropertyValueChanged(FSimpleDelegate::CreateSP(this->AsShared(), &FNiagaraDistributionAdapter::DistributionPropertyChanged));
 			EditorMode = GetDistributionEditorModeFromSourceMode(SourceNumChannels, bDisplayAsColor, SourceDistribution->Mode);
@@ -154,11 +142,18 @@ public:
 
 	virtual void GetSupportedDistributionModes(TArray<ENiagaraDistributionEditorMode>& OutSupportedModes) const override
 	{
-		if (bAllowArray)
+		const FProperty* Property = PropertyHandle ? PropertyHandle->GetProperty() : nullptr;
+		if (SourceDistribution == nullptr)
+		{
+			return;
+		}
+
+		if (SourceDistribution->AllowDistributionMode(ENiagaraDistributionMode::Array, Property) )
 		{
 			OutSupportedModes.Add(ENiagaraDistributionEditorMode::Array);
 		}
-		if (bAllowBinding)
+
+		if (SourceDistribution->AllowDistributionMode(ENiagaraDistributionMode::Binding, Property))
 		{
 			OutSupportedModes.Add(ENiagaraDistributionEditorMode::Binding);
 
@@ -168,53 +163,57 @@ public:
 				OutSupportedModes.Add(ENiagaraDistributionEditorMode::Expression);
 			}
 		}
+
 		if (bDisplayAsColor)
 		{
-			OutSupportedModes.Add(ENiagaraDistributionEditorMode::ColorConstant);
-			OutSupportedModes.Add(ENiagaraDistributionEditorMode::ColorRange);
-			if (bAllowCurves)
+			if (SourceDistribution->AllowDistributionMode(ENiagaraDistributionMode::NonUniformConstant, Property))
+			{
+				OutSupportedModes.Add(ENiagaraDistributionEditorMode::ColorConstant);
+			}
+			if (SourceDistribution->AllowDistributionMode(ENiagaraDistributionMode::NonUniformRange, Property))
+			{
+				OutSupportedModes.Add(ENiagaraDistributionEditorMode::ColorRange);
+			}
+			if (SourceDistribution->AllowDistributionMode(ENiagaraDistributionMode::ColorGradient, Property))
 			{
-				OutSupportedModes.Add(ENiagaraDistributionEditorMode::NonUniformCurve);
 				OutSupportedModes.Add(ENiagaraDistributionEditorMode::ColorGradient);
 			}
+			if (SourceDistribution->AllowDistributionMode(ENiagaraDistributionMode::NonUniformCurve, Property))
+			{
+				OutSupportedModes.Add(ENiagaraDistributionEditorMode::NonUniformCurve);
+			}
 		}
-		else if (SourceNumChannels == 1)
+		else
 		{
-			if (bAllowUniform)
+			const bool bSingleChannel = SourceNumChannels == 1;
+
+			if (SourceDistribution->AllowDistributionMode(ENiagaraDistributionMode::UniformConstant, Property))
 			{
-				OutSupportedModes.Add(ENiagaraDistributionEditorMode::Constant);
-				if (bAllowRange)
-				{
-					OutSupportedModes.Add(ENiagaraDistributionEditorMode::Range);
-				}
+				OutSupportedModes.Add(bSingleChannel ? ENiagaraDistributionEditorMode::Constant : ENiagaraDistributionEditorMode::UniformConstant);
 			}
-			if (bAllowCurves)
+			if (SourceDistribution->AllowDistributionMode(ENiagaraDistributionMode::UniformRange, Property))
 			{
-				OutSupportedModes.Add(ENiagaraDistributionEditorMode::Curve);
+				OutSupportedModes.Add(bSingleChannel ? ENiagaraDistributionEditorMode::Range : ENiagaraDistributionEditorMode::UniformRange);
 			}
-		}
-		else if (SourceNumChannels > 1)
-		{
-			if (bAllowUniform)
+			if (SourceDistribution->AllowDistributionMode(ENiagaraDistributionMode::UniformCurve, Property))
 			{
-				OutSupportedModes.Add(ENiagaraDistributionEditorMode::UniformConstant);
-				if (bAllowRange)
-				{
-					OutSupportedModes.Add(ENiagaraDistributionEditorMode::UniformRange);
-				}
+				OutSupportedModes.Add(bSingleChannel ? ENiagaraDistributionEditorMode::Curve : ENiagaraDistributionEditorMode::UniformCurve);
 			}
-			if (bAllowNonUniform)
+
+			if (!bSingleChannel)
 			{
-				OutSupportedModes.Add(ENiagaraDistributionEditorMode::NonUniformConstant);
-				if (bAllowRange)
+				if (SourceDistribution->AllowDistributionMode(ENiagaraDistributionMode::NonUniformConstant, Property))
+				{
+					OutSupportedModes.Add(ENiagaraDistributionEditorMode::NonUniformConstant);
+				}
+				if (SourceDistribution->AllowDistributionMode(ENiagaraDistributionMode::NonUniformRange, Property))
 				{
 					OutSupportedModes.Add(ENiagaraDistributionEditorMode::NonUniformRange);
 				}
-			}
-			if (bAllowCurves)
-			{
-				OutSupportedModes.Add(ENiagaraDistributionEditorMode::UniformCurve);
-				OutSupportedModes.Add(ENiagaraDistributionEditorMode::NonUniformCurve);
+				if (SourceDistribution->AllowDistributionMode(ENiagaraDistributionMode::NonUniformCurve, Property))
+				{
+					OutSupportedModes.Add(ENiagaraDistributionEditorMode::NonUniformCurve);
+				}
 			}
 		}
 	}
@@ -673,6 +672,11 @@ public:
 		}
 	}
 
+	virtual bool AllowInterpolationMode() const override
+	{
+		return SourceDistribution->AllowInterpolationMode(PropertyHandle ? PropertyHandle->GetProperty() : nullptr);
+	}
+
 	virtual ENiagaraDistributionInterpolationMode GetInterpolationMode() const override
 	{
 		return SourceDistribution->GetInterpolationMode();
@@ -689,6 +693,11 @@ public:
 		);
 	}
 
+	virtual bool AllowAddressMode() const override
+	{
+		return SourceDistribution->AllowAddressMode(PropertyHandle ? PropertyHandle->GetProperty() : nullptr);
+	}
+
 	virtual ENiagaraDistributionAddressMode GetAddressMode() const override
 	{
 		return SourceDistribution->GetAddressMode();
@@ -705,6 +714,11 @@ public:
 		);
 	}
 
+	virtual bool AllowLookupValueMode() const override
+	{
+		return SourceDistribution->AllowLookupValueMode(PropertyHandle ? PropertyHandle->GetProperty() : nullptr);
+	}
+
 	virtual uint8 GetLookupValueMode() const override
 	{
 		return SourceDistribution->GetLookupValueMode();
@@ -1012,14 +1026,7 @@ private:
 	TOptional<int32> ContinuousTransactionIndex;
 	bool bContinuousChangeActive = false;
 
-	bool bAllowArray = true;
-	bool bAllowBinding = true;
-	bool bAllowUniform = true;
-	bool bAllowNonUniform = true;
-	bool bAllowRange = true;
-	bool bAllowCurves = true;
 	bool bDisplayAsColor = false;
-
 	FSimpleMulticastDelegate OnDistributionEditorModeChangedDelegate;
 	bool bUpdatingDistributionModeInternally = false;
 };
diff --git a/Engine/Plugins/FX/Niagara/Source/NiagaraEditor/Private/TypeEditorUtilities/NiagaraDistributionPropertyEditorUtilities.cpp b/Engine/Plugins/FX/Niagara/Source/NiagaraEditor/Private/TypeEditorUtilities/NiagaraDistributionPropertyEditorUtilities.cpp
--- a/Engine/Plugins/FX/Niagara/Source/NiagaraEditor/Private/TypeEditorUtilities/NiagaraDistributionPropertyEditorUtilities.cpp
+++ b/Engine/Plugins/FX/Niagara/Source/NiagaraEditor/Private/TypeEditorUtilities/NiagaraDistributionPropertyEditorUtilities.cpp
@@ -145,16 +145,22 @@ bool FNiagaraDistributionPropertyEditorUtilities::TryUpdatePropertyFromClipboard
 				}
 			}
 
-			if (ParsedDistribution.IsSet() == false && Distribution->AllowCurves())
+			if (ParsedDistribution.IsSet() == false)
 			{
 				UScriptStruct* CurveCollectionStruct = FNiagaraClipboardCurveCollection::StaticStruct();
 				FNiagaraClipboardCurveCollection CurveCollectionValue;
 				if (InSourceClipboardPortableValue.TryUpdateStructValue(*CurveCollectionStruct, (uint8*)&CurveCollectionValue))
 				{
-					if (CurveCollectionValue.Curves.Num() == 1 || CurveCollectionValue.Curves.Num() == Distribution->GetBaseNumberOfChannels())
+					if (CurveCollectionValue.Curves.Num() == 1 && Distribution->AllowDistributionMode(ENiagaraDistributionMode::UniformCurve, InPropertyHandle.GetProperty()))
 					{
 						ParsedDistribution = FNiagaraDistributionBase();
-						ParsedDistribution->Mode = CurveCollectionValue.Curves.Num() == 1 ? ENiagaraDistributionMode::UniformCurve : ENiagaraDistributionMode::NonUniformCurve;
+						ParsedDistribution->Mode = ENiagaraDistributionMode::UniformCurve;
+						ParsedDistribution->ChannelCurves = CurveCollectionValue.Curves;
+					}
+					else if (CurveCollectionValue.Curves.Num() == Distribution->GetBaseNumberOfChannels() && Distribution->AllowDistributionMode(ENiagaraDistributionMode::NonUniformCurve, InPropertyHandle.GetProperty()))
+					{
+						ParsedDistribution = FNiagaraDistributionBase();
+						ParsedDistribution->Mode = Distribution->DisplayAsColor() ? ENiagaraDistributionMode::ColorGradient : ENiagaraDistributionMode::NonUniformCurve;
 						ParsedDistribution->ChannelCurves = CurveCollectionValue.Curves;
 					}
 				}
diff --git a/Engine/Plugins/FX/Niagara/Source/NiagaraEditor/Private/Widgets/INiagaraDistributionAdapter.h b/Engine/Plugins/FX/Niagara/Source/NiagaraEditor/Private/Widgets/INiagaraDistributionAdapter.h
--- a/Engine/Plugins/FX/Niagara/Source/NiagaraEditor/Private/Widgets/INiagaraDistributionAdapter.h
+++ b/Engine/Plugins/FX/Niagara/Source/NiagaraEditor/Private/Widgets/INiagaraDistributionAdapter.h
@@ -134,16 +134,22 @@ public:
 	/** Set a curve value for the specified channel */
 	virtual void SetCurveValue(int32 ChannelIndex, const FRichCurve& InValue) = 0;
 
+	/** Are we allowed to read / write the interpolation mode for the current distribution mode. */
+	virtual bool AllowInterpolationMode() const = 0;
 	/** Get the interpolation mode from the distribution */
 	virtual ENiagaraDistributionInterpolationMode GetInterpolationMode() const = 0;
 	/** Set the interpolation mode on the distribution */
 	virtual void SetInterpolationMode(ENiagaraDistributionInterpolationMode Mode) = 0;
 
+	/** Are we allowed to read / write the address mode for the current distribution mode. */
+	virtual bool AllowAddressMode() const = 0;
 	/** Get the address mode from the distribution */
 	virtual ENiagaraDistributionAddressMode GetAddressMode() const = 0;
 	/** Set the address mode on the distribution */
 	virtual void SetAddressMode(ENiagaraDistributionAddressMode Mode) = 0;
 
+	/** Are we allowed to read / write the value lookup mode for the current distribution mode. */
+	virtual bool AllowLookupValueMode() const = 0;
 	/** Get the lookup value mode from the distribution */
 	virtual uint8 GetLookupValueMode() const = 0;
 	/** Set the lookup value mode on the distribution */
diff --git a/Engine/Plugins/FX/Niagara/Source/NiagaraEditor/Private/Widgets/SNiagaraDistributionArrayEditor.cpp b/Engine/Plugins/FX/Niagara/Source/NiagaraEditor/Private/Widgets/SNiagaraDistributionArrayEditor.cpp
--- a/Engine/Plugins/FX/Niagara/Source/NiagaraEditor/Private/Widgets/SNiagaraDistributionArrayEditor.cpp
+++ b/Engine/Plugins/FX/Niagara/Source/NiagaraEditor/Private/Widgets/SNiagaraDistributionArrayEditor.cpp
@@ -3,6 +3,7 @@
 #include "Widgets/SNiagaraDistributionArrayEditor.h"
 
 #include "Widgets/INiagaraDistributionAdapter.h"
+#include "Widgets/SNiagaraDistributionLookupValueWidget.h"
 #include "Widgets/SNiagaraParameterName.h"
 #include "SNiagaraParameterEditor.h"
 #include "NiagaraVariableMetaData.h"
@@ -473,9 +474,6 @@ void SNiagaraDistributionArrayEditor::Construct(const FArguments& InArgs, TShare
 {
 	DistributionAdapter = InDistributionAdapter;
 
-	TSharedPtr<IPropertyHandle> PropertyHandle = DistributionAdapter->GetPropertyHandle();
-	LookupValueModeEnum = FNiagaraDistributionBase::GetLookupValueModeEnum(PropertyHandle ? PropertyHandle->GetProperty() : nullptr);
-
 	RefreshRows();
 
 	SAssignNew(ArrayListViewWidget, SListView<TSharedPtr<int32>>)
@@ -514,81 +512,93 @@ void SNiagaraDistributionArrayEditor::Construct(const FArguments& InArgs, TShare
 		ToolBarBuilder.EndSection();
 	}
 
-	ChildSlot
-	[
-		SNew(SVerticalBox)
-		+ SVerticalBox::Slot()
-		. AutoHeight()
+	TSharedRef<SGridPanel> GridPanel = SNew(SGridPanel);
+	int32 GridPanelY = 0;
+
+	if (DistributionAdapter->AllowLookupValueMode())
+	{
+		GridPanel->AddSlot(0, GridPanelY)
+		.HAlign(HAlign_Left)
+		.VAlign(VAlign_Center)
 		[
-			SNew(SGridPanel)
-			+ SGridPanel::Slot(0, 0)
-			.HAlign(HAlign_Left)
-			.VAlign(VAlign_Center)
-			[
-				SNew(STextBlock)
-				.Font(IDetailLayoutBuilder::GetDetailFont())
-				.Text(LOCTEXT("LookupValueMode", "Lookup Value: "))
-			]
-			+ SGridPanel::Slot(1, 0)
-			.HAlign(HAlign_Left)
-			.VAlign(VAlign_Center)
-			[
-				SNew(SComboButton)
-				.OnGetMenuContent(this, &SNiagaraDistributionArrayEditor::OnGetLookupValueModeOptions)
-				.ContentPadding(1)
-				.ButtonContent()
-				[
-					SNew(SNiagaraParameterName)
-					.ParameterName(this, &SNiagaraDistributionArrayEditor::GetLookupValueModeName)
-					.IsReadOnly(true)
-				]
-			]
+			SNew(STextBlock)
+			.Font(IDetailLayoutBuilder::GetDetailFont())
+			.Text(LOCTEXT("LookupValueMode", "Lookup Value: "))
+		];
+
+		GridPanel->AddSlot(1, GridPanelY)
+		.HAlign(HAlign_Left)
+		.VAlign(VAlign_Center)
+		[
+			SNew(SNiagaraDistributionLookupValueWidget, DistributionAdapter)
+			.ShowValueOnly(true)
+		];
+		++GridPanelY;
+	}
 
-			+ SGridPanel::Slot(0, 1)
-			.HAlign(HAlign_Left)
-			.VAlign(VAlign_Center)
+	if (DistributionAdapter->AllowInterpolationMode())
+	{
+		GridPanel->AddSlot(0, GridPanelY)
+		.HAlign(HAlign_Left)
+		.VAlign(VAlign_Center)
+		[
+			SNew(STextBlock)
+			.Font(IDetailLayoutBuilder::GetDetailFont())
+			.Text(LOCTEXT("InterpolationMode", "Interpolation: "))
+		];
+
+		GridPanel->AddSlot(1, GridPanelY)
+		.HAlign(HAlign_Left)
+		.VAlign(VAlign_Center)
+		[
+			SNew(SComboButton)
+			.OnGetMenuContent(this, &SNiagaraDistributionArrayEditor::OnGetInterpolationModeOptions)
+			.ContentPadding(1)
+			//-TODO:.ToolTip(this, &SNiagaraDistributionArrayEditor::GetInterpolationModeToolTip)
+			.ButtonContent()
 			[
 				SNew(STextBlock)
-				.Font(IDetailLayoutBuilder::GetDetailFont())
-				.Text(LOCTEXT("InterpolationMode", "Interpolation: "))
-			]
-			+ SGridPanel::Slot(1, 1)
-			.HAlign(HAlign_Left)
-			.VAlign(VAlign_Center)
-			[
-				SNew(SComboButton)
-				.OnGetMenuContent(this, &SNiagaraDistributionArrayEditor::OnGetInterpolationModeOptions)
-				.ContentPadding(1)
-				//-TODO:.ToolTip(this, &SNiagaraDistributionArrayEditor::GetInterpolationModeToolTip)
-				.ButtonContent()
-				[
-					SNew(STextBlock)
-					.Text(this, &SNiagaraDistributionArrayEditor::GetInterpolationModeText)
-				]
+				.Text(this, &SNiagaraDistributionArrayEditor::GetInterpolationModeText)
 			]
+		];
+		++GridPanelY;
+	}
 
-			+ SGridPanel::Slot(0, 2)
-			.HAlign(HAlign_Left)
-			.VAlign(VAlign_Center)
+	if (DistributionAdapter->AllowAddressMode())
+	{
+		GridPanel->AddSlot(0, GridPanelY)
+		.HAlign(HAlign_Left)
+		.VAlign(VAlign_Center)
+		[
+			SNew(STextBlock)
+			.Font(IDetailLayoutBuilder::GetDetailFont())
+			.Text(LOCTEXT("AddessMode", "Address Mode: "))
+		];
+
+		GridPanel->AddSlot(1, GridPanelY)
+		.HAlign(HAlign_Left)
+		.VAlign(VAlign_Center)
+		[
+			SNew(SComboButton)
+			.OnGetMenuContent(this, &SNiagaraDistributionArrayEditor::OnGetAddressModeOptions)
+			.ContentPadding(1)
+			//-TODO:.ToolTip(this, &SNiagaraDistributionArrayEditor::GetAddressModeToolTip)
+			.ButtonContent()
 			[
 				SNew(STextBlock)
-				.Font(IDetailLayoutBuilder::GetDetailFont())
-				.Text(LOCTEXT("AddessMode", "Address Mode: "))
-			]
-			+ SGridPanel::Slot(1, 2)
-			.HAlign(HAlign_Left)
-			.VAlign(VAlign_Center)
-			[
-				SNew(SComboButton)
-				.OnGetMenuContent(this, &SNiagaraDistributionArrayEditor::OnGetAddressModeOptions)
-				.ContentPadding(1)
-				//-TODO:.ToolTip(this, &SNiagaraDistributionArrayEditor::GetAddressModeToolTip)
-				.ButtonContent()
-				[
-					SNew(STextBlock)
-					.Text(this, &SNiagaraDistributionArrayEditor::GetAddressModeText)
-				]
+				.Text(this, &SNiagaraDistributionArrayEditor::GetAddressModeText)
 			]
+		];
+		++GridPanelY;
+	}
+
+	ChildSlot
+	[
+		SNew(SVerticalBox)
+		+ SVerticalBox::Slot()
+		. AutoHeight()
+		[
+			GridPanel
 		]
 		//////////////////////////////////////////////////////
 		+ SVerticalBox::Slot()
@@ -654,24 +664,6 @@ void SNiagaraDistributionArrayEditor::RemoveAllArrayElements()
 	RefreshRows();
 }
 
-void SNiagaraDistributionArrayEditor::SetLookupValueMode(uint8 Mode)
-{
-	DistributionAdapter->SetLookupValueMode(Mode);
-}
-
-FName SNiagaraDistributionArrayEditor::GetLookupValueModeName() const
-{
-	const uint8 Mode = DistributionAdapter->GetLookupValueMode();
-	return FName(LookupValueModeEnum->GetDisplayNameTextByValue(Mode).ToString());
-}
-
-FText SNiagaraDistributionArrayEditor::GetLookupValueModeToolTip() const
-{
-	const uint8 Mode = DistributionAdapter->GetLookupValueMode();
-	const int64 Index = LookupValueModeEnum->GetIndexByValue(Mode);
-	return Index != INDEX_NONE ? LookupValueModeEnum->GetToolTipTextByIndex(Index) : FText();
-}
-
 void SNiagaraDistributionArrayEditor::SetInterpolationMode(ENiagaraDistributionInterpolationMode Mode)
 {
 	DistributionAdapter->SetInterpolationMode(Mode);
@@ -732,17 +724,6 @@ TSharedRef<ITableRow> SNiagaraDistributionArrayEditor::MakeArrayEntryWidget(cons
 		;
 }
 
-TSharedRef<SWidget> SNiagaraDistributionArrayEditor::OnGetLookupValueModeOptions()
-{
-	return SNiagaraDistributionArrayEditorPrivate::GetEnumMenuOptions(
-		LookupValueModeEnum,
-		[&](int64 iValue)
-		{
-			return FExecuteAction::CreateSP(this, &SNiagaraDistributionArrayEditor::SetLookupValueMode, uint8(iValue));
-		}
-	);
-}
-
 TSharedRef<SWidget> SNiagaraDistributionArrayEditor::OnGetInterpolationModeOptions()
 {
 	return SNiagaraDistributionArrayEditorPrivate::GetEnumMenuOptions(
diff --git a/Engine/Plugins/FX/Niagara/Source/NiagaraEditor/Private/Widgets/SNiagaraDistributionArrayEditor.h b/Engine/Plugins/FX/Niagara/Source/NiagaraEditor/Private/Widgets/SNiagaraDistributionArrayEditor.h
--- a/Engine/Plugins/FX/Niagara/Source/NiagaraEditor/Private/Widgets/SNiagaraDistributionArrayEditor.h
+++ b/Engine/Plugins/FX/Niagara/Source/NiagaraEditor/Private/Widgets/SNiagaraDistributionArrayEditor.h
@@ -50,14 +50,9 @@ private:
 
 	TSharedRef<ITableRow> MakeArrayEntryWidget(const TSharedPtr<int32> RowIndexPtr, const TSharedRef<STableViewBase>& OwnerTable);
 
-	TSharedRef<SWidget> OnGetLookupValueModeOptions();
 	TSharedRef<SWidget> OnGetInterpolationModeOptions();
 	TSharedRef<SWidget> OnGetAddressModeOptions();
 
-	void SetLookupValueMode(uint8 Mode);
-	FName GetLookupValueModeName() const;
-	FText GetLookupValueModeToolTip() const;
-
 	void SetInterpolationMode(ENiagaraDistributionInterpolationMode Mode);
 	FText GetInterpolationModeText() const;
 	FText GetInterpolationModeToolTip() const;
@@ -72,6 +67,4 @@ private:
 
 	TSharedPtr<SListView<TSharedPtr<int32>>>	ArrayListViewWidget;
 	TArray<TSharedPtr<int32>>					ArrayEntryItems;
-
-	const UEnum*								LookupValueModeEnum = nullptr;
 };
diff --git a/Engine/Plugins/FX/Niagara/Source/NiagaraEditor/Private/Widgets/SNiagaraDistributionEditor.cpp b/Engine/Plugins/FX/Niagara/Source/NiagaraEditor/Private/Widgets/SNiagaraDistributionEditor.cpp
--- a/Engine/Plugins/FX/Niagara/Source/NiagaraEditor/Private/Widgets/SNiagaraDistributionEditor.cpp
+++ b/Engine/Plugins/FX/Niagara/Source/NiagaraEditor/Private/Widgets/SNiagaraDistributionEditor.cpp
@@ -18,6 +18,7 @@
 #include "Widgets/SNiagaraColorEditor.h"
 #include "Widgets/SNiagaraDistributionArrayEditor.h"
 #include "Widgets/SNiagaraDistributionCurveEditor.h"
+#include "Widgets/SNiagaraDistributionLookupValueWidget.h"
 #include "Widgets/SNiagaraParameterName.h"
 
 #include "Stateless/SNiagaraStatelessExpressionWidget.h"
@@ -677,7 +678,17 @@ public:
 			SNew(SBox)
 			.Padding(2)
 			[
-				GradientEditor
+				SNew(SVerticalBox)
+				+ SVerticalBox::Slot()
+				.AutoHeight()
+				[
+					GradientEditor
+				]
+				+ SVerticalBox::Slot()
+				.AutoHeight()
+				[
+					SNew(SNiagaraDistributionLookupValueWidget, DistributionAdapter)
+				]
 			]
 		];
 	}
@@ -750,6 +761,13 @@ TSharedRef<SWidget> SNiagaraDistributionEditor::ConstructContentForMode()
 				.CurveColor(DistributionAdapter->GetChannelColor(ChannelIndex))
 			];
 		}
+
+		CurveBox->AddSlot()
+		.AutoHeight()
+		[
+			SNew(SNiagaraDistributionLookupValueWidget, DistributionAdapter)
+		];
+
 		return CurveBox;
 	}
 	else if (FNiagaraDistributionEditorUtilities::IsGradient(Mode))
diff --git a/Engine/Plugins/FX/Niagara/Source/NiagaraEditor/Private/Widgets/SNiagaraDistributionLookupValueWidget.cpp b/Engine/Plugins/FX/Niagara/Source/NiagaraEditor/Private/Widgets/SNiagaraDistributionLookupValueWidget.cpp
new file mode 100644
--- /dev/null
+++ b/Engine/Plugins/FX/Niagara/Source/NiagaraEditor/Private/Widgets/SNiagaraDistributionLookupValueWidget.cpp
@@ -0,0 +1,108 @@
+// Copyright Epic Games, Inc. All Rights Reserved.
+
+#include "Widgets/SNiagaraDistributionLookupValueWidget.h"
+
+#include "Widgets/INiagaraDistributionAdapter.h"
+#include "Widgets/SNiagaraParameterName.h"
+
+#include "Framework/MultiBox/MultiBoxBuilder.h"
+#include "Widgets/Input/SComboButton.h"
+#include "PropertyHandle.h"
+
+#define LOCTEXT_NAMESPACE "SNiagaraDistributionLookupValueWidget"
+
+void SNiagaraDistributionLookupValueWidget::Construct(const FArguments& InArgs, TSharedPtr<INiagaraDistributionAdapter> InDistributionAdapter)
+{
+	DistributionAdapter = InDistributionAdapter;
+	if (DistributionAdapter == nullptr || !DistributionAdapter->AllowLookupValueMode())
+	{
+		return;
+	}
+
+	TSharedPtr<IPropertyHandle> DistributionPropertyHandle = DistributionAdapter->GetPropertyHandle();
+	FProperty* DistributionProperty = DistributionPropertyHandle ? DistributionPropertyHandle->GetProperty() : nullptr;
+	if (DistributionProperty == nullptr)
+	{
+		return;
+	}
+
+	LookupValueModeEnum = FNiagaraDistributionBase::GetLookupValueModeEnum(DistributionProperty);
+
+	TSharedRef<SHorizontalBox> BoxWidget = SNew(SHorizontalBox);
+	if (InArgs._ShowValueOnly == false)
+	{
+		BoxWidget->AddSlot()
+		.AutoWidth()
+		.VAlign(VAlign_Center)
+		[
+			SNew(STextBlock)
+			.Font(FAppStyle::Get().GetFontStyle("PropertyWindow.NormalFont"))
+			.Text(LOCTEXT("LookupValueMode", "Lookup Value"))
+		];
+	}
+
+	BoxWidget->AddSlot()
+	.VAlign(VAlign_Center)
+	[
+		SNew(SComboButton)
+		.OnGetMenuContent(this, &SNiagaraDistributionLookupValueWidget::OnGetLookupValueModeOptions)
+		.ContentPadding(1)
+		.ButtonContent()
+		[
+			SNew(SNiagaraParameterName)
+			.ParameterName(this, &SNiagaraDistributionLookupValueWidget::GetLookupValueModeName)
+			.IsReadOnly(true)
+		]
+	];
+
+	ChildSlot
+	[
+		BoxWidget
+	];
+}
+
+TSharedRef<SWidget> SNiagaraDistributionLookupValueWidget::OnGetLookupValueModeOptions()
+{
+	FMenuBuilder MenuBuilder(true, nullptr);
+	MenuBuilder.BeginSection("Actions");
+	for (int64 iValue = 0; iValue < LookupValueModeEnum->GetMaxEnumValue(); ++iValue)
+	{
+		if (LookupValueModeEnum->IsValidEnumValue(iValue) == false)
+		{
+			continue;
+		}
+
+		const int64 iIndex = LookupValueModeEnum->GetIndexByValue(iValue);
+
+		MenuBuilder.AddMenuEntry(
+			LookupValueModeEnum->GetDisplayNameTextByIndex(iIndex),
+			LookupValueModeEnum->GetToolTipTextByIndex(iIndex),
+			FSlateIcon(),
+			FUIAction(
+				FExecuteAction::CreateSP(this, &SNiagaraDistributionLookupValueWidget::SetLookupValueMode, uint8(iValue))
+			)
+		);
+	}
+	MenuBuilder.EndSection();
+	return MenuBuilder.MakeWidget();
+}
+
+FName SNiagaraDistributionLookupValueWidget::GetLookupValueModeName() const
+{
+	const uint8 Mode = DistributionAdapter->GetLookupValueMode();
+	return FName(LookupValueModeEnum->GetDisplayNameTextByValue(Mode).ToString());
+}
+
+FText SNiagaraDistributionLookupValueWidget::GetLookupValueModeToolTip() const
+{
+	const uint8 Mode = DistributionAdapter->GetLookupValueMode();
+	const int64 Index = LookupValueModeEnum->GetIndexByValue(Mode);
+	return Index != INDEX_NONE ? LookupValueModeEnum->GetToolTipTextByIndex(Index) : FText();
+}
+
+void SNiagaraDistributionLookupValueWidget::SetLookupValueMode(uint8 Mode)
+{
+	DistributionAdapter->SetLookupValueMode(Mode);
+}
+
+#undef LOCTEXT_NAMESPACE
diff --git a/Engine/Plugins/FX/Niagara/Source/NiagaraEditor/Private/Widgets/SNiagaraDistributionLookupValueWidget.h b/Engine/Plugins/FX/Niagara/Source/NiagaraEditor/Private/Widgets/SNiagaraDistributionLookupValueWidget.h
new file mode 100644
--- /dev/null
+++ b/Engine/Plugins/FX/Niagara/Source/NiagaraEditor/Private/Widgets/SNiagaraDistributionLookupValueWidget.h
@@ -0,0 +1,33 @@
+// Copyright Epic Games, Inc. All Rights Reserved.
+
+#pragma once
+
+#include "Stateless/NiagaraStatelessDistribution.h"
+
+#include "Misc/NotifyHook.h"
+#include "Widgets/Views/SListView.h"
+#include "Widgets/SCompoundWidget.h"
+
+class INiagaraDistributionAdapter;
+
+class SNiagaraDistributionLookupValueWidget : public SCompoundWidget
+{
+public:
+	SLATE_BEGIN_ARGS(SNiagaraDistributionLookupValueWidget) { }
+		SLATE_ARGUMENT_DEFAULT(bool, ShowValueOnly) = false;
+	SLATE_END_ARGS();
+
+	void Construct(const FArguments& InArgs, TSharedPtr<INiagaraDistributionAdapter> InDistributionAdapter);
+
+private:
+	TSharedRef<SWidget> OnGetLookupValueModeOptions();
+
+	FName GetLookupValueModeName() const;
+	FText GetLookupValueModeToolTip() const;
+
+	void SetLookupValueMode(uint8 Mode);
+
+protected:
+	TSharedPtr<INiagaraDistributionAdapter>	DistributionAdapter;
+	const UEnum*							LookupValueModeEnum = nullptr;
+};
