- Disable lookup on curves inside initialize particle

#rnx

diff --git a/Engine/Plugins/FX/Niagara/Source/Niagara/Internal/Stateless/Modules/NiagaraStatelessModule_InitializeParticle.h b/Engine/Plugins/FX/Niagara/Source/Niagara/Internal/Stateless/Modules/NiagaraStatelessModule_InitializeParticle.h
--- a/Engine/Plugins/FX/Niagara/Source/Niagara/Internal/Stateless/Modules/NiagaraStatelessModule_InitializeParticle.h
+++ b/Engine/Plugins/FX/Niagara/Source/Niagara/Internal/Stateless/Modules/NiagaraStatelessModule_InitializeParticle.h
@@ -20,7 +20,7 @@ public:
 	UPROPERTY(EditAnywhere, Category = "Parameters", meta = (DisplayName = "Lifetime", Units="s"))
 	FNiagaraDistributionRangeFloat LifetimeDistribution = FNiagaraDistributionRangeFloat(FNiagaraStatelessGlobals::GetDefaultLifetimeValue());
 
-	UPROPERTY(EditAnywhere, Category = "Parameters", meta = (DisplayName = "Color", DisableCurveDistribution))
+	UPROPERTY(EditAnywhere, Category = "Parameters", meta = (DisplayName = "Color", DisableCurveDistribution, DisableDistributionLookupValueMode))
 	FNiagaraDistributionRangeColor ColorDistribution = FNiagaraDistributionRangeColor(FNiagaraStatelessGlobals::GetDefaultColorValue());
 
 	UPROPERTY(EditAnywhere, Category = "Parameters", meta = (DisplayName = "Mass"))
@@ -41,7 +41,7 @@ public:
 	UPROPERTY(EditAnywhere, Category = "Parameters", meta = (DisplayName = "Ribbon Width", EditCondition="bWriteRibbonWidth"))
 	FNiagaraDistributionRangeFloat RibbonWidthDistribution = FNiagaraDistributionRangeFloat(FNiagaraStatelessGlobals::GetDefaultRibbonWidthValue());
 
-	UPROPERTY(EditAnywhere, Category = "Parameters", meta = (DisableCurveDistribution, DisableArrayDistribution))
+	UPROPERTY(EditAnywhere, Category = "Parameters", meta = (DisableCurveDistribution, DisableArrayDistribution, DisableDistributionLookupValueMode))
 	FNiagaraDistributionPosition InitialPositionDistribution = FNiagaraDistributionPosition(FVector3f::ZeroVector);
 
 	virtual void BuildEmitterData(const FNiagaraStatelessEmitterDataBuildContext& BuildContext) const override;
