- Handling a couple of edge cases

#rnx

diff --git a/Engine/Plugins/FX/Niagara/Source/NiagaraEditor/Private/TypeEditorUtilities/NiagaraDistributionPropertyEditorUtilities.cpp b/Engine/Plugins/FX/Niagara/Source/NiagaraEditor/Private/TypeEditorUtilities/NiagaraDistributionPropertyEditorUtilities.cpp
--- a/Engine/Plugins/FX/Niagara/Source/NiagaraEditor/Private/TypeEditorUtilities/NiagaraDistributionPropertyEditorUtilities.cpp
+++ b/Engine/Plugins/FX/Niagara/Source/NiagaraEditor/Private/TypeEditorUtilities/NiagaraDistributionPropertyEditorUtilities.cpp
@@ -160,7 +160,11 @@ bool FNiagaraDistributionPropertyEditorUtilities::TryUpdatePropertyFromClipboard
 					else if (CurveCollectionValue.Curves.Num() == Distribution->GetBaseNumberOfChannels() && Distribution->AllowDistributionMode(ENiagaraDistributionMode::NonUniformCurve, InPropertyHandle.GetProperty()))
 					{
 						ParsedDistribution = FNiagaraDistributionBase();
-						ParsedDistribution->Mode = Distribution->DisplayAsColor() ? ENiagaraDistributionMode::ColorGradient : ENiagaraDistributionMode::NonUniformCurve;
+						ParsedDistribution->Mode = ENiagaraDistributionMode::NonUniformCurve;
+						if (Distribution->DisplayAsColor() && Distribution->AllowDistributionMode(ENiagaraDistributionMode::ColorGradient, InPropertyHandle.GetProperty()))
+						{
+							ParsedDistribution->Mode = ENiagaraDistributionMode::ColorGradient;
+						}
 						ParsedDistribution->ChannelCurves = CurveCollectionValue.Curves;
 					}
 				}
diff --git a/Engine/Plugins/FX/Niagara/Source/NiagaraEditor/Private/Widgets/SNiagaraDistributionLookupValueWidget.cpp b/Engine/Plugins/FX/Niagara/Source/NiagaraEditor/Private/Widgets/SNiagaraDistributionLookupValueWidget.cpp
--- a/Engine/Plugins/FX/Niagara/Source/NiagaraEditor/Private/Widgets/SNiagaraDistributionLookupValueWidget.cpp
+++ b/Engine/Plugins/FX/Niagara/Source/NiagaraEditor/Private/Widgets/SNiagaraDistributionLookupValueWidget.cpp
@@ -27,6 +27,10 @@ void SNiagaraDistributionLookupValueWidget::Construct(const FArguments& InArgs,
 	}
 
 	LookupValueModeEnum = FNiagaraDistributionBase::GetLookupValueModeEnum(DistributionProperty);
+	if (LookupValueModeEnum == nullptr)
+	{
+		return;
+	}
 
 	TSharedRef<SHorizontalBox> BoxWidget = SNew(SHorizontalBox);
 	if (InArgs._ShowValueOnly == false)
@@ -65,23 +69,26 @@ TSharedRef<SWidget> SNiagaraDistributionLookupValueWidget::OnGetLookupValueModeO
 {
 	FMenuBuilder MenuBuilder(true, nullptr);
 	MenuBuilder.BeginSection("Actions");
-	for (int64 iValue = 0; iValue < LookupValueModeEnum->GetMaxEnumValue(); ++iValue)
+	if (LookupValueModeEnum)
 	{
-		if (LookupValueModeEnum->IsValidEnumValue(iValue) == false)
+		for (int64 iValue = 0; iValue < LookupValueModeEnum->GetMaxEnumValue(); ++iValue)
 		{
-			continue;
+			if (LookupValueModeEnum->IsValidEnumValue(iValue) == false)
+			{
+				continue;
+			}
+
+			const int64 iIndex = LookupValueModeEnum->GetIndexByValue(iValue);
+
+			MenuBuilder.AddMenuEntry(
+				LookupValueModeEnum->GetDisplayNameTextByIndex(iIndex),
+				LookupValueModeEnum->GetToolTipTextByIndex(iIndex),
+				FSlateIcon(),
+				FUIAction(
+					FExecuteAction::CreateSP(this, &SNiagaraDistributionLookupValueWidget::SetLookupValueMode, uint8(iValue))
+				)
+			);
 		}
-
-		const int64 iIndex = LookupValueModeEnum->GetIndexByValue(iValue);
-
-		MenuBuilder.AddMenuEntry(
-			LookupValueModeEnum->GetDisplayNameTextByIndex(iIndex),
-			LookupValueModeEnum->GetToolTipTextByIndex(iIndex),
-			FSlateIcon(),
-			FUIAction(
-				FExecuteAction::CreateSP(this, &SNiagaraDistributionLookupValueWidget::SetLookupValueMode, uint8(iValue))
-			)
-		);
 	}
 	MenuBuilder.EndSection();
 	return MenuBuilder.MakeWidget();
