Fixed hash function in data channel to use full string data instead of just half of it

diff --git a/Engine/Plugins/FX/Niagara/Source/Niagara/Private/NiagaraDataChannel.cpp b/Engine/Plugins/FX/Niagara/Source/Niagara/Private/NiagaraDataChannel.cpp
--- a/Engine/Plugins/FX/Niagara/Source/Niagara/Private/NiagaraDataChannel.cpp
+++ b/Engine/Plugins/FX/Niagara/Source/Niagara/Private/NiagaraDataChannel.cpp
@@ -1366,7 +1366,7 @@ void AddVarToHash(const FNiagaraVariable& Var, FBlake3& Builder)
 	FNameBuilder VarName(Var.GetName());
 	FStringView Name = VarName.ToView();
 	uint32 ClassHash = GetTypeHash(Var.GetType().ClassStructOrEnum);
-	Builder.Update(Name.GetData(), Name.Len());
+	Builder.Update(Name.GetData(), Name.Len() * sizeof(TCHAR));
 	Builder.Update(&ClassHash, sizeof(uint32));
 	Builder.Update(&Var.GetType().UnderlyingType, sizeof(uint16));
 }
