deallocation option for grids on 2d emitter
expose execute pin for setting number of cells on RT DI

#virtualized

diff --git a/Engine/Plugins/FX/Niagara/Source/Niagara/Private/NiagaraDataInterfaceRenderTarget2D.cpp b/Engine/Plugins/FX/Niagara/Source/Niagara/Private/NiagaraDataInterfaceRenderTarget2D.cpp
--- a/Engine/Plugins/FX/Niagara/Source/Niagara/Private/NiagaraDataInterfaceRenderTarget2D.cpp
+++ b/Engine/Plugins/FX/Niagara/Source/Niagara/Private/NiagaraDataInterfaceRenderTarget2D.cpp
@@ -69,6 +69,7 @@ struct FNDIRenderTarget2DFunctionVersion
 		InitialVersion = 0,
 		AddedOptionalExecute = 1,
 		AddedMipLevel = 2,
+		AddedOptionalExecuteToSetSize = 3,
 
 		VersionPlusOne,
 		LatestVersion = VersionPlusOne - 1
@@ -148,6 +149,7 @@ void UNiagaraDataInterfaceRenderTarget2D::GetFunctionsInternal(TArray<FNiagaraFu
 		Sig.Name = SetSizeFunctionName;
 		Sig.Inputs.Emplace(FNiagaraTypeDefinition::GetIntDef(), TEXT("Width"));
 		Sig.Inputs.Emplace(FNiagaraTypeDefinition::GetIntDef(), TEXT("Height"));
+		Sig.Inputs.Add_GetRef(FNiagaraVariable(FNiagaraTypeDefinition::GetBoolDef(), TEXT("Execute"))).SetValue(true);
 		Sig.Outputs.Emplace(FNiagaraTypeDefinition::GetBoolDef(), TEXT("Success"));
 		Sig.ModuleUsageBitmask = EmitterSystemOnlyBitmask;
 		Sig.bRequiresExecPin = true;
@@ -260,6 +262,16 @@ bool UNiagaraDataInterfaceRenderTarget2D::UpgradeFunctionCall(FNiagaraFunctionSi
 		}
 	}
 
+	if (FunctionSignature.FunctionVersion < FNDIRenderTarget2DFunctionVersion::AddedOptionalExecuteToSetSize)
+	{
+		if (FunctionSignature.Name == SetSizeFunctionName)
+		{
+			check(FunctionSignature.Inputs.Num() == 3);
+			FunctionSignature.Inputs.Add_GetRef(FNiagaraVariable(FNiagaraTypeDefinition::GetBoolDef(), TEXT("Execute"))).SetValue(true);			
+			bWasChanged = true;
+		}
+	}
+
 	// Set latest version
 	FunctionSignature.FunctionVersion = FNDIRenderTarget2DFunctionVersion::LatestVersion;
 
@@ -731,13 +743,15 @@ void UNiagaraDataInterfaceRenderTarget2D::VMSetSize(FVectorVMExternalFunctionCon
 	VectorVM::FUserPtrHandler<FRenderTarget2DRWInstanceData_GameThread> InstData(Context);
 	FNDIInputParam<int> InSizeX(Context);
 	FNDIInputParam<int> InSizeY(Context);
+	VectorVM::FExternalFuncInputHandler<FNiagaraBool> InExecute(Context);
 	FNDIOutputParam<FNiagaraBool> OutSuccess(Context);
 
 	for (int32 InstanceIdx = 0; InstanceIdx < Context.GetNumInstances(); ++InstanceIdx)
 	{
 		const int SizeX = InSizeX.GetAndAdvance();
 		const int SizeY = InSizeY.GetAndAdvance();
-		const bool bSuccess = (InstData.Get() != nullptr && Context.GetNumInstances() == 1 && SizeX >= 0 && SizeY >= 0);
+		bool bExecute = InExecute.GetAndAdvance();
+		const bool bSuccess = (InstData.Get() != nullptr && Context.GetNumInstances() == 1 && SizeX >= 0 && SizeY >= 0) && bExecute;
 		OutSuccess.SetAndAdvance(bSuccess);
 		if (bSuccess)
 		{
