- Initialize particle add lookup value mode options
- Add initial look up value mode enum
- Enable lookup on initial position
- Change color to a distribution

#virtualized

diff --git a/Engine/Plugins/FX/Niagara/Shaders/Private/Stateless/Modules/NiagaraStatelessModule_InitializeParticle.ush b/Engine/Plugins/FX/Niagara/Shaders/Private/Stateless/Modules/NiagaraStatelessModule_InitializeParticle.ush
--- a/Engine/Plugins/FX/Niagara/Shaders/Private/Stateless/Modules/NiagaraStatelessModule_InitializeParticle.ush
+++ b/Engine/Plugins/FX/Niagara/Shaders/Private/Stateless/Modules/NiagaraStatelessModule_InitializeParticle.ush
@@ -5,8 +5,7 @@ static const uint EInitializeParticleModuleFlag_UniformMeshScale	= 1 << 1;
 
 uint	InitializeParticle_ModuleFlags;
 FNiagaraStatelessBuiltDistributionType	InitializeParticle_InitialPosition;
-float4	InitializeParticle_ColorScale;
-float4	InitializeParticle_ColorBias;
+FNiagaraStatelessBuiltDistributionType	InitializeParticle_InitialColor;
 float2	InitializeParticle_SpriteSizeScale;
 float2	InitializeParticle_SpriteSizeBias;
 float	InitializeParticle_SpriteRotationScale;
@@ -23,8 +22,8 @@ void InitializeParticle_Simulate(inout FStatelessParticle Particle)
 	const bool bUniformSpriteSize = (InitializeParticle_ModuleFlags & EInitializeParticleModuleFlag_UniformSpriteSize) != 0;
 	const bool bUniformMeshScale  = (InitializeParticle_ModuleFlags & EInitializeParticleModuleFlag_UniformMeshScale) != 0;
 
-	Particle.Position		= SampleDistributionValueFloat3(InitializeParticle_InitialPosition, 0, 0.0f);
-	Particle.Color			= RandomScaleBiasFloat(1, InitializeParticle_ColorScale, InitializeParticle_ColorBias);
+	Particle.Position		= SampleDistributionValueFloat3(InitializeParticle_InitialPosition, 0, float(Particle.UniqueID));
+	Particle.Color			= SampleDistributionValueFloat4(InitializeParticle_InitialColor, 1, float(Particle.UniqueID));
 	Particle.RibbonWidth	= RandomScaleBiasFloat(2, InitializeParticle_RibbonWidthScale, InitializeParticle_RibbonWidthBias);
 	Particle.SpriteSize		= RandomScaleBiasFloat(3, InitializeParticle_SpriteSizeScale, InitializeParticle_SpriteSizeBias, bUniformSpriteSize);
 	Particle.SpriteRotation	= RandomScaleBiasFloat(4, InitializeParticle_SpriteRotationScale, InitializeParticle_SpriteRotationBias);
diff --git a/Engine/Plugins/FX/Niagara/Source/Niagara/Internal/Stateless/Modules/NiagaraStatelessModule_InitializeParticle.h b/Engine/Plugins/FX/Niagara/Source/Niagara/Internal/Stateless/Modules/NiagaraStatelessModule_InitializeParticle.h
--- a/Engine/Plugins/FX/Niagara/Source/Niagara/Internal/Stateless/Modules/NiagaraStatelessModule_InitializeParticle.h
+++ b/Engine/Plugins/FX/Niagara/Source/Niagara/Internal/Stateless/Modules/NiagaraStatelessModule_InitializeParticle.h
@@ -20,8 +20,8 @@ public:
 	UPROPERTY(EditAnywhere, Category = "Parameters", meta = (DisplayName = "Lifetime", Units="s"))
 	FNiagaraDistributionRangeFloat LifetimeDistribution = FNiagaraDistributionRangeFloat(FNiagaraStatelessGlobals::GetDefaultLifetimeValue());
 
-	UPROPERTY(EditAnywhere, Category = "Parameters", meta = (DisplayName = "Color", DisableCurveDistribution, DisableDistributionLookupValueMode))
-	FNiagaraDistributionRangeColor ColorDistribution = FNiagaraDistributionRangeColor(FNiagaraStatelessGlobals::GetDefaultColorValue());
+	UPROPERTY(EditAnywhere, Category = "Parameters", meta = (DisplayName = "Color", DistributionLookupValueEnumPath = "/Script/Niagara.ENiagaraDistributionInitialLookupValueMode"))
+	FNiagaraDistributionColor ColorDistribution = FNiagaraDistributionColor(FNiagaraStatelessGlobals::GetDefaultColorValue());
 
 	UPROPERTY(EditAnywhere, Category = "Parameters", meta = (DisplayName = "Mass"))
 	FNiagaraDistributionRangeFloat MassDistribution = FNiagaraDistributionRangeFloat(FNiagaraStatelessGlobals::GetDefaultMassValue());
@@ -41,7 +41,7 @@ public:
 	UPROPERTY(EditAnywhere, Category = "Parameters", meta = (DisplayName = "Ribbon Width", EditCondition="bWriteRibbonWidth"))
 	FNiagaraDistributionRangeFloat RibbonWidthDistribution = FNiagaraDistributionRangeFloat(FNiagaraStatelessGlobals::GetDefaultRibbonWidthValue());
 
-	UPROPERTY(EditAnywhere, Category = "Parameters", meta = (DisableCurveDistribution, DisableArrayDistribution, DisableDistributionLookupValueMode))
+	UPROPERTY(EditAnywhere, Category = "Parameters", meta = (DistributionLookupValueEnumPath = "/Script/Niagara.ENiagaraDistributionInitialLookupValueMode"))
 	FNiagaraDistributionPosition InitialPositionDistribution = FNiagaraDistributionPosition(FVector3f::ZeroVector);
 
 	virtual void BuildEmitterData(const FNiagaraStatelessEmitterDataBuildContext& BuildContext) const override;
diff --git a/Engine/Plugins/FX/Niagara/Source/Niagara/Private/Stateless/Modules/NiagaraStatelessModule_InitializeParticle.cpp b/Engine/Plugins/FX/Niagara/Source/Niagara/Private/Stateless/Modules/NiagaraStatelessModule_InitializeParticle.cpp
--- a/Engine/Plugins/FX/Niagara/Source/Niagara/Private/Stateless/Modules/NiagaraStatelessModule_InitializeParticle.cpp
+++ b/Engine/Plugins/FX/Niagara/Source/Niagara/Private/Stateless/Modules/NiagaraStatelessModule_InitializeParticle.cpp
@@ -22,9 +22,9 @@ namespace NSMInitializeParticlePrivate
 	struct FModuleBuiltData
 	{
 		uint32							ModuleFlags = 0;
-		FUintVector3					InitialPosition = FUintVector3::ZeroValue;
 		FNiagaraStatelessRangeFloat		LifetimeRange;
-		FNiagaraStatelessRangeColor		ColorRange;
+		FUintVector3					InitialPosition = FUintVector3::ZeroValue;
+		FUintVector3					InitialColor = FUintVector3::ZeroValue;
 		FNiagaraStatelessRangeFloat		MassRange;
 		FNiagaraStatelessRangeVector2	SpriteSizeRange;
 		FNiagaraStatelessRangeFloat		SpriteRotationRange;
@@ -55,10 +55,12 @@ namespace NSMInitializeParticlePrivate
 		const bool bUniformSpriteSize = (ModuleBuiltData->ModuleFlags & EInitializeParticleModuleFlag_UniformSpriteSize) != 0;
 		const bool bUniformMeshScale = (ModuleBuiltData->ModuleFlags & EInitializeParticleModuleFlag_UniformMeshScale) != 0;
 
+		const int32* ParticleUniqueIDs = ParticleSimulationContext.GetParticleUniqueID();
+
 		for (uint32 i = 0; i < ParticleSimulationContext.GetNumInstances(); ++i)
 		{
-			const FVector3f		Position = ParticleSimulationContext.SampleDistributionValue<FVector3f>(ModuleBuiltData->InitialPosition, i, 0, 0.0f);
-			const FLinearColor	Color = ParticleSimulationContext.RandomScaleBiasFloat(i, 1, ShaderParameters->InitializeParticle_ColorScale, ShaderParameters->InitializeParticle_ColorBias);
+			const FVector3f		Position = ParticleSimulationContext.SampleDistributionValue<FVector3f>(ModuleBuiltData->InitialPosition, i, 0, float(ParticleUniqueIDs[i]));
+			const FLinearColor	Color = ParticleSimulationContext.SampleDistributionValue<FLinearColor>(ModuleBuiltData->InitialColor, i, 1, float(ParticleUniqueIDs[i]));
 			const float			RibbonWidth = ParticleSimulationContext.RandomScaleBiasFloat(i, 2, ShaderParameters->InitializeParticle_RibbonWidthScale, ShaderParameters->InitializeParticle_RibbonWidthBias);
 			const FVector2f		SpriteSize = ParticleSimulationContext.RandomScaleBiasFloat(i, 3, ShaderParameters->InitializeParticle_SpriteSizeScale, ShaderParameters->InitializeParticle_SpriteSizeBias, bUniformSpriteSize);
 			const float			SpriteRot = ParticleSimulationContext.RandomScaleBiasFloat(i, 4, ShaderParameters->InitializeParticle_SpriteRotationScale, ShaderParameters->InitializeParticle_SpriteRotationBias);
@@ -90,9 +92,9 @@ void UNiagaraStatelessModule_InitializeParticle::BuildEmitterData(const FNiagara
 	BuiltData->ModuleFlags				 = SpriteSizeDistribution.IsUniform() ? EInitializeParticleModuleFlag_UniformSpriteSize : 0;
 	BuiltData->ModuleFlags				|= MeshScaleDistribution.IsUniform() ? EInitializeParticleModuleFlag_UniformMeshScale : 0;
 
-	BuiltData->InitialPosition			= BuildContext.AddDistribution(InitialPositionDistribution);
 	BuiltData->LifetimeRange			= LifetimeDistribution.CalculateRange(FNiagaraStatelessGlobals::GetDefaultLifetimeValue());
-	BuiltData->ColorRange				= BuildContext.ConvertDistributionToRange(ColorDistribution, FNiagaraStatelessGlobals::GetDefaultColorValue());
+	BuiltData->InitialPosition			= BuildContext.AddDistribution(InitialPositionDistribution);
+	BuiltData->InitialColor				= BuildContext.AddDistribution(ColorDistribution);
 	BuiltData->MassRange				= BuildContext.ConvertDistributionToRange(MassDistribution, FNiagaraStatelessGlobals::GetDefaultMassValue());
 	BuiltData->SpriteSizeRange			= BuildContext.ConvertDistributionToRange(SpriteSizeDistribution, FNiagaraStatelessGlobals::GetDefaultSpriteSizeValue());
 	BuiltData->SpriteRotationRange		= BuildContext.ConvertDistributionToRange(SpriteRotationDistribution, FNiagaraStatelessGlobals::GetDefaultSpriteRotationValue());
@@ -132,7 +134,7 @@ void UNiagaraStatelessModule_InitializeParticle::SetShaderParameters(const FNiag
 
 	Parameters->InitializeParticle_ModuleFlags			= ModuleBuiltData->ModuleFlags;
 	Parameters->InitializeParticle_InitialPosition		= ModuleBuiltData->InitialPosition;
-	SetShaderParameterContext.ConvertRangeToScaleBias(ModuleBuiltData->ColorRange,			Parameters->InitializeParticle_ColorScale, Parameters->InitializeParticle_ColorBias);
+	Parameters->InitializeParticle_InitialColor			= ModuleBuiltData->InitialColor;
 	SetShaderParameterContext.ConvertRangeToScaleBias(ModuleBuiltData->SpriteSizeRange,		Parameters->InitializeParticle_SpriteSizeScale, Parameters->InitializeParticle_SpriteSizeBias);
 	SetShaderParameterContext.ConvertRangeToScaleBias(ModuleBuiltData->SpriteRotationRange, Parameters->InitializeParticle_SpriteRotationScale, Parameters->InitializeParticle_SpriteRotationBias);
 	SetShaderParameterContext.ConvertRangeToScaleBias(ModuleBuiltData->MeshScaleRange,		Parameters->InitializeParticle_MeshScaleScale, Parameters->InitializeParticle_MeshScaleBias);
diff --git a/Engine/Plugins/FX/Niagara/Source/Niagara/Private/Stateless/NiagaraStatelessDistribution.cpp b/Engine/Plugins/FX/Niagara/Source/Niagara/Private/Stateless/NiagaraStatelessDistribution.cpp
--- a/Engine/Plugins/FX/Niagara/Source/Niagara/Private/Stateless/NiagaraStatelessDistribution.cpp
+++ b/Engine/Plugins/FX/Niagara/Source/Niagara/Private/Stateless/NiagaraStatelessDistribution.cpp
@@ -1175,6 +1175,39 @@ void FNiagaraDistributionCurveVector3::UpdateValuesFromDistribution()
 }
 #endif
 
+#if WITH_EDITORONLY_DATA
+bool FNiagaraDistributionColor::SerializeFromMismatchedTag(const struct FPropertyTag& Tag, FStructuredArchive::FSlot Slot)
+{
+	if (Tag.GetType().IsStruct(FNiagaraDistributionRangeColor::StaticStruct()->GetFName()))
+	{
+		FNiagaraDistributionRangeColor RangeValue;
+		FNiagaraDistributionRangeColor::StaticStruct()->SerializeItem(Slot, &RangeValue, nullptr);
+
+		*static_cast<FNiagaraDistributionBase*>(this) = static_cast<FNiagaraDistributionBase&>(RangeValue);
+
+		// Range would use UniformConstant / UniformRange when meaning non uniform so we need to convert to the correct mode here
+		switch (Mode)
+		{
+			case ENiagaraDistributionMode::UniformConstant:
+			case ENiagaraDistributionMode::NonUniformConstant:
+				Mode = ENiagaraDistributionMode::NonUniformConstant;
+				break;
+
+			case ENiagaraDistributionMode::UniformRange:
+			case ENiagaraDistributionMode::NonUniformRange:
+				Mode = ENiagaraDistributionMode::NonUniformRange;
+				break;
+		}
+
+		UpdateValuesFromDistribution();
+
+		return true;
+	}
+
+	return false;
+}
+#endif
+
 bool FNiagaraDistributionRangeFloat::SerializeFromMismatchedTag(const struct FPropertyTag& Tag, FStructuredArchive::FSlot Slot)
 {
 	if (Tag.Type == NAME_FloatProperty)
diff --git a/Engine/Plugins/FX/Niagara/Source/Niagara/Public/Stateless/NiagaraStatelessDistribution.h b/Engine/Plugins/FX/Niagara/Source/Niagara/Public/Stateless/NiagaraStatelessDistribution.h
--- a/Engine/Plugins/FX/Niagara/Source/Niagara/Public/Stateless/NiagaraStatelessDistribution.h
+++ b/Engine/Plugins/FX/Niagara/Source/Niagara/Public/Stateless/NiagaraStatelessDistribution.h
@@ -62,6 +62,16 @@ enum class ENiagaraDistributionLookupValueMode : uint8
 	ParticlesNormalizedAge = 0 UMETA(TimeRangeMode = "Normalized", DisplayName = "Particles.NormalizedAge"),
 };
 
+// Distribution type generally used in initial modules
+UENUM()
+enum class ENiagaraDistributionInitialLookupValueMode : uint8
+{
+	// Use a randomly generated value, evaluted at birth only
+	Random = 0xff UMETA(TimeRangeMode = "Normalized"),
+	// Bound to the variable Particles.UniqueID
+	UniqueID = 0 UMETA(TimeRangeMode = "Direct", DisplayName = "Particles.UniqueID"),
+};
+
 enum class ENiagaraDistributionCurveLUTMode
 {
 	Sample,		// Each sample in the LUT represents the curve evaulation
@@ -469,6 +479,7 @@ struct FNiagaraDistributionColor : public FNiagaraDistributionBase
 	NIAGARA_API virtual void UpdateValuesFromDistribution() override;
 	NIAGARA_API virtual void UpdateTimeRangeFromDistribution(ENiagaraDistributionTimeRangeMode TimeRangeMode, const FVector2f& CustomTimeRange) override;
 	virtual FNiagaraTypeDefinition GetBindingTypeDef() const { return FNiagaraTypeDefinition::GetColorDef(); }
+	bool SerializeFromMismatchedTag(const struct FPropertyTag& Tag, FStructuredArchive::FSlot Slot);
 #endif
 };
 
@@ -566,3 +577,14 @@ struct TStructOpsTypeTraits<FNiagaraDistributionRangeRotator> : public TStructOp
 		WithStructuredSerializeFromMismatchedTag = true,
 	};
 };
+
+#if WITH_EDITORONLY_DATA
+template<>
+struct TStructOpsTypeTraits<FNiagaraDistributionColor> : public TStructOpsTypeTraitsBase2<FNiagaraDistributionColor>
+{
+	enum
+	{
+		WithStructuredSerializeFromMismatchedTag = true,
+	};
+};
+#endif
diff --git a/Engine/Plugins/FX/Niagara/Source/NiagaraShader/Internal/Stateless/NiagaraStatelessModuleShaderParameters.h b/Engine/Plugins/FX/Niagara/Source/NiagaraShader/Internal/Stateless/NiagaraStatelessModuleShaderParameters.h
--- a/Engine/Plugins/FX/Niagara/Source/NiagaraShader/Internal/Stateless/NiagaraStatelessModuleShaderParameters.h
+++ b/Engine/Plugins/FX/Niagara/Source/NiagaraShader/Internal/Stateless/NiagaraStatelessModuleShaderParameters.h
@@ -11,9 +11,7 @@ namespace NiagaraStateless
 	BEGIN_SHADER_PARAMETER_STRUCT(FInitializeParticleModule_ShaderParameters, NIAGARASHADER_API)
 		SHADER_PARAMETER(uint32,									InitializeParticle_ModuleFlags)
 		SHADER_PARAMETER(FNiagaraStatelessBuiltDistributionType,	InitializeParticle_InitialPosition)
-
-		SHADER_PARAMETER(FLinearColor,	InitializeParticle_ColorScale)				// Unset / Direct Set / Random Range(Link RGBA / Link RGB|Link A / Random Channels) / Random Hue|Saturation|Value
-		SHADER_PARAMETER(FLinearColor,	InitializeParticle_ColorBias)
+		SHADER_PARAMETER(FNiagaraStatelessBuiltDistributionType,	InitializeParticle_InitialColor)
 
 		SHADER_PARAMETER(FVector2f, InitializeParticle_SpriteSizeScale)			// Unset / Uniform / Random Uniform / Non-Uniform / Random Non-Uniform
 		SHADER_PARAMETER(FVector2f, InitializeParticle_SpriteSizeBias)
