Cell Size and Bounds Padding converted to FVectors.
Modified some handling of cell coords.

diff --git a/Engine/Plugins/FX/Niagara/Source/Niagara/Private/NiagaraDataChannel_GameplayBurst.cpp b/Engine/Plugins/FX/Niagara/Source/Niagara/Private/NiagaraDataChannel_GameplayBurst.cpp
--- a/Engine/Plugins/FX/Niagara/Source/Niagara/Private/NiagaraDataChannel_GameplayBurst.cpp
+++ b/Engine/Plugins/FX/Niagara/Source/Niagara/Private/NiagaraDataChannel_GameplayBurst.cpp
@@ -11,13 +11,20 @@
 
 namespace NDCGameplayBurst_Local
 {
+	[[nodiscard]] inline int64 GetCellCoordAsInt64(double CellCoord)
+	{
+		ensure(FMath::IsFinite(CellCoord));
+		ensure(CellCoord == FMath::FloorToDouble(CellCoord));
+		return static_cast<int64>(CellCoord);
+	}
+
 	//Returns the grid cell for the given location and cell size.
-	inline FInt64Vector3 GetCellCoords(FVector Location, double InvCellSize)
+	//Cell size is stored as double which is accurate up to ~9.5ly assuming a 10m cell size.
+	[[nodiscard]] inline FVector GetCellCoords(FVector Location, FVector InvCellSize)
 	{
-		return FInt64Vector3(
-			(int64)FMath::FloorToInt(Location.X * InvCellSize),
-			(int64)FMath::FloorToInt(Location.Y * InvCellSize),
-			(int64)FMath::FloorToInt(Location.Z * InvCellSize));
+		return FVector(FMath::FloorToDouble(Location.X * InvCellSize.X),
+			FMath::FloorToDouble(Location.Y * InvCellSize.Y),
+			FMath::FloorToDouble(Location.Z * InvCellSize.Z));
 	}
 };
  
@@ -41,7 +48,7 @@ void FNDCMapEntry_GameplayBurst::Init(FNDCAccessContextInst& AccessContext, UNia
 
 	UObject* SysToSpawn = TypedOwner->GetSystemToSpawn(AccessContext);
 
-	int32 BoundsPadding = TypedContext.bOverrideBoundsPadding ? TypedContext.SystemBoundsPadding : TypedOwner->DefaultSystemBoundsPadding;
+	FVector BoundsPadding = TypedContext.bOverrideBoundsPadding ? TypedContext.SystemBoundsPadding : TypedOwner->DefaultSystemBoundsPadding;
 	bool bIsAttached = TypedContext.bAttachedToOwningComponent;
 	if (bIsAttached)
 	{
@@ -75,12 +82,12 @@ void FNDCMapEntry_GameplayBurst::Init(FNDCAccessContextInst& AccessContext, UNia
 
 		if (SysToSpawn)
 		{
-			int32 CellSize = TypedContext.bOverrideCellSize ? TypedContext.CellSizeOverride : TypedOwner->DefaultCellSize;
+			FVector CellSize = TypedContext.bOverrideCellSize ? TypedContext.CellSizeOverride : TypedOwner->DefaultCellSize;
 			FVector CellExtents = FVector(CellSize * 0.5);
-			FInt64Vector3 CellCoord = NDCGameplayBurst_Local::GetCellCoords(Location, 1.0 / CellSize);
-			FVector CellCentre = FVector(CellCoord * CellSize) + CellExtents;
-			CellExtents += FVector(BoundsPadding);
-			
+			FVector CellCoord = NDCGameplayBurst_Local::GetCellCoords(Location, CellSize.Reciprocal());
+			FVector CellCentre = CellCoord * CellSize + CellExtents;
+
+			CellExtents += BoundsPadding;
 			SpawnSystem(AccessContext.GetChecked<FNDCAccessContext_MapBase>(), SysToSpawn, nullptr, CellCentre, CellExtents, true);
 		}
 	}
@@ -234,12 +241,11 @@ void UNiagaraDataChannelHandler_GameplayBurst::GenerateKey(FNDCAccessContextInst
 	//Inform any interested calling code that we are using the attached path so dest data and all spawned FX are specific to the given owner component.
 	TypedContext.bAttachedToOwningComponent = bUseAttachedPath;
 
+	uint8 Flags = bUseAttachedPath ? 1 : 0;
+
 	if (bUseAttachedPath)
 	{
-		//First key entry is an indicator for the main path being taken.
-		//Should aid readability/debugging and avoid the very low chance of key conflicts between paths.
-		uint8 KeyPrefix = 0;
-		KeyWriter << KeyPrefix;
+		KeyWriter << Flags;
 
 		//In the attached path we use the component ptr as key.
 		uint32 ObjID = Comp->GetUniqueID();
@@ -247,21 +253,25 @@ void UNiagaraDataChannelHandler_GameplayBurst::GenerateKey(FNDCAccessContextInst
 	}
 	else
 	{
-		uint8 KeyPrefix = 1;
-		KeyWriter << KeyPrefix;
+		Flags |= TypedContext.bOverrideCellSize ?  (1 << 1): 0;
+
+		KeyWriter << Flags;
 
 		//We're not attached so write our world grid cell to the key.
 		FVector Location = TypedContext.GetLocation();
-		int32 CellSize = TypedContext.bOverrideCellSize ? TypedContext.CellSizeOverride : DefaultCellSize;
-		FInt64Vector3 CellCoord = NDCGameplayBurst_Local::GetCellCoords(Location, 1.0 / CellSize);
-		KeyWriter << CellCoord;
+		FVector CellSize = TypedContext.bOverrideCellSize ? TypedContext.CellSizeOverride : DefaultCellSize;
+		FVector CellCoord = NDCGameplayBurst_Local::GetCellCoords(Location, CellSize.Reciprocal());
 
-		KeyWriter << TypedContext.GameplayTag.GetTagName();
+		KeyWriter << NDCGameplayBurst_Local::GetCellCoordAsInt64(CellCoord.X);
+		KeyWriter << NDCGameplayBurst_Local::GetCellCoordAsInt64(CellCoord.Y);
+		KeyWriter << NDCGameplayBurst_Local::GetCellCoordAsInt64(CellCoord.Z);
 
 		if(TypedContext.bOverrideCellSize)
 		{
 			KeyWriter << TypedContext.CellSizeOverride;
 		}
+		
+		KeyWriter << TypedContext.GameplayTag.GetTagName();
 	}
 
 	Super::GenerateKey(AccessContext, KeyWriter);
diff --git a/Engine/Plugins/FX/Niagara/Source/Niagara/Public/NiagaraDataChannel_GameplayBurst.h b/Engine/Plugins/FX/Niagara/Source/Niagara/Public/NiagaraDataChannel_GameplayBurst.h
--- a/Engine/Plugins/FX/Niagara/Source/Niagara/Public/NiagaraDataChannel_GameplayBurst.h
+++ b/Engine/Plugins/FX/Niagara/Source/Niagara/Public/NiagaraDataChannel_GameplayBurst.h
@@ -40,7 +40,7 @@ struct FNDCAccessContext_GameplayBurst : public FNDCAccessContext_MapBase
 	Use with care. Overriding cell size provides more flexibility for users but can cause problems if not done correctly.
 	*/
 	UPROPERTY(EditAnywhere, BlueprintReadWrite, Category = "Input", meta = (NDCAccessContextInput, EditCondition="bOverrideCellSize"))
-	int32 CellSizeOverride = 2500;
+	FVector CellSizeOverride = FVector(2500.0);
 
 	/** 
 	Padding applied to the bounds of spawned handler systems. 
@@ -48,7 +48,7 @@ struct FNDCAccessContext_GameplayBurst : public FNDCAccessContext_MapBase
 	When systems are attached to the owning component, bounds are the component bounds + SystemBoundsPadding.
 	*/
 	UPROPERTY(EditAnywhere, BlueprintReadWrite, Category = "Input", meta = (NDCAccessContextInput, EditCondition="bOverrideBoundsPadding"))
-	int32 SystemBoundsPadding = 250;
+	FVector SystemBoundsPadding = FVector(250.0);
 
 	/** Data can optionally be routed by a given Gameplay Tag for additional variation. */
 	UPROPERTY(EditAnywhere, BlueprintReadWrite, Category = "Input", meta = (NDCAccessContextInput))
@@ -98,8 +98,8 @@ public:
 
 	[[nodiscard]] virtual TNDCAccessContextType GetAccessContextType()const override { return TNDCAccessContextType(FNDCAccessContext_GameplayBurst::StaticStruct()); }
 	[[nodiscard]] const FNDCGameplayBurstAttachmentSettings& GetAttachmentSettings()const{ return AttachmentSettings; }
-	[[nodiscard]] float GetGridCellSize()const { return CellSize; }
-	[[nodiscard]] float GetSystemBoundsPadding()const { return SystemBoundsPadding; }
+	[[nodiscard]] FVector GetGridCellSize()const { return CellSize; }
+	[[nodiscard]] FVector GetSystemBoundsPadding()const { return SystemBoundsPadding; }
 
 protected:
 	
@@ -108,7 +108,7 @@ protected:
 	Internal data and their spawned handler FX are bucketed internally in a world aligned grid of cells.
 	*/
 	UPROPERTY(EditAnywhere, Category = "Grid", meta = (NDCAccessContextInput))
-	int32 CellSize = 2500;
+	FVector CellSize = FVector(2500.0);
 
 	/** 
 	Default System Bounds Padding. Can be overridden by users via the AccessContext.
@@ -117,7 +117,7 @@ protected:
 	For non attached, the bounds will be the CellSize + padding.
 	*/
 	UPROPERTY(EditAnywhere, Category = "Grid", meta = (NDCAccessContextInput))
-	int32 SystemBoundsPadding = 250;
+	FVector SystemBoundsPadding = FVector(250.0);
 
 	/** 
 	Settings controlling whether incoming data should route to the main world aligned spatial grids or use a special bucket attached to and owned by a specific component.
@@ -165,8 +165,8 @@ protected:
 	UPROPERTY()
 	FNDCGameplayBurstAttachmentSettings AttachmentSettings;
 
-	int32 DefaultCellSize = 2500;
-	int32 DefaultSystemBoundsPadding = 250;
+	FVector DefaultCellSize = FVector(2500.0);
+	FVector DefaultSystemBoundsPadding = FVector(250.0);
 
 	/** A tag we use to mark spawned attached components to skip some work re-checking attachment settings. */
 	static const FName AttachedComponentTag;
