- Add setting to enable HLSL 2021 in Niagara
- HLSL 2021 features will fail to compile on the VM and on some platforms, but this allows users to experiment

#tests Make a HLSL file using template (for example) include this in a custom HLSL node and use the template function

diff --git a/Engine/Plugins/FX/Niagara/Source/Niagara/Private/NiagaraSettings.cpp b/Engine/Plugins/FX/Niagara/Source/Niagara/Private/NiagaraSettings.cpp
--- a/Engine/Plugins/FX/Niagara/Source/Niagara/Private/NiagaraSettings.cpp
+++ b/Engine/Plugins/FX/Niagara/Source/Niagara/Private/NiagaraSettings.cpp
@@ -40,9 +40,7 @@ FText UNiagaraSettings::GetSectionText() const
 {
 	return NSLOCTEXT("NiagaraPlugin", "NiagaraSettingsSection", "Niagara");
 }
-#endif
 
-#if WITH_EDITOR
 void UNiagaraSettings::PostEditChangeProperty(FPropertyChangedEvent& PropertyChangedEvent)
 {
 	if (PropertyChangedEvent.Property != nullptr)
@@ -65,3 +63,14 @@ UNiagaraEffectType* UNiagaraSettings::GetDefaultEffectType() const
 
 #endif
 
+void UNiagaraSettings::PostInitProperties()
+{
+	Super::PostInitProperties();
+
+#if WITH_EDITORONLY_DATA
+	if (IConsoleVariable* CVar = IConsoleManager::Get().FindConsoleVariable(TEXT("fx.Niagara.Shader.EnableHLSL2021")))
+	{
+		CVar->Set(bEnableHLSL2021);
+	}
+#endif
+}
diff --git a/Engine/Plugins/FX/Niagara/Source/Niagara/Public/NiagaraSettings.h b/Engine/Plugins/FX/Niagara/Source/Niagara/Public/NiagaraSettings.h
--- a/Engine/Plugins/FX/Niagara/Source/Niagara/Public/NiagaraSettings.h
+++ b/Engine/Plugins/FX/Niagara/Source/Niagara/Public/NiagaraSettings.h
@@ -201,6 +201,13 @@ class UNiagaraSettings : public UDeveloperSettings
 #if WITH_EDITORONLY_DATA
 	UPROPERTY(config, EditAnywhere, Category = Niagara, meta = (DisplayName = "Option for how to compile Niagara scripts"))
 	ENiagaraCompilationMode CompilationMode = ENiagaraCompilationMode::AsyncTasks;
+
+	/**
+	Enable HLSL 2021 compilation flag.
+	Note: Only GPU simulations on some platforms will work with HLSL 2021.
+	*/
+	UPROPERTY(config, EditAnywhere, Category = Niagara, meta = (ConfigRestartRequired = true, ConsoleVariable="fx.Niagara.Shader.EnableHLSL2021"))
+	bool bEnableHLSL2021 = false;
 #endif
 
 	/** The quality levels Niagara uses. */
@@ -316,4 +323,6 @@ public:
 protected:
 	static NIAGARA_API FOnNiagaraSettingsChanged SettingsChangedDelegate;
 #endif
+
+	NIAGARA_API virtual void PostInitProperties() override;
 };
diff --git a/Engine/Plugins/FX/Niagara/Source/NiagaraShader/Private/NiagaraShader.cpp b/Engine/Plugins/FX/Niagara/Source/NiagaraShader/Private/NiagaraShader.cpp
--- a/Engine/Plugins/FX/Niagara/Source/NiagaraShader/Private/NiagaraShader.cpp
+++ b/Engine/Plugins/FX/Niagara/Source/NiagaraShader/Private/NiagaraShader.cpp
@@ -1,6 +1,7 @@
 // Copyright Epic Games, Inc. All Rights Reserved.
 
 #include "NiagaraShader.h"
+#include "NiagaraSettings.h"
 #include "NiagaraShared.h"
 #include "NiagaraScriptBase.h"
 #include "RenderingThread.h"
@@ -51,6 +52,14 @@ static FAutoConsoleVariableRef CVarNiagaraShaderForceBindEverything(
 	TEXT("Forces Niagara to display errors about missing shader bindings.")
 );
 
+bool GNiagaraEnableHLSL2021 = false;
+static FAutoConsoleVariableRef CVarNiagaraEnableHLSL2021(
+	TEXT("fx.Niagara.Shader.EnableHLSL2021"),
+	GNiagaraEnableHLSL2021,
+	TEXT("Enable HLSL 2021 compilation flag.\n")
+	TEXT("Note: Only GPU simulations on some platforms will work with HLSL 2021.")
+);
+
 #if ENABLE_COOK_STATS
 namespace NiagaraShaderCookStats
 {
@@ -358,6 +367,11 @@ void FNiagaraShaderType::BeginCompileShader(
 	NewJob->Input.Environment.SetDefine(TEXT("THREADGROUP_SIZE_Y"), ThreadGroupSize.Y);
 	NewJob->Input.Environment.SetDefine(TEXT("THREADGROUP_SIZE_Z"), ThreadGroupSize.Z);
 
+	if (GNiagaraEnableHLSL2021)
+	{
+		NewJob->Input.Environment.CompilerFlags.Add(CFLAG_HLSL2021);
+	}
+
 	Script->ModifyCompilationEnvironment(Platform, NewJob->Input.Environment);
 
 	AddUniformBufferIncludesToEnvironment(NewJob->Input.Environment, Platform);
@@ -436,6 +450,11 @@ void FNiagaraShaderType::BeginCompileShaderFromSource(
 	NewJob->Input.Environment.SetDefine(TEXT("THREADGROUP_SIZE_Y"), ThreadGroupSize.Y);
 	NewJob->Input.Environment.SetDefine(TEXT("THREADGROUP_SIZE_Z"), ThreadGroupSize.Z);
 
+	if (GNiagaraEnableHLSL2021)
+	{
+		NewJob->Input.Environment.CompilerFlags.Add(CFLAG_HLSL2021);
+	}
+
 	AddUniformBufferIncludesToEnvironment(NewJob->Input.Environment, Platform);
 
 	FShaderCompilerEnvironment& ShaderEnvironment = NewJob->Input.Environment;
