- Handle null data interface inside Niagara Variant
- It is unclear how the variant DI becomes nullptr at the moment, this will at least avoid crashing the editor, but more investigation is required to understand why

diff --git a/Engine/Plugins/FX/Niagara/Source/Niagara/Private/NiagaraComponent.cpp b/Engine/Plugins/FX/Niagara/Source/Niagara/Private/NiagaraComponent.cpp
--- a/Engine/Plugins/FX/Niagara/Source/Niagara/Private/NiagaraComponent.cpp
+++ b/Engine/Plugins/FX/Niagara/Source/Niagara/Private/NiagaraComponent.cpp
@@ -3508,9 +3508,11 @@ void UNiagaraComponent::EnsureOverrideParametersConsistent() const
 			if (Key.IsDataInterface())
 			{
 				const UNiagaraDataInterface* ActualDI = OverrideParameters.GetDataInterface(Key);
-				if (ActualDI != nullptr)
+				const UNiagaraDataInterface* VariantDI = OverrideValue.GetDataInterface();
+				if (ActualDI != nullptr && VariantDI != nullptr)
 				{
-					ensureAlways(OverrideValue.GetDataInterface()->Equals(ActualDI));
+					// Null DIs are logged elsewhere so we don't need to raise the issue here also
+					ensureAlways(VariantDI->Equals(ActualDI));
 				}
 			}
 			else if (Key.IsUObject())
@@ -3869,11 +3871,11 @@ void UNiagaraComponent::FixInvalidUserParameterOverrideData()
 			UNiagaraDataInterface* DataInterface = OverrideParameters.GetDataInterface(ExistingVar);
 			if (DataInterface == nullptr)
 			{
-				UE_LOG(LogNiagara, Log, TEXT("ParameterStore has null data interface '%s' for variable '%s', making a default."), *ExistingVar.GetType().GetClass()->GetName(), *ExistingVar.GetName().ToString());
+				UE_LOG(LogNiagara, Log, TEXT("ParameterStore has null data interface '%s' for variable '%s' asset '%s', making a default."), *ExistingVar.GetType().GetClass()->GetName(), *ExistingVar.GetName().ToString(), *GetNameSafe(Asset));
 			}
 			else if (DataInterface->GetClass() != ExistingVar.GetType().GetClass())
 			{
-				UE_LOG(LogNiagara, Log, TEXT("ParameterStore has invalid type '%s' data interface '%s' for variable '%s', making a default."), *DataInterface->GetClass()->GetName(), *ExistingVar.GetType().GetClass()->GetName(), *ExistingVar.GetName().ToString());
+				UE_LOG(LogNiagara, Log, TEXT("ParameterStore has invalid type '%s' data interface '%s' for variable '%s' asset '%s', making a default."), *DataInterface->GetClass()->GetName(), *ExistingVar.GetType().GetClass()->GetName(), *ExistingVar.GetName().ToString(), *GetNameSafe(Asset));
 				DataInterface = nullptr;
 			}
 
@@ -4075,6 +4077,17 @@ void UNiagaraComponent::SetOverrideParameterStoreValue(const FNiagaraVariableBas
 	if (InKey.IsDataInterface())
 	{
 		UNiagaraDataInterface* OriginalDI = InValue.GetDataInterface();
+		if (OriginalDI == nullptr)
+		{
+			// There is an edge case where the data interface inside FNiagaraVariant has been garbage collected, it's unclear on a consistent repro right now.
+			// Rather than crashing make a default and log a warning for investigation (see UE-314893).
+			UClass* DataInterfaceClass = InKey.GetType().GetClass();
+			check(DataInterfaceClass);
+			UE_LOG(LogNiagara, Log, TEXT("SetOverrideParameterStoreValue has null data interface '%s' for variable '%s' asset '%s', making a default."), *DataInterfaceClass->GetName(), *InKey.GetName().ToString(), *GetNameSafe(Asset));
+
+			const EObjectFlags DataInterfaceObjectFlags = UNiagaraDataInterface::BuildObjectFlagsForOwner(this, RF_Transactional);
+			OriginalDI = NewObject<UNiagaraDataInterface>(this, DataInterfaceClass, NAME_None, DataInterfaceObjectFlags);
+		}
 		UNiagaraDataInterface* DuplicatedDI = DuplicateObject(OriginalDI, this);
 		OverrideParameters.SetDataInterface(DuplicatedDI, InKey);
 		DuplicatedDI->SetUsedWithCPUScript(OriginalDI->IsUsedWithCPUScript());
