- Support read / write icons for stateless
- Allows you to more easily see where bindings are used in LW emitters

diff --git a/Engine/Plugins/FX/Niagara/Source/Niagara/Private/Stateless/NiagaraStatelessDistribution.cpp b/Engine/Plugins/FX/Niagara/Source/Niagara/Private/Stateless/NiagaraStatelessDistribution.cpp
--- a/Engine/Plugins/FX/Niagara/Source/Niagara/Private/Stateless/NiagaraStatelessDistribution.cpp
+++ b/Engine/Plugins/FX/Niagara/Source/Niagara/Private/Stateless/NiagaraStatelessDistribution.cpp
@@ -741,6 +741,38 @@ void FNiagaraDistributionBase::ForEachParameterBinding(TFunction<void(const FNia
 	}
 }
 
+void FNiagaraDistributionBase::ForEachParameterBinding(const UStruct* ObjectStruct, const void* Object, TFunction<void(const FNiagaraVariableBase&)> Delegate)
+{
+	for (TFieldIterator<FStructProperty> PropIt(ObjectStruct); PropIt; ++PropIt)
+	{
+		if (FStructProperty* StructProp = *PropIt)
+		{
+			if (StructProp->Struct && StructProp->Struct->IsChildOf(FNiagaraDistributionBase::StaticStruct()))
+			{
+				const FNiagaraDistributionBase* Distribution = StructProp->ContainerPtrToValuePtr<const FNiagaraDistributionBase>(Object);
+				Distribution->ForEachParameterBinding(Delegate);
+			}
+			else if (StructProp->Struct && StructProp->Struct->IsChildOf(FNiagaraDistributionRangeInt::StaticStruct()))
+			{
+				const FNiagaraDistributionRangeInt* Distribution = StructProp->ContainerPtrToValuePtr<const FNiagaraDistributionRangeInt>(Object);
+				if (Distribution->IsBinding())
+				{
+					Delegate(Distribution->ParameterBinding);
+				}
+				else if (Distribution->IsExpression())
+				{
+					FNiagaraStatelessExpression::ForEachBinding(Distribution->ParameterExpression, Delegate);
+				}
+			}
+		}
+	}
+}
+
+void FNiagaraDistributionBase::ForEachParameterBinding(const UObject* InObject, TFunction<void(const FNiagaraVariableBase&)> Delegate)
+{
+	ForEachParameterBinding(InObject->GetClass(), InObject, Delegate);
+}
+
 bool FNiagaraDistributionBase::AllowDistributionMode(ENiagaraDistributionMode InMode, const FProperty* Property) const
 {
 	static const FName DisableArrayDistributionName("DisableArrayDistribution");
@@ -896,6 +928,46 @@ const UEnum* FNiagaraDistributionBase::GetLookupValueModeEnum(const FProperty* P
 	return LookupEnum ? LookupEnum : StaticEnum<ENiagaraDistributionLookupValueMode>();
 }
 
+void FNiagaraDistributionBase::ForEachDistribution(const UStruct* ObjectStruct, const void* Object, TFunction<void(const FNiagaraDistributionBase&)> Delegate)
+{
+	for (TFieldIterator<FStructProperty> PropIt(ObjectStruct); PropIt; ++PropIt)
+	{
+		if (FStructProperty* StructProp = *PropIt)
+		{
+			if (StructProp->Struct && StructProp->Struct->IsChildOf(FNiagaraDistributionBase::StaticStruct()))
+			{
+				const FNiagaraDistributionBase* Distribution = StructProp->ContainerPtrToValuePtr<const FNiagaraDistributionBase>(Object);
+				Delegate(*Distribution);
+			}
+		}
+	}
+}
+
+void FNiagaraDistributionBase::ForEachDistribution(const UStruct* ObjectStruct, void* Object, TFunction<void(FNiagaraDistributionBase&)> Delegate)
+{
+	for (TFieldIterator<FStructProperty> PropIt(ObjectStruct); PropIt; ++PropIt)
+	{
+		if (FStructProperty* StructProp = *PropIt)
+		{
+			if (StructProp->Struct && StructProp->Struct->IsChildOf(FNiagaraDistributionBase::StaticStruct()))
+			{
+				FNiagaraDistributionBase* Distribution = StructProp->ContainerPtrToValuePtr<FNiagaraDistributionBase>(Object);
+				Delegate(*Distribution);
+			}
+		}
+	}
+}
+
+void FNiagaraDistributionBase::ForEachDistribution(const UObject* InObject, TFunction<void(const FNiagaraDistributionBase&)> Delegate)
+{
+	ForEachDistribution(InObject->GetClass(), InObject, Delegate);
+}
+
+void FNiagaraDistributionBase::ForEachDistribution(UObject* InObject, TFunction<void(FNiagaraDistributionBase&)> Delegate)
+{
+	ForEachDistribution(InObject->GetClass(), InObject, Delegate);
+}
+
 void FNiagaraDistributionBase::PostEditChangeProperty(UObject* OwnerObject, FPropertyChangedEvent& PropertyChangedEvent)
 {
 	FStructProperty* StructProperty = CastField<FStructProperty>(PropertyChangedEvent.Property);
diff --git a/Engine/Plugins/FX/Niagara/Source/Niagara/Public/Stateless/NiagaraStatelessDistribution.h b/Engine/Plugins/FX/Niagara/Source/Niagara/Public/Stateless/NiagaraStatelessDistribution.h
--- a/Engine/Plugins/FX/Niagara/Source/Niagara/Public/Stateless/NiagaraStatelessDistribution.h
+++ b/Engine/Plugins/FX/Niagara/Source/Niagara/Public/Stateless/NiagaraStatelessDistribution.h
@@ -136,13 +136,15 @@ struct FNiagaraDistributionBase
 	NIAGARA_API bool operator==(const FNiagaraDistributionBase& Other) const;
 
 	void ForEachParameterBinding(TFunction<void(const FNiagaraVariableBase&)> Delegate) const;
+	NIAGARA_API static void ForEachParameterBinding(const UStruct* ObjectStruct, const void* Object, TFunction<void(const FNiagaraVariableBase&)> Delegate);
+	NIAGARA_API static void ForEachParameterBinding(const UObject* InObject, TFunction<void(const FNiagaraVariableBase&)> Delegate);
 
 	NIAGARA_API virtual bool AllowDistributionMode(ENiagaraDistributionMode Mode, const FProperty* Property) const;
 protected:
 	NIAGARA_API virtual bool InternalAllowRangeDistribution(ENiagaraDistributionMode Mode, const FProperty* Property) const;
 	NIAGARA_API virtual bool InternalAllowCurveDistribution(ENiagaraDistributionMode Mode, const FProperty* Property) const;
-public:
 
+public:
 	NIAGARA_API bool AllowAddressMode(const FProperty* Property) const;
 	NIAGARA_API bool AllowInterpolationMode(const FProperty* Property) const;
 	NIAGARA_API bool AllowLookupValueMode(const FProperty* Property) const;
@@ -157,6 +159,12 @@ public:
 	static void PostEditChangeProperty(UObject* OwnerObject, FPropertyChangedEvent& PropertyChangedEvent);
 
 	NIAGARA_API static const UEnum* GetLookupValueModeEnum(const FProperty* Property);
+
+	NIAGARA_API static void ForEachDistribution(const UStruct* ObjectStruct, const void* Object, TFunction<void(const FNiagaraDistributionBase&)> Delegate);
+	NIAGARA_API static void ForEachDistribution(const UStruct* ObjectStruct, void* Object, TFunction<void(FNiagaraDistributionBase&)> Delegate);
+
+	NIAGARA_API static void ForEachDistribution(const UObject* InObject, TFunction<void(const FNiagaraDistributionBase&)> Delegate);
+	NIAGARA_API static void ForEachDistribution(UObject* InObject, TFunction<void(FNiagaraDistributionBase&)> Delegate);
 #endif
 };
 
diff --git a/Engine/Plugins/FX/Niagara/Source/NiagaraEditor/Private/ViewModels/Stack/NiagaraStackStatelessEmitterSimulateGroup.cpp b/Engine/Plugins/FX/Niagara/Source/NiagaraEditor/Private/ViewModels/Stack/NiagaraStackStatelessEmitterSimulateGroup.cpp
--- a/Engine/Plugins/FX/Niagara/Source/NiagaraEditor/Private/ViewModels/Stack/NiagaraStackStatelessEmitterSimulateGroup.cpp
+++ b/Engine/Plugins/FX/Niagara/Source/NiagaraEditor/Private/ViewModels/Stack/NiagaraStackStatelessEmitterSimulateGroup.cpp
@@ -14,6 +14,7 @@
 #include "Stateless/NiagaraStatelessCommon.h"
 #include "Stateless/NiagaraStatelessEmitter.h"
 #include "Stateless/NiagaraStatelessModule.h"
+#include "ViewModels/NiagaraParameterPanelViewModel.h"
 #include "ViewModels/NiagaraSystemSelectionViewModel.h"
 #include "ViewModels/NiagaraSystemViewModel.h"
 #include "ViewModels/Stack/NiagaraStackItemPropertyHeaderValueShared.h"
@@ -482,6 +483,32 @@ void UNiagaraStackStatelessModuleItem::GetHeaderValueHandlers(TArray<TSharedRef<
 	}
 }
 
+const UNiagaraStackEntry::FCollectedUsageData& UNiagaraStackStatelessModuleItem::GetCollectedUsageData() const
+{
+	if (CachedCollectedUsageData.IsSet() == false)
+	{
+		CachedCollectedUsageData = FCollectedUsageData();
+
+		INiagaraParameterPanelViewModel* ParamVM = GetSystemViewModel()->GetParameterPanelViewModel();
+		UNiagaraStatelessModule* StatelessModule = StatelessModuleWeak.Get();
+		if ( ParamVM && StatelessModule )
+		{
+			FNiagaraDistributionBase::ForEachParameterBinding(
+				StatelessModule,
+				[this, &ParamVM](const FNiagaraVariableBase& Variable)
+				{
+					if (ParamVM->IsVariableSelected(Variable))
+					{
+						CachedCollectedUsageData.GetValue().bHasReferencedParameterRead = true;
+					}
+				}
+			);
+		}
+	}
+
+	return CachedCollectedUsageData.GetValue();
+}
+
 void UNiagaraStackStatelessModuleItem::FilterDetailNodes(const TArray<TSharedRef<IDetailTreeNode>>& InSourceNodes, TArray<TSharedRef<IDetailTreeNode>>& OutFilteredNodes)
 {
 	for (const TSharedRef<IDetailTreeNode>& SourceNode : InSourceNodes)
diff --git a/Engine/Plugins/FX/Niagara/Source/NiagaraEditor/Private/ViewModels/Stack/NiagaraStackStatelessEmitterSpawnGroup.cpp b/Engine/Plugins/FX/Niagara/Source/NiagaraEditor/Private/ViewModels/Stack/NiagaraStackStatelessEmitterSpawnGroup.cpp
--- a/Engine/Plugins/FX/Niagara/Source/NiagaraEditor/Private/ViewModels/Stack/NiagaraStackStatelessEmitterSpawnGroup.cpp
+++ b/Engine/Plugins/FX/Niagara/Source/NiagaraEditor/Private/ViewModels/Stack/NiagaraStackStatelessEmitterSpawnGroup.cpp
@@ -12,6 +12,7 @@
 #include "Stateless/NiagaraStatelessEmitter.h"
 #include "Stateless/NiagaraStatelessModule.h"
 #include "Styling/AppStyle.h"
+#include "ViewModels/NiagaraParameterPanelViewModel.h"
 #include "ViewModels/NiagaraSystemSelectionViewModel.h"
 #include "ViewModels/NiagaraSystemViewModel.h"
 #include "ViewModels/Stack/NiagaraStackItemPropertyHeaderValueShared.h"
@@ -382,6 +383,32 @@ void UNiagaraStackStatelessEmitterSpawnItem::GetHeaderValueHandlers(TArray<TShar
 	}
 }
 
+const UNiagaraStackEntry::FCollectedUsageData& UNiagaraStackStatelessEmitterSpawnItem::GetCollectedUsageData() const
+{
+	if (CachedCollectedUsageData.IsSet() == false)
+	{
+		CachedCollectedUsageData = FCollectedUsageData();
+
+		INiagaraParameterPanelViewModel* ParamVM = GetSystemViewModel()->GetParameterPanelViewModel();
+		FNiagaraStatelessSpawnInfo* SpawnInfo = GetSpawnInfo();
+		if (SpawnInfo && ParamVM)
+		{
+			FNiagaraDistributionBase::ForEachParameterBinding(
+				FNiagaraStatelessSpawnInfo::StaticStruct(),
+				SpawnInfo,
+				[this, &ParamVM](const FNiagaraVariableBase& Variable)
+				{
+					if (ParamVM->IsVariableSelected(Variable))
+					{
+						CachedCollectedUsageData.GetValue().bHasReferencedParameterRead = true;
+					}
+				}
+			);
+		}
+	}
+
+	return CachedCollectedUsageData.GetValue();
+}
 void UNiagaraStackStatelessEmitterSpawnItem::RefreshChildrenInternal(const TArray<UNiagaraStackEntry*>& CurrentChildren, TArray<UNiagaraStackEntry*>& NewChildren, TArray<FStackIssue>& NewIssues)
 {
 	Super::RefreshChildrenInternal(CurrentChildren, NewChildren, NewIssues);
diff --git a/Engine/Plugins/FX/Niagara/Source/NiagaraEditor/Public/ViewModels/Stack/NiagaraStackStatelessEmitterSimulateGroup.h b/Engine/Plugins/FX/Niagara/Source/NiagaraEditor/Public/ViewModels/Stack/NiagaraStackStatelessEmitterSimulateGroup.h
--- a/Engine/Plugins/FX/Niagara/Source/NiagaraEditor/Public/ViewModels/Stack/NiagaraStackStatelessEmitterSimulateGroup.h
+++ b/Engine/Plugins/FX/Niagara/Source/NiagaraEditor/Public/ViewModels/Stack/NiagaraStackStatelessEmitterSimulateGroup.h
@@ -80,6 +80,8 @@ public:
 	virtual bool SupportsHeaderValues() const override { return true; }
 	virtual void GetHeaderValueHandlers(TArray<TSharedRef<INiagaraStackItemHeaderValueHandler>>& OutHeaderValueHandlers) const override;
 
+	virtual const FCollectedUsageData& GetCollectedUsageData() const override;
+
 	UNiagaraStatelessModule* GetStatelessModule() const { return StatelessModuleWeak.Get(); }
 	UNiagaraStatelessEmitter* GetStatelessEmitter() const;
 
diff --git a/Engine/Plugins/FX/Niagara/Source/NiagaraEditor/Public/ViewModels/Stack/NiagaraStackStatelessEmitterSpawnGroup.h b/Engine/Plugins/FX/Niagara/Source/NiagaraEditor/Public/ViewModels/Stack/NiagaraStackStatelessEmitterSpawnGroup.h
--- a/Engine/Plugins/FX/Niagara/Source/NiagaraEditor/Public/ViewModels/Stack/NiagaraStackStatelessEmitterSpawnGroup.h
+++ b/Engine/Plugins/FX/Niagara/Source/NiagaraEditor/Public/ViewModels/Stack/NiagaraStackStatelessEmitterSpawnGroup.h
@@ -83,6 +83,8 @@ public:
 	virtual bool SupportsHeaderValues() const override { return true; }
 	virtual void GetHeaderValueHandlers(TArray<TSharedRef<INiagaraStackItemHeaderValueHandler>>& OutHeaderValueHandlers) const override;
 
+	virtual const FCollectedUsageData& GetCollectedUsageData() const override;
+
 	UNiagaraStatelessEmitter* GetStatelessEmitter() const { return StatelessEmitterWeak.Get(); }
 	int32 GetIndex() const { return Index; }
 	const FGuid& GetSourceId() const { return SourceId; }
