Removing stray backslash.

diff --git a/Engine/Plugins/FX/Niagara/Source/NiagaraShader/Private/NiagaraAsyncGpuTraceProviderHwrt.cpp b/Engine/Plugins/FX/Niagara/Source/NiagaraShader/Private/NiagaraAsyncGpuTraceProviderHwrt.cpp
--- a/Engine/Plugins/FX/Niagara/Source/NiagaraShader/Private/NiagaraAsyncGpuTraceProviderHwrt.cpp
+++ b/Engine/Plugins/FX/Niagara/Source/NiagaraShader/Private/NiagaraAsyncGpuTraceProviderHwrt.cpp
@@ -170,7 +170,7 @@ void FNiagaraCollisionRayTraceCS::ModifyCompilationEnvironment(const FGlobalShad
 	OutEnvironment.CompilerFlags.Add(CFLAG_InlineRayTracing);
 	OutEnvironment.CompilerFlags.Add(CFLAG_Wave32);
 
-	OutEnvironment.SetDefine(TEXT("RAY_TRACING_THREAD_GROUP_SIZE_X"), FNiagaraCollisionRayTraceCS::DispatchSize); \
+	OutEnvironment.SetDefine(TEXT("RAY_TRACING_THREAD_GROUP_SIZE_X"), FNiagaraCollisionRayTraceCS::DispatchSize);
 	OutEnvironment.SetDefine(TEXT("NIAGARA_SUPPORTS_RAY_TRACING"), 1);
 	OutEnvironment.SetDefine(TEXT("VF_SUPPORTS_PRIMITIVE_SCENE_DATA"), 1);
 	OutEnvironment.SetDefine(TEXT("NIAGARA_RAYTRACE_FAKE_INDIRECT"), 0);
