- Add a way to release resources from render target 2D
- Change handling of null texture to set the material default to avoid dangling references

#tests Made a system which will swap between an allocated RT and released.  Bind to sim stage to write and render to display.

diff --git a/Engine/Plugins/FX/Niagara/Source/Niagara/Classes/NiagaraDataInterfaceRenderTarget2D.h b/Engine/Plugins/FX/Niagara/Source/Niagara/Classes/NiagaraDataInterfaceRenderTarget2D.h
--- a/Engine/Plugins/FX/Niagara/Source/Niagara/Classes/NiagaraDataInterfaceRenderTarget2D.h
+++ b/Engine/Plugins/FX/Niagara/Source/Niagara/Classes/NiagaraDataInterfaceRenderTarget2D.h
@@ -17,6 +17,7 @@ class UTextureRenderTarget2D;
 
 struct FRenderTarget2DRWInstanceData_GameThread
 {
+	bool bReleaseResources = false;
 	FIntPoint Size = FIntPoint(EForceInit::ForceInitToZero);
 	ETextureRenderTargetFormat Format = RTF_RGBA16f;
 	TextureFilter Filter = TextureFilter::TF_Default;
@@ -135,6 +136,7 @@ public:
 
 	NIAGARA_API void VMGetSize(FVectorVMExternalFunctionContext& Context);
 	NIAGARA_API void VMSetSize(FVectorVMExternalFunctionContext& Context);
+	NIAGARA_API void VMReleaseResources(FVectorVMExternalFunctionContext& Context);
 	NIAGARA_API void VMGetNumMipLevels(FVectorVMExternalFunctionContext& Context);
 	NIAGARA_API void VMSetFormat(FVectorVMExternalFunctionContext& Context);
 
diff --git a/Engine/Plugins/FX/Niagara/Source/Niagara/Private/NiagaraDataInterfaceRenderTarget2D.cpp b/Engine/Plugins/FX/Niagara/Source/Niagara/Private/NiagaraDataInterfaceRenderTarget2D.cpp
--- a/Engine/Plugins/FX/Niagara/Source/Niagara/Private/NiagaraDataInterfaceRenderTarget2D.cpp
+++ b/Engine/Plugins/FX/Niagara/Source/Niagara/Private/NiagaraDataInterfaceRenderTarget2D.cpp
@@ -50,6 +50,7 @@ namespace NDIRenderTarget2DLocal
 	static const FName SampleValueFunctionName("SampleRenderTargetValue");
 	static const FName SetSizeFunctionName("SetRenderTargetSize");
 	static const FName GetSizeFunctionName("GetRenderTargetSize");
+	static const FName ReleaseResourcesFunctionName("ReleaseResources");
 	static const FName GetNumMipLevelsName("GetNumMipLevels");
 	static const FName SetFormatFunctionName("SetRenderTargetFormat");
 	static const FName LinearToIndexName("LinearToIndex");
@@ -156,6 +157,15 @@ void UNiagaraDataInterfaceRenderTarget2D::GetFunctionsInternal(TArray<FNiagaraFu
 		Sig.bSupportsGPU = false;
 	}
 
+	{
+		FNiagaraFunctionSignature& Sig = OutFunctions.Add_GetRef(DefaultSig);
+		Sig.Name = ReleaseResourcesFunctionName;
+		Sig.Inputs.Add_GetRef(FNiagaraVariable(FNiagaraTypeDefinition::GetBoolDef(), TEXT("Execute"))).SetValue(true);
+		Sig.ModuleUsageBitmask = EmitterSystemOnlyBitmask;
+		Sig.bRequiresExecPin = true;
+		Sig.bSupportsGPU = false;
+	}
+
 	{
 		FNiagaraFunctionSignature& Sig = OutFunctions.Add_GetRef(DefaultSig);
 		Sig.Name = GetNumMipLevelsName;
@@ -292,6 +302,10 @@ void UNiagaraDataInterfaceRenderTarget2D::GetVMExternalFunction(const FVMExterna
 	{
 		OutFunc = FVMExternalFunction::CreateUObject(this, &UNiagaraDataInterfaceRenderTarget2D::VMSetSize);
 	}
+	else if (BindingInfo.Name == ReleaseResourcesFunctionName)
+	{
+		OutFunc = FVMExternalFunction::CreateUObject(this, &UNiagaraDataInterfaceRenderTarget2D::VMReleaseResources);
+	}
 	else if (BindingInfo.Name == GetNumMipLevelsName)
 	{
 		OutFunc = FVMExternalFunction::CreateUObject(this, &UNiagaraDataInterfaceRenderTarget2D::VMGetNumMipLevels);
@@ -755,12 +769,31 @@ void UNiagaraDataInterfaceRenderTarget2D::VMSetSize(FVectorVMExternalFunctionCon
 		OutSuccess.SetAndAdvance(bSuccess);
 		if (bSuccess)
 		{
+			InstData->bReleaseResources = false;
 			InstData->Size.X = FMath::Clamp<int>(int(float(SizeX) * NiagaraDataInterfaceRenderTargetCommon::GResolutionMultiplier), 1, GMaxTextureDimensions);
 			InstData->Size.Y = FMath::Clamp<int>(int(float(SizeY) * NiagaraDataInterfaceRenderTargetCommon::GResolutionMultiplier), 1, GMaxTextureDimensions);
 		}
 	}
 }
 
+void UNiagaraDataInterfaceRenderTarget2D::VMReleaseResources(FVectorVMExternalFunctionContext& Context)
+{
+	// This should only be called from a system or emitter script due to a need for only setting up initially.
+	VectorVM::FUserPtrHandler<FRenderTarget2DRWInstanceData_GameThread> InstData(Context);
+	VectorVM::FExternalFuncInputHandler<FNiagaraBool> InExecute(Context);
+
+	for (int32 InstanceIdx = 0; InstanceIdx < Context.GetNumInstances(); ++InstanceIdx)
+	{
+		const bool bExecute = InExecute.GetAndAdvance();
+		if (bExecute)
+		{
+			InstData->bReleaseResources = true;
+			InstData->Size.X = 1;
+			InstData->Size.Y = 1;
+		}
+	}
+}
+
 void UNiagaraDataInterfaceRenderTarget2D::VMGetSize(FVectorVMExternalFunctionContext& Context)
 {
 	VectorVM::FUserPtrHandler<FRenderTarget2DRWInstanceData_GameThread> InstData(Context);
@@ -824,39 +857,41 @@ bool UNiagaraDataInterfaceRenderTarget2D::PerInstanceTick(void* PerInstanceData,
 	using namespace NDIRenderTarget2DLocal;
 
 	FRenderTarget2DRWInstanceData_GameThread* InstanceData = static_cast<FRenderTarget2DRWInstanceData_GameThread*>(PerInstanceData);
-
-	// Pull from user parameter
-	UTextureRenderTarget2D* UserTargetTexture = InstanceData->RTUserParamBinding.GetValue<UTextureRenderTarget2D>();
-	if (UserTargetTexture && (InstanceData->TargetTexture != UserTargetTexture))
+	if (InstanceData->bReleaseResources == false)
 	{
-		NiagaraDataInterfaceRenderTargetCommon::ReleaseRenderTarget(SystemInstance, InstanceData);
-		InstanceData->TargetTexture = UserTargetTexture;
-	}
+		// Pull from user parameter
+		UTextureRenderTarget2D* UserTargetTexture = InstanceData->RTUserParamBinding.GetValue<UTextureRenderTarget2D>();
+		if (UserTargetTexture && (InstanceData->TargetTexture != UserTargetTexture))
+		{
+			NiagaraDataInterfaceRenderTargetCommon::ReleaseRenderTarget(SystemInstance, InstanceData);
+			InstanceData->TargetTexture = UserTargetTexture;
+		}
 
-	// Do we inherit the texture parameters from the user supplied texture?
-	if (bInheritUserParameterSettings)
-	{
-		if (UserTargetTexture)
+		// Do we inherit the texture parameters from the user supplied texture?
+		if (bInheritUserParameterSettings)
 		{
-			InstanceData->Size.X = UserTargetTexture->SizeX;
-			InstanceData->Size.Y = UserTargetTexture->SizeY;
-			if (UserTargetTexture->bAutoGenerateMips)
+			if (UserTargetTexture)
 			{
-				// We have to take a guess at user intention
-				InstanceData->MipMapGeneration = MipMapGeneration == ENiagaraMipMapGeneration::Disabled ? ENiagaraMipMapGeneration::PostStage : MipMapGeneration;
-				InstanceData->MipMapGenerationType = MipMapGenerationType;
+				InstanceData->Size.X = UserTargetTexture->SizeX;
+				InstanceData->Size.Y = UserTargetTexture->SizeY;
+				if (UserTargetTexture->bAutoGenerateMips)
+				{
+					// We have to take a guess at user intention
+					InstanceData->MipMapGeneration = MipMapGeneration == ENiagaraMipMapGeneration::Disabled ? ENiagaraMipMapGeneration::PostStage : MipMapGeneration;
+					InstanceData->MipMapGenerationType = MipMapGenerationType;
+				}
+				else
+				{
+					InstanceData->MipMapGeneration = ENiagaraMipMapGeneration::Disabled;
+					InstanceData->MipMapGenerationType = ENiagaraMipMapGenerationType::Unfiltered;
+				}
+				InstanceData->Format = InstanceData->TargetTexture->RenderTargetFormat;
+				InstanceData->Filter = InstanceData->TargetTexture->Filter;
 			}
 			else
 			{
-				InstanceData->MipMapGeneration = ENiagaraMipMapGeneration::Disabled;
-				InstanceData->MipMapGenerationType = ENiagaraMipMapGenerationType::Unfiltered;
+				UE_LOG(LogNiagara, Error, TEXT("RenderTarget UserParam is required but invalid."));
 			}
-			InstanceData->Format = InstanceData->TargetTexture->RenderTargetFormat;
-			InstanceData->Filter = InstanceData->TargetTexture->Filter;
-		}
-		else
-		{
-			UE_LOG(LogNiagara, Error, TEXT("RenderTarget UserParam is required but invalid."));
 		}
 	}
 
@@ -876,47 +911,57 @@ bool UNiagaraDataInterfaceRenderTarget2D::PerInstanceTickPostSimulate(void* PerI
 	//-TEMP: Until we prune data interface on cook this will avoid consuming memory
 	const bool bValidGpuDataInterface = NiagaraDataInterfaceRenderTargetCommon::GIgnoreCookedOut == 0 || IsUsedWithGPUScript();
 
-	if (::IsValid(InstanceData->TargetTexture) == false)
+	if (InstanceData->bReleaseResources)
 	{
-		InstanceData->TargetTexture = nullptr;
+		if (InstanceData->TargetTexture)
+		{
+			NiagaraDataInterfaceRenderTargetCommon::ReleaseRenderTarget(SystemInstance, InstanceData);
+		}
 	}
-
-	// Do we need to create a new texture?
-	if (bValidGpuDataInterface && !bInheritUserParameterSettings && (InstanceData->TargetTexture == nullptr))
+	else
 	{
-		if (NiagaraDataInterfaceRenderTargetCommon::CreateRenderTarget<UTextureRenderTarget2D>(SystemInstance, InstanceData) == false)
+		if (::IsValid(InstanceData->TargetTexture) == false)
 		{
-			return false;
+			InstanceData->TargetTexture = nullptr;
 		}
-		check(InstanceData->TargetTexture);
-		InstanceData->TargetTexture->bCanCreateUAV = true;
-		InstanceData->TargetTexture->bAutoGenerateMips = InstanceData->MipMapGeneration != ENiagaraMipMapGeneration::Disabled;
-		InstanceData->TargetTexture->RenderTargetFormat = InstanceData->Format;
-		InstanceData->TargetTexture->ClearColor = FLinearColor(0.0, 0, 0, 0);
-		InstanceData->TargetTexture->Filter = InstanceData->Filter;
-		InstanceData->TargetTexture->InitAutoFormat(InstanceData->Size.X, InstanceData->Size.Y);
-		InstanceData->TargetTexture->UpdateResourceImmediate(true);
-	}
 
-	// Do we need to update the existing texture?
-	if (InstanceData->TargetTexture)
-	{
-		const bool bAutoGenerateMips = InstanceData->MipMapGeneration != ENiagaraMipMapGeneration::Disabled;
-		if ((InstanceData->TargetTexture->SizeX != InstanceData->Size.X) || (InstanceData->TargetTexture->SizeY != InstanceData->Size.Y) ||
-			(InstanceData->TargetTexture->RenderTargetFormat != InstanceData->Format) ||
-			(InstanceData->TargetTexture->Filter != InstanceData->Filter) ||
-			!InstanceData->TargetTexture->bCanCreateUAV ||
-			(InstanceData->TargetTexture->bAutoGenerateMips != bAutoGenerateMips) ||
-			!InstanceData->TargetTexture->GetResource())
+		// Do we need to create a new texture?
+		if (bValidGpuDataInterface && !bInheritUserParameterSettings && (InstanceData->TargetTexture == nullptr))
 		{
-			// resize RT to match what we need for the output
+			if (NiagaraDataInterfaceRenderTargetCommon::CreateRenderTarget<UTextureRenderTarget2D>(SystemInstance, InstanceData) == false)
+			{
+				return false;
+			}
+			check(InstanceData->TargetTexture);
 			InstanceData->TargetTexture->bCanCreateUAV = true;
-			InstanceData->TargetTexture->bAutoGenerateMips = bAutoGenerateMips;
+			InstanceData->TargetTexture->bAutoGenerateMips = InstanceData->MipMapGeneration != ENiagaraMipMapGeneration::Disabled;
 			InstanceData->TargetTexture->RenderTargetFormat = InstanceData->Format;
+			InstanceData->TargetTexture->ClearColor = FLinearColor(0.0, 0, 0, 0);
 			InstanceData->TargetTexture->Filter = InstanceData->Filter;
 			InstanceData->TargetTexture->InitAutoFormat(InstanceData->Size.X, InstanceData->Size.Y);
 			InstanceData->TargetTexture->UpdateResourceImmediate(true);
 		}
+
+		// Do we need to update the existing texture?
+		if (InstanceData->TargetTexture)
+		{
+			const bool bAutoGenerateMips = InstanceData->MipMapGeneration != ENiagaraMipMapGeneration::Disabled;
+			if ((InstanceData->TargetTexture->SizeX != InstanceData->Size.X) || (InstanceData->TargetTexture->SizeY != InstanceData->Size.Y) ||
+				(InstanceData->TargetTexture->RenderTargetFormat != InstanceData->Format) ||
+				(InstanceData->TargetTexture->Filter != InstanceData->Filter) ||
+				!InstanceData->TargetTexture->bCanCreateUAV ||
+				(InstanceData->TargetTexture->bAutoGenerateMips != bAutoGenerateMips) ||
+				!InstanceData->TargetTexture->GetResource())
+			{
+				// resize RT to match what we need for the output
+				InstanceData->TargetTexture->bCanCreateUAV = true;
+				InstanceData->TargetTexture->bAutoGenerateMips = bAutoGenerateMips;
+				InstanceData->TargetTexture->RenderTargetFormat = InstanceData->Format;
+				InstanceData->TargetTexture->Filter = InstanceData->Filter;
+				InstanceData->TargetTexture->InitAutoFormat(InstanceData->Size.X, InstanceData->Size.Y);
+				InstanceData->TargetTexture->UpdateResourceImmediate(true);
+			}
+		}
 	}
 
 	//-TODO: We could avoid updating each frame if we cache the resource pointer or a serial number
@@ -925,7 +970,7 @@ bool UNiagaraDataInterfaceRenderTarget2D::PerInstanceTickPostSimulate(void* PerI
 	{
 		FNiagaraDataInterfaceProxyRenderTarget2DProxy* RT_Proxy = GetProxyAs<FNiagaraDataInterfaceProxyRenderTarget2DProxy>();
 		FTextureRenderTargetResource* RT_TargetTexture = InstanceData->TargetTexture ? InstanceData->TargetTexture->GameThread_GetRenderTargetResource() : nullptr;
-		ensureMsgf(RT_TargetTexture != nullptr, TEXT("NiagaraDataInterfaceRenderTarget2D - null RT for system (%s)"), *GetNameSafe(SystemInstance->GetSystem()));
+		ensureMsgf(InstanceData->bReleaseResources || RT_TargetTexture != nullptr, TEXT("NiagaraDataInterfaceRenderTarget2D - null RT for system (%s)"), *GetNameSafe(SystemInstance->GetSystem()));
 		ENQUEUE_RENDER_COMMAND(NDIRenderTarget2DUpdate)
 		(
 			[RT_Proxy, RT_InstanceID=SystemInstance->GetId(), RT_InstanceData=*InstanceData, RT_TargetTexture](FRHICommandListImmediate& RHICmdList)
diff --git a/Engine/Plugins/FX/Niagara/Source/Niagara/Private/NiagaraRenderer.cpp b/Engine/Plugins/FX/Niagara/Source/Niagara/Private/NiagaraRenderer.cpp
--- a/Engine/Plugins/FX/Niagara/Source/Niagara/Private/NiagaraRenderer.cpp
+++ b/Engine/Plugins/FX/Niagara/Source/Niagara/Private/NiagaraRenderer.cpp
@@ -568,13 +568,17 @@ void FNiagaraRenderer::ProcessMaterialParameterBindings(const FNiagaraRendererMa
 						{
 							UObject* Var = nullptr;
 							InEmitter->GetBoundRendererValue_GT(Binding.GetParamMapBindableVariable(), Binding.NiagaraChildVariable, &Var);
-							if (Var)
+							UTexture* Tex = Cast<UTexture>(Var);
+							if (Tex && Tex->GetResource())
 							{
-								UTexture* Tex = Cast<UTexture>(Var);
-								if (Tex && Tex->GetResource() != nullptr)
-								{
-									MatDyn->SetTextureParameterValue(Binding.MaterialParameterName, Tex);
-								}
+								MatDyn->SetTextureParameterValue(Binding.MaterialParameterName, Tex);
+							}
+							// Not a valid texture?  Setting nullptr will not remove the override and set to default.
+							// We have to do this manually getting the parent material default.
+							else if ( const UMaterial* ParentMaterial = MatDyn->GetMaterial() )
+							{
+								ParentMaterial->GetTextureParameterValue(Binding.MaterialParameterName, Tex, false);
+								MatDyn->SetTextureParameterValue(Binding.MaterialParameterName, Tex);
 							}
 						}
 					}
