Remove slate dependency from emitter merging

diff --git a/Engine/Plugins/FX/Niagara/Source/NiagaraEditor/NiagaraEditor.Build.cs b/Engine/Plugins/FX/Niagara/Source/NiagaraEditor/NiagaraEditor.Build.cs
--- a/Engine/Plugins/FX/Niagara/Source/NiagaraEditor/NiagaraEditor.Build.cs
+++ b/Engine/Plugins/FX/Niagara/Source/NiagaraEditor/NiagaraEditor.Build.cs
@@ -20,7 +20,6 @@ public class NiagaraEditor : ModuleRules
 				"Json",
 				"Slate", 
 				"SlateCore",
-				"SlateNullRenderer",
 				"Kismet",
                 "EditorStyle",
 				"VectorVM",
diff --git a/Engine/Plugins/FX/Niagara/Source/NiagaraEditor/Private/NiagaraScriptMergeManager.cpp b/Engine/Plugins/FX/Niagara/Source/NiagaraEditor/Private/NiagaraScriptMergeManager.cpp
--- a/Engine/Plugins/FX/Niagara/Source/NiagaraEditor/Private/NiagaraScriptMergeManager.cpp
+++ b/Engine/Plugins/FX/Niagara/Source/NiagaraEditor/Private/NiagaraScriptMergeManager.cpp
@@ -24,8 +24,6 @@
 #include "NiagaraEditorModule.h"
 #include "NiagaraEmitterEditorData.h"
 #include "NiagaraStackEditorData.h"
-#include "Framework/Application/SlateApplication.h"
-#include "Interfaces/ISlateNullRendererModule.h"
 
 #include "UObject/PropertyPortFlags.h"
 
@@ -1434,14 +1432,7 @@ void FNiagaraScriptMergeManager::UpdateModuleVersions(const FVersionedNiagaraEmi
 	if (ChangedVersionOtherModules.Num() > 0)
 	{
 		TSharedPtr<FNiagaraSystemViewModel> SystemViewModel;
-
-		if (!FSlateApplication::IsInitialized())
-		{
-			// it's possible that during the cook process we don't have slate initialized, but we need it set up for the view model
-			FSlateApplication::Create();
-			TSharedRef<FSlateRenderer> SlateRenderer = FModuleManager::Get().LoadModuleChecked<ISlateNullRendererModule>("SlateNullRenderer").CreateSlateNullRenderer();
-			FSlateApplication::Get().InitializeRenderer(SlateRenderer);
-		}
+		
 		System = NewObject<UNiagaraSystem>(GetTransientPackage(), NAME_None, RF_Transient | RF_Standalone);
 		UNiagaraSystemFactoryNew::InitializeSystem(System, true);
 		SystemViewModel = MakeShared<FNiagaraSystemViewModel>();
