Bit of refactoring classes out of NiagaraDataChannels.cpp/h
More can be done but this improves things somewhat.

diff --git a/Engine/Plugins/FX/Niagara/Source/Niagara/Internal/DataInterface/NiagaraDataInterfaceDataChannelRead.h b/Engine/Plugins/FX/Niagara/Source/Niagara/Internal/DataInterface/NiagaraDataInterfaceDataChannelRead.h
--- a/Engine/Plugins/FX/Niagara/Source/Niagara/Internal/DataInterface/NiagaraDataInterfaceDataChannelRead.h
+++ b/Engine/Plugins/FX/Niagara/Source/Niagara/Internal/DataInterface/NiagaraDataInterfaceDataChannelRead.h
@@ -1,9 +1,7 @@
 // Copyright Epic Games, Inc. All Rights Reserved.
 #pragma once
 
-
 /**
-
 Data Channel Read Interface.
 Enabled Niagara Systems to read data from a NiagaraDataChannel that has been previously generated by another System or Game code/BP.
 
@@ -18,22 +16,16 @@ This avoids cumbersome work in the graph to access data but requires special han
 
 TODO: Local support - We can read directly from a writer DI in the same system to avoid having to publish that data in some cases. 
 There is working code for this already but the concept/API needs more fleshing out before we include it.
-
 */
 
 #include "Containers/DynamicRHIResourceArray.h"
-
-#include "NiagaraDataChannelPublic.h"
 #include "NiagaraDataInterfaceDataChannelCommon.h"
 #include "NiagaraDataInterfaceRW.h"
+#include "NiagaraDataChannelAccessContext.h"
+
 #include "NiagaraDataInterfaceDataChannelRead.generated.h"
 
-class UNiagaraDataInterfaceDataChannelWrite;
 struct FNDIDataChannelWriteInstanceData;
-class UNiagaraDataChannelHandler;
-class FNiagaraDataBuffer;
-
-using FNiagaraDataChannelDataProxyPtr = TSharedPtr<struct FNiagaraDataChannelDataProxy>;
 
 /** Mode controlling the behavior of the Spawn function for Niagara's Data Channel Read Data Interface.*/
 UENUM()
diff --git a/Engine/Plugins/FX/Niagara/Source/Niagara/Private/DataInterface/NiagaraDataInterfaceDataChannelRead.cpp b/Engine/Plugins/FX/Niagara/Source/Niagara/Private/DataInterface/NiagaraDataInterfaceDataChannelRead.cpp
--- a/Engine/Plugins/FX/Niagara/Source/Niagara/Private/DataInterface/NiagaraDataInterfaceDataChannelRead.cpp
+++ b/Engine/Plugins/FX/Niagara/Source/Niagara/Private/DataInterface/NiagaraDataInterfaceDataChannelRead.cpp
@@ -18,6 +18,7 @@
 #include "NiagaraDataChannelCommon.h"
 #include "NiagaraDataChannelHandler.h"
 #include "NiagaraDataChannelManager.h"
+#include "NiagaraDataChannelData.h"
 
 #include "NiagaraEmitterInstanceImpl.h"
 
diff --git a/Engine/Plugins/FX/Niagara/Source/Niagara/Private/DataInterface/NiagaraDataInterfaceDataChannelWrite.cpp b/Engine/Plugins/FX/Niagara/Source/Niagara/Private/DataInterface/NiagaraDataInterfaceDataChannelWrite.cpp
--- a/Engine/Plugins/FX/Niagara/Source/Niagara/Private/DataInterface/NiagaraDataInterfaceDataChannelWrite.cpp
+++ b/Engine/Plugins/FX/Niagara/Source/Niagara/Private/DataInterface/NiagaraDataInterfaceDataChannelWrite.cpp
@@ -19,6 +19,8 @@
 #include "NiagaraDataChannel.h"
 #include "NiagaraDataChannelHandler.h"
 #include "NiagaraDataChannelManager.h"
+#include "NiagaraDataChannelData.h"
+#include "NiagaraDataChannelGameData.h"
 
 #include "NiagaraRenderer.h"
 #include "NiagaraGPUSystemTick.h"
diff --git a/Engine/Plugins/FX/Niagara/Source/Niagara/Private/NiagaraDataChannel.cpp b/Engine/Plugins/FX/Niagara/Source/Niagara/Private/NiagaraDataChannel.cpp
--- a/Engine/Plugins/FX/Niagara/Source/Niagara/Private/NiagaraDataChannel.cpp
+++ b/Engine/Plugins/FX/Niagara/Source/Niagara/Private/NiagaraDataChannel.cpp
@@ -2,2342 +2,173 @@
 
 #include "NiagaraDataChannel.h"
 
-#include "Misc/LargeWorldRenderPosition.h"
 #include "NiagaraWorldManager.h"
 #include "NiagaraDataChannelCommon.h"
 #include "NiagaraDataChannelHandler.h"
 #include "NiagaraDataChannelManager.h"
-#include "Engine/Engine.h"
-#include "UObject/UObjectGlobals.h"
-#include "NiagaraDataChannelAccessor.h"
-#include "NiagaraGpuComputeDispatchInterface.h"
-#include "Misc/LazySingleton.h"
+#include "NiagaraDataChannelGameData.h"
 
-#if WITH_NIAGARA_DEBUGGER
-#include "Serialization/JsonWriter.h"
-#include "Serialization/JsonSerializer.h"
-#include "Dom/JsonObject.h"
-#endif
-#include "NiagaraDataSetReadback.h"
-
-#include "RenderGraphUtils.h"
-
-#include "Blueprint/BlueprintExceptionInfo.h"
-
-#include UE_INLINE_GENERATED_CPP_BY_NAME(NiagaraDataChannel)
-
-#define LOCTEXT_NAMESPACE "NiagaraDataChannels"
-
-namespace NDCCVars
-{
-	bool bEmitWarningsOnLateNDCWrites = !UE_BUILD_SHIPPING;
-	static FAutoConsoleVariableRef CVarEmitWarningsOnLateNDCWrites(TEXT("fx.Niagara.DataChannels.WarnOnLateWrites"), bEmitWarningsOnLateNDCWrites, TEXT("If true, late writes to NDCs will generate warnings. Late meaning after their final allowed tick group."), ECVF_Default);
-
-	int32 LogWritesToOutputLog = 0;
-	static FAutoConsoleVariableRef CVarLogWritesToOutputLog(TEXT("fx.Niagara.DataChannels.LogWritesToOutputLog"), LogWritesToOutputLog, TEXT("0=Disabled, 1=Log write summary, 2=Also write data; If >0, the NDC debugger will print all data channel writes to the output log."), ECVF_Default);
-
-	int32 FrameDataToCapture = 0;
-	static FAutoConsoleVariableRef CVarFrameDataToCapture(TEXT("fx.Niagara.DataChannels.FrameDataToCapture"), FrameDataToCapture, TEXT("The number of frames the debugger will capture for write requests."), ECVF_Default);
-
-#if WITH_NIAGARA_DEBUGGER
-	FAutoConsoleCommand CmdDumpWriteLog(
-		TEXT("fx.Niagara.DataChannels.DumpWriteLog"),
-		TEXT("Dump all the currently stored writes to the log (see fx.Niagara.DataChannels.FrameDataToCapture on how many frames are captured)"),
-		FConsoleCommandDelegate::CreateStatic(FNiagaraDataChannelDebugUtilities::DumpAllWritesToLog)
-	);
-#endif
-
-
-	bool AutoUploadGPUSpawnData = true;
-	static FAutoConsoleVariableRef CVarAutoUploadGPUSpawnData(TEXT("fx.Niagara.DataChannels.AutoUploadGPUSpawnData"), AutoUploadGPUSpawnData, TEXT("When true we will automatically upload any CPU NDC data to the GPU if it has been used to spawn GPU particles."), ECVF_Default);
-
-	bool EnableAccumulatedWrites = true;
-	static FAutoConsoleVariableRef CVarEnableAccumulatedWrites(TEXT("fx.Niagara.DataChannels.EnableAccumulatedWrites"), EnableAccumulatedWrites, TEXT("True if we attempt to accumulate writes into larger shared buffers."), ECVF_Default);
-};
-
-//////////////////////////////////////////////////////////////////////////
-
-void FNiagaraDataChannelGameDataLayout::Init(TConstArrayView<FNiagaraDataChannelVariable> Variables)
-{
-	VariableIndices.Reset();
-	LwcConverters.Reserve(Variables.Num());
-	for (const FNiagaraDataChannelVariable& Var : Variables)
-	{
-		//Sigh.
-		//We must convert from the variable stored var in the data channels definition as we currently cannot serialize/store actual LWC types in FNiagaraTypeDefinitions.
-		FNiagaraTypeDefinition LWCType = FNiagaraTypeHelper::GetLWCType(Var.GetType());
-		FNiagaraVariableBase LWCVar(LWCType, Var.GetName());
-
-		int32& VarIdx = VariableIndices.Add(LWCVar);
-		VarIdx = VariableIndices.Num() - 1;
-
-		FNiagaraLwcStructConverter& Converter = LwcConverters.AddDefaulted_GetRef();
-		Converter = FNiagaraTypeRegistry::GetStructConverter(LWCType);
-	}
-}
-
-//////////////////////////////////////////////////////////////////////////
-
-FNiagaraDataChannelGameData::FNiagaraDataChannelGameData(const FNiagaraDataChannelLayoutInfoPtr& InLayoutInfo)
-{
-	Init(InLayoutInfo);
-}
-
-void FNiagaraDataChannelGameData::Init(const FNiagaraDataChannelLayoutInfoPtr& InLayoutInfo)
-{
-	LayoutInfo = InLayoutInfo;
-
-	const FNiagaraDataChannelGameDataLayout& Layout = LayoutInfo->GetGameDataLayout();
-	VariableData.Empty();
-	VariableData.SetNum(Layout.VariableIndices.Num());
-	for(const auto& Pair : Layout.VariableIndices)
-	{
-		int32 Index = Pair.Value;
-		VariableData[Index].Init(Pair.Key);
-	}
-}
-
-void FNiagaraDataChannelGameData::Empty()
-{
-	NumElements = 0;
-	PrevNumElements = 0;
-	MaxElements = 0;
-	for (FNiagaraDataChannelVariableBuffer& VarData : VariableData)
-	{
-		VarData.Empty();
-	}
-}
-
-void FNiagaraDataChannelGameData::Reset()
-{
-	NumElements = 0;
-	PrevNumElements = 0;
-	for (FNiagaraDataChannelVariableBuffer& VarData : VariableData)
-	{
-		VarData.Reset();
-	}
-}
-
-void FNiagaraDataChannelGameData::BeginFrame()
-{
-	bool bKeepPrevious = LayoutInfo->KeepPreviousFrameData();
-	PrevNumElements = bKeepPrevious ? NumElements : 0;
-	NumElements = 0;
-	for (FNiagaraDataChannelVariableBuffer& VarData : VariableData)
-	{
-		VarData.BeginFrame(bKeepPrevious);
-	}
-}
-
-void FNiagaraDataChannelGameData::SetNum(int32 NewNum)
-{
-	NumElements = NewNum;
-	MaxElements = FMath::Max(NewNum, MaxElements);
-	for (auto& Buffer : VariableData)
-	{
-		Buffer.SetNum(NewNum);
-	}
-}
-
-void FNiagaraDataChannelGameData::Reserve(int32 NewNum)
-{
-	MaxElements = FMath::Max(NewNum, MaxElements);
-	for (auto& Buffer : VariableData)
-	{
-		Buffer.Reserve(NewNum);
-	}
-}
-
-FNiagaraDataChannelVariableBuffer* FNiagaraDataChannelGameData::FindVariableBuffer(const FNiagaraVariableBase& Var)
-{
-	const FNiagaraDataChannelGameDataLayout& Layout = LayoutInfo->GetGameDataLayout();
-	const FNiagaraTypeDefinition& VarType = Var.GetType();
-	for(auto& Pair : Layout.VariableIndices)
-	{
-		const FNiagaraVariableBase& LayoutVar = Pair.Key;
-		const FNiagaraTypeDefinition& LayoutVarType = LayoutVar.GetType();
-		int32 Index = Pair.Value;
-		
-		if(Var.GetName() == LayoutVar.GetName())
-		{
-			//For enum variables we'll hack things a little so that correctly named ints also match. This gets around some limitations in calling code not being able to provide the correct enum types.
-			if(VarType == LayoutVarType || (LayoutVarType.IsEnum() && VarType == FNiagaraTypeDefinition::GetIntDef()))
-			{
-				return &VariableData[Index];
-			}
-		}
-	}
-	return nullptr;
-}
-
-void FNiagaraDataChannelGameData::WriteToDataSet(FNiagaraDataBuffer* DestBuffer, int32 DestStartIdx, FVector3f SimulationLwcTile)
-{
-	const FNiagaraDataSetCompiledData& CompiledData = DestBuffer->GetOwner()->GetCompiledData();
-
-	static TArray<uint8> LWCConversionBuffer;
-
-	int32 NumInstances = NumElements;
-
-	if (NumInstances == 0)
-	{
-		return;
-	}
-	
-	DestBuffer->SetNumInstances(DestStartIdx + NumInstances);
-
-	const FNiagaraDataChannelGameDataLayout& Layout = LayoutInfo->GetGameDataLayout();
-
-	for (const TPair<FNiagaraVariableBase, int32>& VarIndexPair : Layout.VariableIndices)
-	{
-		FNiagaraVariableBase Var = VarIndexPair.Key;
-		int32 VarIndex = VarIndexPair.Value;
-		FNiagaraDataChannelVariableBuffer& VarBuffer = VariableData[VarIndex];
-		uint8* SrcDataBase = VarBuffer.Data.GetData();
-		
-		FNiagaraVariableBase SimVar = Var;
-
-		//Convert from LWC types to Niagara Simulation Types where required.
-		if (FNiagaraTypeHelper::IsLWCType(Var.GetType()))
-		{
-			SimVar.SetType(FNiagaraTypeDefinition(FNiagaraTypeHelper::FindNiagaraFriendlyTopLevelStruct(CastChecked<UScriptStruct>(Var.GetType().GetStruct()), ENiagaraStructConversion::Simulation)));
-		}
-
-		//Niagara Positions are a special case where we're actually storing them as FVectors in the game level data and must convert down to actual Positions/FVector3f in the sim data.
-		if(Var.GetType() == FNiagaraTypeDefinition::GetPositionDef())
-		{
-			Var.SetType(FNiagaraTypeHelper::GetVectorDef());
-		}
-
-		int32 SimVarIndex = CompiledData.Variables.IndexOfByKey(SimVar);
-		if (SimVarIndex == INDEX_NONE)
-		{
-			continue; //Did not find this variable in the dataset. Warn?
-		}
-
-		int32 SrcVarSize = Var.GetSizeInBytes();
-		int32 DestVarSize = SimVar.GetSizeInBytes();
-		const FNiagaraVariableLayoutInfo& SimLayout = CompiledData.VariableLayouts[SimVarIndex];
-
-		int32 FloatCompIdx = SimLayout.GetFloatComponentStart();
-		int32 IntCompIdx = SimLayout.GetInt32ComponentStart();
-		int32 HalfCompIdx = SimLayout.GetHalfComponentStart();
-
-		TFunction<void(UScriptStruct*, UScriptStruct*, uint8*)> WriteData;
-		WriteData = [&](UScriptStruct* SrcStruct, UScriptStruct* DestStruct, uint8* SrcPropertyBase)
-		{
-			//TODO: why doesn't this just use FNiagaraLwcStructConverter?
-			
-			//Write all data from the data channel into the data set. Converting into from LWC into the local LWC Tile space as we go.
-
-			uint8* SrcData = SrcPropertyBase;
-
-			//Positions are a special case that are stored as FVectors in game data but converted to an LWCTile local FVector3f in simulation data.
-			if (DestStruct == FNiagaraTypeDefinition::GetPositionStruct())
-			{	
-				float* DestX = DestBuffer->GetInstancePtrFloat(FloatCompIdx++, DestStartIdx);
-				float* DestY = DestBuffer->GetInstancePtrFloat(FloatCompIdx++, DestStartIdx);
-				float* DestZ = DestBuffer->GetInstancePtrFloat(FloatCompIdx++, DestStartIdx);
-								
-				for (int32 i = 0; i < NumInstances; ++i)
-				{
-					FVector* Src = reinterpret_cast<FVector*>((SrcData + i * SrcVarSize));
-
-					FVector3f SimLocalSWC = FVector3f((*Src) - FVector(SimulationLwcTile) * FLargeWorldRenderScalar::GetTileSize());
-					*DestX++ = SimLocalSWC.X;
-					*DestY++ = SimLocalSWC.Y;
-					*DestZ++ = SimLocalSWC.Z;
-				}
-			}
-			else
-			{
-				TFieldIterator<FProperty> SrcPropertyIt(SrcStruct, EFieldIteratorFlags::IncludeSuper);
-				TFieldIterator<FProperty> DestPropertyIt(DestStruct, EFieldIteratorFlags::IncludeSuper);
-				for (; SrcPropertyIt; ++SrcPropertyIt, ++DestPropertyIt)
-				{
-					FProperty* SrcProperty = *SrcPropertyIt;
-					FProperty* DestProperty = *DestPropertyIt;
-					SrcData = SrcPropertyBase + SrcProperty->GetOffset_ForInternal();
-
-					//Convert any LWC doubles to floats. //TODO: Insert LWCTile... probably need to explicitly check for vectors etc.
-					if (SrcProperty->IsA(FDoubleProperty::StaticClass()))
-					{
-						check(DestProperty->IsA(FFloatProperty::StaticClass()));
-						float* Dest = DestBuffer->GetInstancePtrFloat(FloatCompIdx++, DestStartIdx);
-						//Write all instances for this component.
-						for (int32 i = 0; i < NumInstances; ++i)
-						{
-							double* Src = reinterpret_cast<double*>((SrcData + i * SrcVarSize));
-							*Dest++ = static_cast<float>(*Src);
-						}
-					}
-					else if (SrcProperty->IsA(FFloatProperty::StaticClass()))
-					{
-						float* Dest = DestBuffer->GetInstancePtrFloat(FloatCompIdx++, DestStartIdx);
-
-						//Write all instances for this component.
-						for (int32 i = 0; i < NumInstances; ++i)
-						{
-							float* Src = reinterpret_cast<float*>((SrcData + i * SrcVarSize));
-							*Dest++ = *Src;
-						}
-					}
-					else if (SrcProperty->IsA(FUInt16Property::StaticClass()))
-					{
-						FFloat16* Dest = DestBuffer->GetInstancePtrHalf(HalfCompIdx++, DestStartIdx);
-
-						//Write all instances for this component.
-						for (int32 i = 0; i < NumInstances; ++i)
-						{
-							FFloat16* Src = reinterpret_cast<FFloat16*>((SrcData + i * SrcVarSize));
-							*Dest++ = *Src;
-						}
-					}
-					else if (SrcProperty->IsA(FIntProperty::StaticClass()) || SrcProperty->IsA(FBoolProperty::StaticClass()))
-					{
-						int32* Dest = DestBuffer->GetInstancePtrInt32(IntCompIdx++, DestStartIdx);
-
-						//Write all instances for this component.
-						for (int32 i = 0; i < NumInstances; ++i)
-						{
-							int32* Src = reinterpret_cast<int32*>((SrcData + i * SrcVarSize));
-							*Dest++ = *Src;
-						}
-					}
-					//Should be able to support double easily enough
-					else if (FStructProperty* StructProp = CastField<FStructProperty>(SrcProperty))
-					{
-						FStructProperty* DestStructProp = CastField<FStructProperty>(DestProperty);
-						WriteData(StructProp->Struct, DestStructProp->Struct, SrcData);
-					}
-					else
-					{
-						checkf(false, TEXT("Property(%s) Class(%s) is not a supported type"), *SrcProperty->GetName(), *SrcProperty->GetClass()->GetName());
-					}
-				}
-			}
-		};
-		WriteData(Var.GetType().GetScriptStruct(), SimVar.GetType().GetScriptStruct(), SrcDataBase);
-	}
-}
-void FNiagaraDataChannelGameData::AppendFromGameData(const FNiagaraDataChannelGameData& GameData)
-{
-	check(GetLayoutInfo() == GameData.GetLayoutInfo());
-
-	NumElements += GameData.Num();
-	TConstArrayView<FNiagaraDataChannelVariableBuffer> SrcBuffers = GameData.GetVariableBuffers();
-	for(int32 i=0; i < SrcBuffers.Num(); ++i)
-	{
-		const FNiagaraDataChannelVariableBuffer& SrcBuffer = SrcBuffers[i];
-		FNiagaraDataChannelVariableBuffer& DestBuffer = VariableData[i];
-		DestBuffer.Data.Append(SrcBuffer.Data);
-	}
-}
-
-void FNiagaraDataChannelGameData::AppendFromDataSet(const FNiagaraDataBuffer* SrcBuffer, FVector3f SimulationLwcTile)
-{
-	const FNiagaraDataSetCompiledData& CompiledData = SrcBuffer->GetOwner()->GetCompiledData();
-
-	static TArray<uint8> LWCConversionBuffer;
-
-	int32 NumInstances = SrcBuffer->GetNumInstances();
-	int32 StartIndex = NumElements;
-	NumElements += NumInstances;
-
-	const FNiagaraDataChannelGameDataLayout& Layout = LayoutInfo->GetGameDataLayout();
-	for (const TPair<FNiagaraVariableBase, int32>& VarIndexPair : Layout.VariableIndices)
-	{
-		FNiagaraVariableBase Var = VarIndexPair.Key;
-		int32 VarIndex = VarIndexPair.Value;
-		FNiagaraDataChannelVariableBuffer& VarBuffer = VariableData[VarIndex];
-
-		VarBuffer.SetNum(NumElements);
-		
-		FNiagaraVariableBase SimVar = Var;
-		if (FNiagaraTypeHelper::IsLWCType(Var.GetType()))
-		{
-			SimVar.SetType(FNiagaraTypeDefinition(FNiagaraTypeHelper::FindNiagaraFriendlyTopLevelStruct(Var.GetType().GetScriptStruct(), ENiagaraStructConversion::Simulation)));
-		}
-
-		//Niagara Positions are a special case where we're actually storing them as FVectors in the game level data and must convert down to actual Positions/FVector3f in the sim data and visa versa
-		if(Var.GetType() == FNiagaraTypeDefinition::GetPositionDef())
-		{
-			Var.SetType(FNiagaraTypeHelper::GetVectorDef());
-		}
-		int32 VarSize = Var.GetSizeInBytes();
-
-		uint8* DestDataBase = VarBuffer.Data.GetData() + StartIndex * VarSize;
-		
-		int32 SimVarIndex = CompiledData.Variables.IndexOfByKey(SimVar);
-		if (SimVarIndex == INDEX_NONE)
-		{
-			continue; //Did not find this variable in the dataset. Warn?
-		}
-
-		const FNiagaraVariableLayoutInfo& SimLayout = CompiledData.VariableLayouts[SimVarIndex];
-		
-		int32 FloatCompIdx = SimLayout.GetFloatComponentStart();
-		int32 IntCompIdx = SimLayout.GetInt32ComponentStart();
-		int32 HalfCompIdx = SimLayout.GetHalfComponentStart();
-		
-		TFunction<void(UScriptStruct*, UScriptStruct*, uint8*)> ReadData;
-		ReadData = [&](UScriptStruct* SrcStruct, UScriptStruct* DestStruct, uint8* DestDataBase)
-		{
-			//Read all data from the simulation and place in the game level buffers. Converting from LWC local tile space where needed.
-
-			//Special case for writing to Niagara Positions. We must offset the data into the simulation local LWC space.
-			if (SrcStruct == FNiagaraTypeDefinition::GetPositionStruct())
-			{
-				const float* SrcX = reinterpret_cast<const float*>(SrcBuffer->GetComponentPtrFloat(FloatCompIdx++));
-				const float* SrcY = reinterpret_cast<const float*>(SrcBuffer->GetComponentPtrFloat(FloatCompIdx++));
-				const float* SrcZ = reinterpret_cast<const float*>(SrcBuffer->GetComponentPtrFloat(FloatCompIdx++));
-
-				for (int32 i = 0; i < NumInstances; ++i)
-				{
-					FVector* Dest = reinterpret_cast<FVector*>(DestDataBase + VarSize * i);
-					*Dest = FVector(*SrcX++, *SrcY++, *SrcZ++) + FVector(SimulationLwcTile) * FLargeWorldRenderScalar::GetTileSize();
-				}
-			}
-			else
-			{
-				TFieldIterator<FProperty> SrcPropertyIt(SrcStruct, EFieldIteratorFlags::IncludeSuper);
-				TFieldIterator<FProperty> DestPropertyIt(DestStruct, EFieldIteratorFlags::IncludeSuper);
-				for (; SrcPropertyIt; ++SrcPropertyIt, ++DestPropertyIt)
-				{
-					FProperty* SrcProperty = *SrcPropertyIt;
-					FProperty* DestProperty = *DestPropertyIt;
-					int32 DestOffset = DestProperty->GetOffset_ForInternal();
-					uint8* DestData = DestDataBase + DestOffset;
-					if (DestPropertyIt->IsA(FDoubleProperty::StaticClass()))
-					{
-						double* Dest = reinterpret_cast<double*>(DestData);
-						const float* Src = reinterpret_cast<const float*>(SrcBuffer->GetComponentPtrFloat(FloatCompIdx++));
-
-						//Write all instances for this component.
-						for (int32 i = 0; i < NumInstances; ++i)
-						{
-							Dest = reinterpret_cast<double*>(DestData + VarSize * i);
-							*Dest = *Src++;
-						}
-					}
-					else if (DestPropertyIt->IsA(FFloatProperty::StaticClass()))
-					{
-						float* Dest = reinterpret_cast<float*>(DestData);
-						const float* Src = reinterpret_cast<const float*>(SrcBuffer->GetComponentPtrFloat(FloatCompIdx++));
-
-						//Write all instances for this component.
-						for (int32 i = 0; i < NumInstances; ++i)
-						{
-							Dest = reinterpret_cast<float*>(DestData + VarSize * i);
-							*Dest = *Src++;
-						}
-					}
-					else if (DestPropertyIt->IsA(FUInt16Property::StaticClass()))
-					{
-						FFloat16* Dest = reinterpret_cast<FFloat16*>(DestData);
-						const FFloat16* Src = reinterpret_cast<const FFloat16*>(SrcBuffer->GetComponentPtrHalf(HalfCompIdx++));
-
-						//Write all instances for this component.
-						for (int32 i = 0; i < NumInstances; ++i)
-						{
-							Dest = reinterpret_cast<FFloat16*>(DestData + VarSize * i);
-							*Dest = *Src++;
-						}
-					}
-					else if (DestPropertyIt->IsA(FIntProperty::StaticClass()) || DestPropertyIt->IsA(FBoolProperty::StaticClass()))
-					{
-						int32* Dest = reinterpret_cast<int32*>(DestData);
-						const int32* Src = reinterpret_cast<const int32*>(SrcBuffer->GetComponentPtrInt32(IntCompIdx++));
-
-						//Write all instances for this component.
-						for (int32 i = 0; i < NumInstances; ++i)
-						{
-							Dest = reinterpret_cast<int32*>(DestData + VarSize * i);
-							*Dest = *Src++;
-						}
-					}
-					else if (FStructProperty* SrcStructProp = CastField<FStructProperty>(SrcProperty))
-					{
-						FStructProperty* DestStructProp = CastField<FStructProperty>(DestProperty);
-						ReadData(SrcStructProp->Struct, DestStructProp->Struct, DestData);
-					}
-					else
-					{
-						checkf(false, TEXT("Property(%s) Class(%s) is not a supported type"), *SrcProperty->GetName(), *SrcProperty->GetClass()->GetName());
-					}
-				}
-			}
-		};
-		ReadData(SimVar.GetType().GetScriptStruct(), Var.GetType().GetScriptStruct(), DestDataBase);
-	}
-}
-
-void FNiagaraDataChannelGameData::SetFromSimCache(const FNiagaraVariableBase& SourceVar, TConstArrayView<uint8> Data, int32 Size)
-{
-	const FNiagaraDataChannelGameDataLayout& Layout = LayoutInfo->GetGameDataLayout();
-	if (const int* Index = Layout.VariableIndices.Find(SourceVar))
-	{
-		if (VariableData[*Index].Size == Size)
-		{
-			VariableData[*Index].Data = Data;
-		}
-	}
-}
-
-//////////////////////////////////////////////////////////////////////////
-
-FNiagaraDataChannelDataProxy::~FNiagaraDataChannelDataProxy()
-{
-	if(GPUDataSet)
-	{
-		delete GPUDataSet;
-	}
-	GPUDataSet = nullptr;
-
-	Owner = nullptr;
-	CurrFrameData = nullptr;
-	PrevFrameData = nullptr;
-	bNeedsPrevFrameData = false;
-	PendingCPUBuffers.Empty();
-	PendingGPUAllocations = 0;
-}
-
-
-void FNiagaraDataChannelDataProxy::OnAddedToDispatcher(FNiagaraGpuComputeDispatchInterface* ComputeDispatchInterface)
-{
-#if !UE_BUILD_SHIPPING
-	DispatchInterfaceForDebuggingOnly = ComputeDispatchInterface;
-#endif
-}
-
-void FNiagaraDataChannelDataProxy::OnRemovedFromDispatcher(FNiagaraGpuComputeDispatchInterface* ComputeDispatchInterface)
-{
-#if !UE_BUILD_SHIPPING
-	check(DispatchInterfaceForDebuggingOnly == ComputeDispatchInterface);
-#endif
-
-	FNiagaraGPUInstanceCountManager& InstCountManager = ComputeDispatchInterface->GetGPUInstanceCounterManager();
-
-	if (CurrFrameData)
-	{
-		uint32 InstanceCountOffset = CurrFrameData->GetGPUInstanceCountBufferOffset();
-		InstCountManager.FreeEntry(InstanceCountOffset);
-		CurrFrameData->SetGPUInstanceCountBufferOffset(INDEX_NONE);
-	}
-
-	if (PrevFrameData)
-	{
-		uint32 InstanceCountOffset = PrevFrameData->GetGPUInstanceCountBufferOffset();
-		InstCountManager.FreeEntry(InstanceCountOffset);
-		PrevFrameData->SetGPUInstanceCountBufferOffset(INDEX_NONE);
-	}
-
-	for (FNDCGpuReadbackInfo& ReadbackInfo : PendingGPUReadbackBuffers)
-	{
-		check(ReadbackInfo.Buffer);
-		uint32 Offset = ReadbackInfo.Buffer->GetGPUInstanceCountBufferOffset();
-		ComputeDispatchInterface->GetGPUInstanceCounterManager().FreeEntry(Offset);
-		ReadbackInfo.Buffer->SetGPUInstanceCountBufferOffset(INDEX_NONE);
-	}
-	PendingGPUReadbackBuffers.Reset();
-}
-
-void FNiagaraDataChannelDataProxy::Reset()
-{
-	PrevFrameData = nullptr;
-}
-
-void FNiagaraDataChannelDataProxy::BeginFrame(FNiagaraGpuComputeDispatchInterface* DispatchInterface, FRHICommandListImmediate& RHICmdList)
-{
-	check(GPUDataSet);
-	check(GPUDataSet->GetSimTarget() == ENiagaraSimTarget::GPUComputeSim);
-#if !UE_BUILD_SHIPPING
-	check(DispatchInterfaceForDebuggingOnly == DispatchInterface);
-#endif
-
-	FNiagaraGPUInstanceCountManager& InstCountManager = DispatchInterface->GetGPUInstanceCounterManager();	
-
-	for (FNDCGpuReadbackInfo& ReadbackInfo : PendingGPUReadbackBuffers)
-	{
-		uint32 Offset = ReadbackInfo.Buffer->GetGPUInstanceCountBufferOffset();
-		DispatchInterface->GetGPUInstanceCounterManager().FreeEntry(Offset);
-		ReadbackInfo.Buffer->SetGPUInstanceCountBufferOffset(INDEX_NONE);
-	}
-	PendingGPUReadbackBuffers.Reset();
-
-	if(this->bNeedsPrevFrameData)
-	{
-		Swap(PrevFrameData, CurrFrameData);
-	}
-
-	//The base num instances for writing from the GPU.
-	uint32 NumInstanceFromCPU = 0;
-	for (auto& Buffer : PendingCPUBuffers)
-	{
-		NumInstanceFromCPU += Buffer->GetNumInstances();
-	}
-
-	//Allocate our GPU buffers we'll write into. Combine accumulated counts from writing DIs and the data coming in from the CPU.
-	uint32 InstancesToAllocate = NumInstanceFromCPU + PendingGPUAllocations;
-	PendingGPUAllocations = 0;
-
-	if(InstancesToAllocate == 0)
-	{
-		CurrFrameData = nullptr;
-		return;
-	}
-
-	if(CurrFrameData == nullptr)
-	{		
-		FNiagaraDataBuffer& NewBuffer = GPUDataSet->AllocateBuffer();
-		CurrFrameData = NewBuffer.UnlockForRead();
-	}
-
-	uint32 InstanceCountOffset = CurrFrameData->GetGPUInstanceCountBufferOffset();
-	InstCountManager.FreeEntry(InstanceCountOffset);
-
-	CurrFrameData->AllocateGPU(RHICmdList, InstancesToAllocate, DispatchInterface->GetFeatureLevel(), GetDebugName());
-	CurrFrameData->PushCPUBuffersToGPU(PendingCPUBuffers, true, RHICmdList, DispatchInterface->GetFeatureLevel(), GetDebugName(), false);
-	CurrFrameData->SetNumInstances(InstancesToAllocate);
-	PendingCPUBuffers.Reset();
-
-	InstanceCountOffset = InstCountManager.AcquireOrAllocateEntry(RHICmdList);	
-	CurrFrameData->SetGPUInstanceCountBufferOffset(InstanceCountOffset);
-
-	//Init the instance count to the value from the CPU before it is added to by the GPU writes.
-	if(NumInstanceFromCPU > 0)
-	{
-		InstCountManager.AddInstanceCountInitTask(InstanceCountOffset, NumInstanceFromCPU);
-	}
-}
-
-void FNiagaraDataChannelDataProxy::EndFrame(FNiagaraGpuComputeDispatchInterface* DispatchInterface, FRHICommandListImmediate& RHICmdList)
-{
-	check(GPUDataSet);
-	check(GPUDataSet->GetSimTarget() == ENiagaraSimTarget::GPUComputeSim);
-
-#if !UE_BUILD_SHIPPING
-	check(DispatchInterfaceForDebuggingOnly == DispatchInterface);
-#endif
-
-	check(CurrBufferAccessCounts == 0);
-
-	//Handle pending readbacks for GPU->CPU data.
-	{
-		for (FNDCGpuReadbackInfo& ReadbackInfo : PendingGPUReadbackBuffers)
-		{
-			//We enqueue a readback for the data if we're wanting to publish it to the CPU or Game
-			check(ReadbackInfo.bPublishToCPU || ReadbackInfo.bPublishToGame);//We should be doing one or the other or we should not get here.
-
-			auto PublishOnCPU = [
-				LocalOwner = Owner,
-				PASS_PublishCPU = ReadbackInfo.bPublishToCPU,
-				PASS_PublishGame = ReadbackInfo.bPublishToGame,
-				PASS_LWCTile = ReadbackInfo.LWCTile,
-				PASS_DataSetLayoutHash = LayoutInfo->GetDataSetCompiledData().GetLayoutHash()](TSharedRef<FNiagaraDataBufferReadback> Readback)
-				{
-					if (FNiagaraDataChannelData* NDCData = LocalOwner.Pin().Get())
-					{
-						const FNiagaraDataSetCompiledData& CompiledData = NDCData->GetCompiledData(ENiagaraSimTarget::CPUSim);
-						if(PASS_DataSetLayoutHash == CompiledData.GetLayoutHash())
-						{
-							//Only consume the readback it the layouts match.
-							//It's possible the layout has been changed in the time it took for the readback to get here.
-							FNiagaraDataChannelPublishRequest PublishRequest;
-							PublishRequest.bVisibleToGame = PASS_PublishGame;
-							PublishRequest.bVisibleToCPUSims = PASS_PublishCPU;
-							PublishRequest.bVisibleToGPUSims = false;//Don't want to ping pong back to the GPU
-							PublishRequest.LwcTile = PASS_LWCTile;
-
-							FNiagaraDataBuffer* ReadbackBuffer = NDCData->GetBufferForCPUWrite();
-							Readback->ReadResultsToDataBuffer(ReadbackBuffer);
-							PublishRequest.Data = ReadbackBuffer->UnlockForRead();
-
-#if WITH_NIAGARA_DEBUGGER
-							PublishRequest.DebugSource = TEXT("NDC GPU Readback");//TODO: Feed in real source names also
-#endif
-							NDCData->PublishFromGPU(PublishRequest);
-						}
-					}
-				};
-
-			TSharedPtr<FNiagaraDataBufferReadback, ESPMode::ThreadSafe> NewReadback = MakeShared<FNiagaraDataBufferReadback, ESPMode::ThreadSafe>();
-			NewReadback->GetOnReadbackComplete().BindLambda(PublishOnCPU);
-			NewReadback->EnqueueReadback(RHICmdList, ReadbackInfo.Buffer, DispatchInterface->GetGpuReadbackManager(), DispatchInterface->GetGPUInstanceCounterManager());
-
-			//Defer the release of this until next frame to allow other users of this instance count
-// 			uint32 Offset = ReadbackInfo.Buffer->GetGPUInstanceCountBufferOffset();
-// 			DispatchInterface->GetGPUInstanceCounterManager().FreeEntry(Offset);
-// 			ReadbackInfo.Buffer->SetGPUInstanceCountBufferOffset(INDEX_NONE);
-		}
-		PendingGPUReadbackBuffers.Reset();
-	}
-}
-
-FNiagaraDataBufferRef FNiagaraDataChannelDataProxy::PrepareForWriteAccess(FRDGBuilder& GraphBuilder)
-{
-	check(IsInRenderingThread());
-
-	if(CurrFrameData)
-	{
-		if (CurrBufferAccessCounts > 0)
-		{
-#if !UE_BUILD_SHIPPING
-			if(!bWarnedAboutSameStageRW)
-			{
-				bWarnedAboutSameStageRW = true;
-				UE_LOG(LogNiagara, Warning, TEXT("Attempting to write to a Niagara Data Channel in the same stage in which it's being read. {%s}\n\
-				This is most often caused by an NDC reading the current frame data and the same NDC being written to in PostRenderOpaque.\n\
-				A possible fix is to read the previous frames data rather than the current frame.\n\
-				Another is to ensure writes to the NDC are done earlier. Things such as Depth/GBuffer reads can make the writing system tick late in the frame."), *DebugName);
-			}
-#endif
-			return nullptr;
-		}
-
-		//If we're prepping the first writer, transition to UAV
-		if (CurrBufferAccessCounts == 0)
-		{
-			AddTransition(GraphBuilder, ERHIAccess::SRVMask, ERHIAccess::UAVCompute, CurrFrameData.GetReference());
-		}
-
-		--CurrBufferAccessCounts;
-	}
-
-	return CurrFrameData;
-}
-
-void FNiagaraDataChannelDataProxy::EndWriteAccess(FRDGBuilder& GraphBuilder)
-{
-	check(IsInRenderingThread());
-
-	check(CurrBufferAccessCounts < 0);
-
-	++CurrBufferAccessCounts;
-
-	//If we reach 0 writers, transition back to SRV.
-	if (CurrBufferAccessCounts == 0)
-	{
-		AddTransition(GraphBuilder, ERHIAccess::UAVCompute, ERHIAccess::SRVMask, CurrFrameData.GetReference());
-	}
-}
-
-FNiagaraDataBufferRef FNiagaraDataChannelDataProxy::PrepareForReadAccess(FRDGBuilder& GraphBuilder, bool bCurrentFrame)
-{
-	check(IsInRenderingThread());
-
-	if(bCurrentFrame)
-	{
-		if(CurrFrameData)
-		{
-			if (CurrBufferAccessCounts < 0)
-			{
-#if !UE_BUILD_SHIPPING
-				if (!bWarnedAboutSameStageRW)
-				{
-					bWarnedAboutSameStageRW = true;
-					UE_LOG(LogNiagara, Warning, TEXT("Attempting to read from a Niagara Data Channel in the same stage in which it's being written. {%s}\n\
-					This is most often caused by an NDC reading the current frame data and the same NDC being written to in PostRenderOpaque.\n\
-					A possible fix is to read the previous frame data rather than the current frame.\n\
-					Another is to ensure writes to the NDC are done earlier. Things such as Depth/GBuffer reads can make the writing system tick late in the frame."), *DebugName);
-				}
-#endif
-				return nullptr;
-			}
-
-			++CurrBufferAccessCounts;
-
-			return CurrFrameData;
-		}
-		return nullptr;
-	}
-	else
-	{
-		return PrevFrameData;
-	}
-}
-
-void FNiagaraDataChannelDataProxy::EndReadAccess(FRDGBuilder& GraphBuilder, bool bCurrentFrame)
-{
-	check(IsInRenderingThread());
-	if(bCurrentFrame)
-	{
-		check(CurrBufferAccessCounts > 0);
-		--CurrBufferAccessCounts;
-	}
-}
-
-void FNiagaraDataChannelDataProxy::AddTransition(FRDGBuilder& GraphBuilder, ERHIAccess AccessBefore, ERHIAccess AccessAfter, FNiagaraDataBuffer* Buffer)
-{
-	if (Buffer)
-	{
-		TArray<FRHITransitionInfo, TInlineAllocator<6>> Transitions;
-		Transitions.Reserve(3);
-
-		Transitions.Emplace(Buffer->GetGPUBufferFloat().UAV, AccessBefore, AccessAfter);
-		Transitions.Emplace(Buffer->GetGPUBufferInt().UAV, AccessBefore, AccessAfter);
-		//TODO: Half Support | Transitions.Emplace(Buffer->GetGPUBufferHalf().UAV, AccessBefore, AccessAfter);
-
-		TArray<FRHIUnorderedAccessView*> UAVsToOverlap;
-		UAVsToOverlap.Emplace(Buffer->GetGPUBufferFloat().UAV);
-		UAVsToOverlap.Emplace(Buffer->GetGPUBufferInt().UAV);
-		//TODO: Half Support | UAVsToOverlap.Emplace(Buffer->GetGPUBufferHalf().UAV);
-
-	GraphBuilder.AddPass(
-		RDG_EVENT_NAME("NDC Proxy - Transition Buffers"),
-		ERDGPassFlags::None,
-		[Transitions, AccessAfter, UAVsToOverlap](FRHICommandListImmediate& RHICmdList)
-		{	
-			RHICmdList.Transition(Transitions);		
-			
-			//We may have multiple overlapping dispatches accessing the same NDC buffers.
-			if(AccessAfter == ERHIAccess::UAVCompute)
-			{
-				RHICmdList.BeginUAVOverlap(UAVsToOverlap);
-			}
-			else
-			{
-				RHICmdList.EndUAVOverlap(UAVsToOverlap);
-			}
-		});
-	}
-}
-
-FNiagaraDataBufferRef FNiagaraDataChannelDataProxy::AllocateBufferForCPU(FRDGBuilder& GraphBuilder, ERHIFeatureLevel::Type FeatureLevel, int32 AllocationSize, bool bPublishToGame, bool bPublishToCPU, FVector3f LWCTile)
-{
-	FNiagaraDataBuffer* Buffer = &GPUDataSet->AllocateBuffer();
-	Buffer->AllocateGPU(GraphBuilder.RHICmdList, AllocationSize, FeatureLevel, TEXT("NDC GPU Buffers for CPU"));
-
-	FNiagaraDataBufferRef Ret = Buffer->UnlockForRead();
-
-	if(bPublishToCPU || bPublishToGame)
-	{
-		FNDCGpuReadbackInfo& NewReadback = PendingGPUReadbackBuffers.AddDefaulted_GetRef();
-		NewReadback.Buffer = Ret;
-		NewReadback.bPublishToCPU = bPublishToCPU;
-		NewReadback.bPublishToGame = bPublishToGame;
-		NewReadback.LWCTile = LWCTile;
-	}
-
-	AddTransition(GraphBuilder, ERHIAccess::SRVMask, ERHIAccess::UAVCompute, Buffer);
-
-	return Ret;
-}
-
-void FNiagaraDataChannelDataProxy::AddBuffersFromCPU(const TArray<FNiagaraDataBufferRef>& BuffersFromCPU)
-{
-	PendingCPUBuffers.Append(BuffersFromCPU);
-}
-
-void FNiagaraDataChannelDataProxy::AddGPUAllocationForNextTick(int32 AllocationCount)
-{
-	PendingGPUAllocations += AllocationCount;
-}
-
-//////////////////////////////////////////////////////////////////////////
-
-FNiagaraDataChannelData::FNiagaraDataChannelData()
-{
-
-}
-
-FNiagaraDataChannelData::~FNiagaraDataChannelData()
-{
-	GameData.Reset();
-
-	/** We defer the deletion of the dataset to the RT to be sure all in-flight RT commands have finished using it.*/
-	ENQUEUE_RENDER_COMMAND(FDeleteContextCommand)(
-		[CPUDataChannelDataSet = CPUSimData](FRHICommandListImmediate& RHICmdList)
-	{
-		if (CPUDataChannelDataSet != nullptr)
-		{
-			delete CPUDataChannelDataSet;
-		}
-	}
-	);
-	CPUSimData = nullptr;
-	RTProxy = nullptr;
-}
-
-void FNiagaraDataChannelData::Init(UNiagaraDataChannelHandler* Owner)
-{
-	check(Owner);
-	const UNiagaraDataChannel* DataChannel = Owner->GetDataChannel();
-	check(DataChannel);
-	
-	LayoutInfo = DataChannel->GetLayoutInfo();
-	
-	GameData = DataChannel->CreateGameData();
-
-	CPUSimData = new FNiagaraDataSet();
-
-	CPUSimData->Init(&LayoutInfo->GetDataSetCompiledData(), 1);
-
-	//TODO: Send game data to GPU direct without staging.
-	GameDataStaging = new FNiagaraDataSet();
-	GameDataStaging->Init(&LayoutInfo->GetDataSetCompiledData());
-
-	UWorld* OwnerWorld = Owner->GetWorld();
-	WeakOwnerWorld = OwnerWorld;
-
-	CreateRenderThreadProxy(Owner);
-}
-
-void FNiagaraDataChannelData::Reset()
-{
-	FScopeLock Lock(&PublishCritSec);
-	PublishRequests.Reset();
-	PublishRequestsFromGPU.Reset();
-	PublishRequestsForGPU.Reset();
-	
-	PrevCPUSimData = nullptr;	
-
-	GameData->Empty();
-	CPUSimData->ResetBuffers();
-	GameDataStaging->ResetBuffers();
-	ENQUEUE_RENDER_COMMAND(FResetProxyCommand)(
-		[RT_Proxy = RTProxy](FRHICommandListImmediate& RHICmdList)
-	{
-		if (RT_Proxy)
-		{
-			RT_Proxy->Reset();
-		}
-	}
-	);
-}
-
-void FNiagaraDataChannelData::BeginFrame(UNiagaraDataChannelHandler* Owner)
-{
-	//Lazy re-init our proxy data in cases where we've had to recreate our render state while running due to an FScene change etc.
-	if(RTProxy == nullptr)
-	{
-		CreateRenderThreadProxy(Owner);
-	}
-
-	GameData->BeginFrame();
-
-	bool bRequirePreviousData = Owner->GetDataChannel()->KeepPreviousFrameData();
-	if (bRequirePreviousData)
-	{
-		PrevCPUSimData = CPUSimData->GetCurrentData();
-	}
-
-	//Grab a new buffer to store the CPU data.
-	CPUSimData->BeginSimulate();
-	CPUSimData->EndSimulate();
-
-	//Pull in our publish requests from the GPU to be processed in the first tick group.
-	PublishRequests.Append(PublishRequestsFromGPU);
-	PublishRequestsFromGPU.Reset();
-}
-
-void FNiagaraDataChannelData::EndFrame(UNiagaraDataChannelHandler* Owner)
-{
-	//We must do one final tick to process any final items generated by the last things to tick this frame.
-	ConsumePublishRequests(Owner, TG_LastDemotable);
-
-	if(RTProxy)
-	{
-		int32 TotalGPUInstances = 0;
-		//Pass over any data bound for the GPU to the render thread.
-		TArray<FNiagaraDataBufferRef> BuffersForGPU;
-		BuffersForGPU.Reserve(PublishRequestsForGPU.Num());
-		
-		//If we've used our CPU data to spawn into a GPU sim this frame then we need to pass over all our CPU data regardless of if it was marked for GPU or not.		
-		if(NumGPUSpawningReaders > 0 && NDCCVars::AutoUploadGPUSpawnData)
-		{
-			if(CPUSimData->GetCurrentData() && CPUSimData->GetCurrentData()->GetNumInstances() > 0)
-			{
-				//First add our entire CPU sim data buffer.
-				BuffersForGPU.Emplace(CPUSimData->GetCurrentData());
-				TotalGPUInstances += CPUSimData->GetCurrentData()->GetNumInstances();
-			}
-
-			//Next add existing buffers we've kept around to pass to the GPU but only those that don't already exist in the main CPU buffers.
-			for(auto It = PublishRequestsForGPU.CreateIterator() ; It; ++It)
-			{
-				FNiagaraDataChannelPublishRequest& PublishRequest = *It;
-				check(PublishRequest.bVisibleToGPUSims);
-				if(!PublishRequest.bVisibleToCPUSims)
-				{
-					if(PublishRequest.Data->GetNumInstances() > 0)
-					{
-						BuffersForGPU.Emplace(PublishRequest.Data);
-						TotalGPUInstances += PublishRequest.Data->GetNumInstances();
-					}
-				}
-			}
-		}
-		else
-		{
-			//Only add data that was explicitly marked for GPU.
-			for (auto It = PublishRequestsForGPU.CreateIterator(); It; ++It)
-			{
-				FNiagaraDataChannelPublishRequest& PublishRequest = *It;
-				check(PublishRequest.bVisibleToGPUSims);
-				if (PublishRequest.Data->GetNumInstances() > 0)
-				{
-					BuffersForGPU.Emplace(PublishRequest.Data);
-					TotalGPUInstances += PublishRequest.Data->GetNumInstances();
-				}
-			}
-		}
-
-		if(TotalGPUInstances > 0)
-		{
-			ENQUEUE_RENDER_COMMAND(FDataChannelProxyEndFrame) (
-				[RT_Proxy = RTProxy, RT_BuffersForGPU = MoveTemp(BuffersForGPU)](FRHICommandListImmediate& CmdList)
-				{
-
-					RT_Proxy->AddBuffersFromCPU(RT_BuffersForGPU);
-				});
-		}		
-	}
-	PublishRequestsForGPU.Reset();
-}
-
-int32 FNiagaraDataChannelData::ConsumePublishRequests(UNiagaraDataChannelHandler* Owner, const ETickingGroup& TickGroup)
-{
-	check(IsValid(Owner));
-
-	FlushAllPendingGameData();
-
-	//There should be no access on other threads at this point anyway but lock just to be safe.
-	FScopeLock Lock(&PublishCritSec);
-
-	const UNiagaraDataChannel* DataChannel = Owner->GetDataChannel();
-	if(PublishRequests.Num() == 0 || DataChannel == nullptr)
-	{
-		return 0;
-	}
-	
-	if(NDCCVars::bEmitWarningsOnLateNDCWrites && DataChannel->ShouldEnforceTickGroupReadWriteOrder())
-	{
-		ETickingGroup PublishSourceTG = static_cast<ETickingGroup>(FMath::Clamp(static_cast<int32>(Owner->GetCurrentTickGroup()) - 1, 0, static_cast<int32>(ETickingGroup::TG_MAX) - 1));//We're consuming from the previous TG.
-		ETickingGroup FinalWriteTG = DataChannel->GetFinalWriteTickGroup();
-
-		//TODO: Possibly allow late writes to be deferred to the next frame?
-		if(PublishSourceTG > FinalWriteTG)
-		{
-			static UEnum* TGEnum = StaticEnum<ETickingGroup>();
-			UE_LOG(LogNiagara, Warning, TEXT("Data Channel %s is being written to in Tick Group %s which is after it's final write tick group %s. This may cause incorrect read / write ordering and missed data.")
-			, *DataChannel->GetAsset()->GetName()
-			, *TGEnum->GetDisplayNameTextByValue((int32)PublishSourceTG).ToString()
-			, *TGEnum->GetDisplayNameTextByValue((int32)FinalWriteTG).ToString());
-		}
-	}
-
-	UWorld* World = Owner->GetWorld();
-	check(IsValid(World));
-
-	int32 GameDataOrigSize = GameData->Num();
-	int32 CPUDataOrigSize = CPUSimData->GetCurrentData()->GetNumInstances();
-	int32 NewGameDataChannel = GameDataOrigSize;
-	int32 NewCPUDataChannel = CPUDataOrigSize;
-
-	//Do a pass to gather the new total size for our DataChannel data.
-
-	//Each DI that generates DataChannel can control whether it's pushed to Game/CPU/GPU.
-	int32 RequestCount = PublishRequests.Num();
-	PublishRequestsForGPU.Reserve(PublishRequestsForGPU.Num() + RequestCount);
-
-	for (FNiagaraDataChannelPublishRequest& PublishRequest : PublishRequests)
-	{
-		uint32 NumInsts = 0;
-		//Don't bother uploading to the GPU separately if we already know we're going to be sending it via the CPU data as a whole.
-		bool bAutoUploadToGPU = NDCCVars::AutoUploadGPUSpawnData && NumGPUSpawningReaders > 0 && PublishRequest.bVisibleToCPUSims;
-		bool bShouldSendToGPU = PublishRequest.bVisibleToGPUSims && !bAutoUploadToGPU;
-
-		if (FNiagaraDataChannelGameData* RequestGameData = PublishRequest.GameData.Get())
-		{
-			NumInsts = RequestGameData->Num();
-			if (RequestGameData->Num() > 0 && bShouldSendToGPU)
-			{
-				//We stage the game data into a data set to facilitate easier copy over to the GPU. TODO: Send and copy the game data directly.	
-				FNiagaraDataBuffer* StagingBuf = &GameDataStaging->BeginSimulate();
-				StagingBuf->Allocate(NumInsts);
-				RequestGameData->WriteToDataSet(StagingBuf, 0, LwcTile);
-				GameDataStaging->EndSimulate();
-				FNiagaraDataChannelPublishRequest& NewGPUReq = PublishRequestsForGPU.AddDefaulted_GetRef();
-				NewGPUReq.Data = StagingBuf;
-				NewGPUReq.bVisibleToCPUSims = PublishRequest.bVisibleToCPUSims;
-				NewGPUReq.bVisibleToGPUSims = PublishRequest.bVisibleToGPUSims;
-				NewGPUReq.LwcTile = PublishRequest.LwcTile;
-			}
-		}
-		else if (ensure(PublishRequest.Data))
-		{
-			NumInsts = PublishRequest.Data->GetNumInstances();
-
-			if (NumInsts > 0 && bShouldSendToGPU)
-			{
-				FNiagaraDataChannelPublishRequest& NewGPUReq = PublishRequestsForGPU.AddDefaulted_GetRef();
-				NewGPUReq.Data = PublishRequest.Data;
-				NewGPUReq.bVisibleToCPUSims = PublishRequest.bVisibleToCPUSims;
-				NewGPUReq.bVisibleToGPUSims = PublishRequest.bVisibleToGPUSims;
-				NewGPUReq.LwcTile = PublishRequest.LwcTile;
-			}
-		}
-
-		if (PublishRequest.bVisibleToGame)
-		{
-			NewGameDataChannel += NumInsts;
-		}
-		if (PublishRequest.bVisibleToCPUSims)
-		{
-			NewCPUDataChannel += NumInsts;
-		}
-	}
-
-	//Allocate Sim buffers ready for gather.
-	CPUSimData->BeginSimulate();
-	CPUSimData->Allocate(NewCPUDataChannel, true);
-
-	//Now do the actual data collection.
-
-	GameData->Reserve(NewGameDataChannel);
-
-	for (auto It = PublishRequests.CreateIterator(); It ; ++It)
-	{
-		FNiagaraDataChannelPublishRequest& PublishRequest = *It;
-		if (FNiagaraDataChannelGameData* RequestGameData = PublishRequest.GameData.Get())
-		{
-			if (RequestGameData->Num())
-			{
-				if(PublishRequest.bVisibleToGame)
-				{
-					GameData->AppendFromGameData(*RequestGameData);
-				}
-				if(PublishRequest.bVisibleToCPUSims)
-				{
-					RequestGameData->WriteToDataSet(CPUSimData->GetDestinationData(), CPUSimData->GetDestinationDataChecked().GetNumInstances(), LwcTile);
-				}
-			}
-		}
-
-		if (PublishRequest.Data)
-		{
-			if (PublishRequest.bVisibleToGame)
-			{
-				GameData->AppendFromDataSet(PublishRequest.Data, PublishRequest.LwcTile);
-			}
-			if (PublishRequest.bVisibleToCPUSims)
-			{
-				PublishRequest.Data->CopyToUnrelated(CPUSimData->GetDestinationDataChecked(), 0, CPUSimData->GetDestinationDataChecked().GetNumInstances(), PublishRequest.Data->GetNumInstances());
-			}
-		}
-#if WITH_NIAGARA_DEBUGGER
-		FNiagaraDataChannelDebugUtilities::LogWrite(PublishRequest, Owner->GetDataChannel(), TickGroup);
-#endif
-	}
-
-	PublishRequests.Reset();
-
-	CPUSimData->EndSimulate();
-
-#if WITH_NIAGARA_DEBUGGER
-	if(Owner->GetDataChannel()->GetVerboseLogging())
-	{
-		FString Label = FString::Printf(TEXT("Data Channel %s - CURR"), *Owner->GetDataChannel()->GetName());
-		CPUSimData->GetCurrentData()->Dump(0, CPUSimData->GetCurrentData()->GetNumInstances(), Label);
-
-		if(PrevCPUSimData)
-		{
-			FString LabelPrev = FString::Printf(TEXT("Data Channel %s - PREV"), *Owner->GetDataChannel()->GetName());
-			PrevCPUSimData->Dump(0, PrevCPUSimData->GetNumInstances(), LabelPrev);
-		}
-	}
-#endif
-
-	Owner->NotifySubscribers(this, GameDataOrigSize, NewGameDataChannel - GameDataOrigSize);
-	return RequestCount;
-}
-
-FNiagaraDataBufferRef FNiagaraDataChannelData::GetCPUData(bool bPreviousFrame)
-{
-	if(bPreviousFrame)
-	{
-		return PrevCPUSimData;
-	}
-	else
-	{
-		return CPUSimData->GetCurrentData();
-	}
-}
-
-FNiagaraDataChannelGameData* FNiagaraDataChannelData::GetGameData()
-{
-	return GameData.Get();
-}
-
-void FNiagaraDataChannelData::Publish(const FNiagaraDataChannelPublishRequest& Request)
-{
-	FScopeLock Lock(&PublishCritSec);
-
-	PublishRequests.Add(Request);
-}
-
-void FNiagaraDataChannelData::PublishFromGPU(const FNiagaraDataChannelPublishRequest& Request)
-{
-	check(IsInGameThread());
-	PublishRequestsFromGPU.Add(Request);
-}
-
-const FNiagaraDataSetCompiledData& FNiagaraDataChannelData::GetCompiledData(ENiagaraSimTarget SimTarget)
-{
-	return SimTarget == ENiagaraSimTarget::CPUSim ? LayoutInfo->GetDataSetCompiledData() : LayoutInfo->GetDataSetCompiledDataGPU();
-}
-
-FNiagaraDataBuffer* FNiagaraDataChannelData::GetBufferForCPUWrite()
-{
-	check(IsInGameThread());
-	if(CPUSimData)
-	{
-		return &CPUSimData->AllocateBuffer();
-	}
-	return nullptr;
-}
-
-bool FNiagaraDataChannelData::IsLayoutValid(UNiagaraDataChannelHandler* Owner)const
-{
-	if(Owner && LayoutInfo)
-	{
-		return Owner->GetDataChannel()->GetLayoutInfo() == LayoutInfo;//If our layout has been modified then this NDC data is no longer valid.
-	}
-	return false;
-}
-
-bool FNiagaraDataChannelData::HasData()const
-{
-	return GameData->Num() > 0 || (CPUSimData && CPUSimData->GetCurrentData()->GetNumInstances() > 0) || PublishRequestsForGPU.Num() > 0;
-}
-
-FNiagaraDataChannelGameDataPtr FNiagaraDataChannelData::GetGameDataForWriteGT(int32 Count, bool bVisibleToGame, bool bVisibleToCPU, bool bVisibleToGPU, const FString& DebugSource)
-{
-	//We keep an array of cached data ptrs, one for each of the most recent combinations of visible flags.
-	//This dynamic choice of sending data to Game/CPU/GPU is honestly annoying. I think I might look at removing and having only as a system wide thing on the NDC	
-	auto GetPendingBufferIdx = [](bool bVisibleToGame, bool bVisibleToCPU, bool bVisibleToGPU)
-	{
-		uint32 PendingBufferIndex = bVisibleToGame ? 1 : 0;
-		PendingBufferIndex += (bVisibleToCPU ? 1 : 0) << 1;
-		PendingBufferIndex += (bVisibleToGPU ? 1 : 0) << 2;
-		return PendingBufferIndex;
-	};
-
-	if(PendingDestGameData.Num() == 0)
-	{
-		PendingDestGameData.SetNum(GetPendingBufferIdx(true,true,true) + 1);
-	}
-
-	//See if we have a pending game data that can accommodate the new write.
-	
-	int32 PendingIndex = GetPendingBufferIdx(bVisibleToGame, bVisibleToCPU, bVisibleToGPU);
-	FNiagaraDataChannelGameDataPtr& CachedData = PendingDestGameData[PendingIndex];
-	if(NDCCVars::EnableAccumulatedWrites)
-	{
-		if(CachedData.IsValid() && CachedData->Max() >= CachedData->Num() + Count)
-		{
-			return CachedData;
-		}
-		else
-		{
-			int32 AllocCount = FMath::Max(Count, 128);
-			FlushPendingGameData(PendingIndex);
-			PendingDestGameData[PendingIndex] = MakeShared<FNiagaraDataChannelGameData>(LayoutInfo);
-			PendingDestGameData[PendingIndex]->Reserve(AllocCount);
-			return PendingDestGameData[PendingIndex];	
-		}
-	}
-	else
-	{
-		FNiagaraDataChannelPublishRequest PublishRequest;
-		PublishRequest.bVisibleToGame = bVisibleToGame;
-		PublishRequest.bVisibleToCPUSims = bVisibleToCPU;
-		PublishRequest.bVisibleToGPUSims = bVisibleToGPU;
-		PublishRequest.GameData = MakeShared<FNiagaraDataChannelGameData>(LayoutInfo);
-#if !UE_BUILD_SHIPPING
-		PublishRequest.DebugSource = DebugSource;
-#endif
-		Publish(PublishRequest); 
-		return PublishRequest.GameData;
-	}
-}
-
-void FNiagaraDataChannelData::FlushPendingGameData(int32 Index)
-{
-	if(PendingDestGameData[Index].IsValid() && PendingDestGameData[Index]->Num() > 0)
-	{
-		FNiagaraDataChannelPublishRequest PublishRequest;
-		PublishRequest.bVisibleToGame = Index & 1;
-		PublishRequest.bVisibleToCPUSims = Index & 2;
-		PublishRequest.bVisibleToGPUSims = Index & 4;
-		PublishRequest.GameData = PendingDestGameData[Index];
-#if !UE_BUILD_SHIPPING
-		PublishRequest.DebugSource = TEXT("Accumulated Pending Game Data");
-#endif
-		Publish(PublishRequest);
-	}
-
-	PendingDestGameData[Index] = nullptr;
-}
-
-void FNiagaraDataChannelData::FlushAllPendingGameData()
-{
-	for(int32 i=0; i < PendingDestGameData.Num(); ++i)
-	{
-		FlushPendingGameData(i);
-	}
-}
-
-void FNiagaraDataChannelData::DestroyRenderThreadProxy(FNiagaraGpuComputeDispatchInterface* ComputeDispatchInterface)
-{
-	//We can leave the proxy to be destroyed by the dispatcher ( or in some edge cases perhaps a DI proxy ).
-	RTProxy = nullptr;
-}
-
-void FNiagaraDataChannelData::CreateRenderThreadProxy(UNiagaraDataChannelHandler* Owner)
-{
-	check(Owner);
-	UWorld* OwnerWorld = Owner->GetWorld();	
-	const UNiagaraDataChannel* DataChannel = Owner->GetDataChannel();
-	check(DataChannel);
-
-	if (FNiagaraGpuComputeDispatchInterface* ComputeDispatchInterface = FNiagaraGpuComputeDispatchInterface::Get(OwnerWorld))
-	{
-		RTProxy = MakeShared<FNiagaraDataChannelDataProxy>();
-		RTProxy->GPUDataSet = new FNiagaraDataSet();
-		RTProxy->GPUDataSet->Init(&LayoutInfo->GetDataSetCompiledDataGPU());
-		RTProxy->bNeedsPrevFrameData = DataChannel->KeepPreviousFrameData();
-		RTProxy->Owner = this->AsWeak();
-		RTProxy->LayoutInfo = LayoutInfo;//Keep our own ref to the layout info as the NDCData may die on the GT before the proxy.
-#if !UE_BUILD_SHIPPING
-		RTProxy->DebugName = FString::Printf(TEXT("%s__GPUData"), *GetNameSafe(DataChannel->GetAsset()));
-		RTProxy->DispatchInterfaceForDebuggingOnly = ComputeDispatchInterface;
-#endif
-		ENQUEUE_RENDER_COMMAND(FNiagaraDataChannelDataProxyInit) (
-			[ComputeDispatchInterface, RT_Proxy = RTProxy](FRHICommandListImmediate& CmdList)
-			{
-				ComputeDispatchInterface->AddNDCDataProxy(RT_Proxy);
-			}
-		);
-	}
-}
-
-//////////////////////////////////////////////////////////////////////////
-
-FNiagaraDataChannelLayoutInfo::FNiagaraDataChannelLayoutInfo(const UNiagaraDataChannel* DataChannel)
-{
-	check(DataChannel);
-	GameDataLayout.Init(DataChannel->GetVariables());
-
-	bKeepPreviousFrameData = DataChannel->KeepPreviousFrameData();
-
-	CompiledData.SimTarget = ENiagaraSimTarget::CPUSim;
-	CompiledDataGPU.SimTarget = ENiagaraSimTarget::GPUComputeSim;
-	for (const FNiagaraDataChannelVariable& NDCVar : DataChannel->GetVariables())
-	{
-		FNiagaraVariableBase Var = NDCVar;
-		if (Var.GetType().IsEnum() == false)
-		{
-			Var.SetType(FNiagaraTypeDefinition(FNiagaraTypeHelper::FindNiagaraFriendlyTopLevelStruct(Var.GetType().GetScriptStruct(), ENiagaraStructConversion::Simulation)));
-		}
-		CompiledData.Variables.Add(Var);
-		CompiledDataGPU.Variables.Add(Var);
-	}
-	CompiledData.BuildLayout();
-	CompiledDataGPU.BuildLayout();
-}
-
-FNiagaraDataChannelLayoutInfo::~FNiagaraDataChannelLayoutInfo()
-{
-	CompiledData.Empty();
-	CompiledDataGPU.Empty();
-}
-
-//////////////////////////////////////////////////////////////////////////
-
-void UNiagaraDataChannel::PostInitProperties()
-{
-	Super::PostInitProperties();
-	INiagaraModule::RequestRefreshDataChannels();
-}
-
-void AddVarToHash(const FNiagaraVariable& Var, FBlake3& Builder)
-{
-	FNameBuilder VarName(Var.GetName());
-	FStringView Name = VarName.ToView();
-	uint32 ClassHash = GetTypeHash(Var.GetType().ClassStructOrEnum);
-	Builder.Update(Name.GetData(), Name.Len() * sizeof(TCHAR));
-	Builder.Update(&ClassHash, sizeof(uint32));
-	Builder.Update(&Var.GetType().UnderlyingType, sizeof(uint16));
-}
-
-void UNiagaraDataChannel::PostLoad()
-{
-	Super::PostLoad();
-
-#if WITH_EDITORONLY_DATA
-	static FGuid BaseVersion(TEXT("182b8dd3-f963-477f-a57d-70a449d922d8"));
-	for (const FNiagaraVariable& Var : Variables_DEPRECATED)
-	{
-		FNiagaraDataChannelVariable ChannelVar;
-		ChannelVar.SetName(Var.GetName());
-		ChannelVar.SetType(FNiagaraDataChannelVariable::ToDataChannelType(Var.GetType()));
-
-		FBlake3 VarHashBuilder;
-		VarHashBuilder.Update(&BaseVersion, sizeof(FGuid));
-		AddVarToHash(Var, VarHashBuilder);
-		FBlake3Hash VarHash = VarHashBuilder.Finalize();
-		ChannelVar.Version = FGuid::NewGuidFromHash(VarHash);
-		
-		ChannelVariables.Add(ChannelVar);
-	}
-	Variables_DEPRECATED.Empty();
-	
-	if (!VersionGuid.IsValid())
-	{
-		// If we don't have a guid yet we create one by hashing the existing variables to get a deterministic start guid
-		FBlake3 Builder;
-		Builder.Update(&BaseVersion, sizeof(FGuid));
-		for (const FNiagaraDataChannelVariable& Var : ChannelVariables)
-		{
-			AddVarToHash(Var, Builder);
-		}
-
-		FBlake3Hash Hash = Builder.Finalize();
-		VersionGuid = FGuid::NewGuidFromHash(Hash);
-	}
-#endif
-
-	GetLayoutInfo();
-
-	TransientAccessContext.Init(GetAccessContextType());
-
-	INiagaraModule::RequestRefreshDataChannels();
-}
-
-void UNiagaraDataChannel::BeginDestroy()
-{
-	Super::BeginDestroy();
-	INiagaraModule::RequestRefreshDataChannels();
-
-	RTFence.BeginFence();
-}
-
-bool UNiagaraDataChannel::IsReadyForFinishDestroy()
-{
-	return RTFence.IsFenceComplete() && Super::IsReadyForFinishDestroy();
-}
-
-#if WITH_EDITOR
-
-void UNiagaraDataChannel::PreEditChange(FProperty* PropertyAboutToChange)
-{
-	Super::PreEditChange(PropertyAboutToChange);
-
-	FNiagaraWorldManager::ForAllWorldManagers(
-		[DataChannel = this](FNiagaraWorldManager& WorldMan)
-		{
-			WorldMan.RemoveDataChannel(DataChannel);
-		});
-}
-
-void UNiagaraDataChannel::PostEditChangeProperty(FPropertyChangedEvent& PropertyChangedEvent)
-{
-	FName VariablesMemberName = GET_MEMBER_NAME_CHECKED(UNiagaraDataChannel, ChannelVariables);
-	if (PropertyChangedEvent.ChangeType == EPropertyChangeType::ArrayAdd && PropertyChangedEvent.GetPropertyName() == VariablesMemberName)
-	{
-		int32 ArrayIndex = PropertyChangedEvent.GetArrayIndex(VariablesMemberName.ToString());
-		if(ChannelVariables.IsValidIndex(ArrayIndex))
-		{
-			TSet<FName> ExistingNames;
-			for (const FNiagaraDataChannelVariable& Var : ChannelVariables)
-			{
-				ExistingNames.Add(Var.GetName());
-			}
-			FName UniqueName = FNiagaraUtilities::GetUniqueName(FName("MyNewVar"), ExistingNames);
-			ChannelVariables[ArrayIndex].SetName(UniqueName);
-		}
-	}
-	if (PropertyChangedEvent.ChangeType == EPropertyChangeType::Duplicate && PropertyChangedEvent.GetPropertyName() == VariablesMemberName)
-	{
-		int32 ArrayIndex = PropertyChangedEvent.GetArrayIndex(VariablesMemberName.ToString());
-		if (ChannelVariables.IsValidIndex(ArrayIndex + 1))
-		{
-			TSet<FName> ExistingNames;
-			for (const FNiagaraDataChannelVariable& Var : ChannelVariables)
-			{
-				ExistingNames.Add(Var.GetName());
-			}
-			FNiagaraDataChannelVariable& NewEntry = ChannelVariables[ArrayIndex + 1];
-			FName UniqueName = FNiagaraUtilities::GetUniqueName(NewEntry.GetName(), ExistingNames);
-			NewEntry.SetName(UniqueName);
-			NewEntry.Version = FGuid::NewGuid();
-		}
-	}
-	if (PropertyChangedEvent.GetPropertyName() == VariablesMemberName || PropertyChangedEvent.GetMemberPropertyName() == VariablesMemberName)
-	{
-		VersionGuid = FGuid::NewGuid();
-		// the guid of the variable is updated by the details customization, as we don't want to change it when just renaming a variable
-	}
-	
-	Super::PostEditChangeProperty(PropertyChangedEvent);
-	
-	check(IsInGameThread());
-
-	//Refresh compiled data
-	LayoutInfo = nullptr;
-	GetLayoutInfo();
-
-	TransientAccessContext.Init(GetAccessContextType());
-
-	INiagaraModule::RequestRefreshDataChannels();
-}
-
-#endif//WITH_EDITOR
-
-const FNiagaraDataChannelLayoutInfoPtr UNiagaraDataChannel::GetLayoutInfo()const
-{
-	if(LayoutInfo == nullptr)
-	{
-		LayoutInfo =  MakeShared<FNiagaraDataChannelLayoutInfo>(this);
-	}
-	return LayoutInfo;
-}
-
-FNiagaraDataChannelGameDataPtr UNiagaraDataChannel::CreateGameData()const
-{
-	return MakeShared<FNiagaraDataChannelGameData>(GetLayoutInfo());
-}
-
-bool UNiagaraDataChannel::IsValid()const
-{
-	return LayoutInfo && ChannelVariables.Num() > 0 && LayoutInfo->GetDataSetCompiledData().Variables.Num() == ChannelVariables.Num() && LayoutInfo->GetDataSetCompiledDataGPU().Variables.Num() == ChannelVariables.Num();
-}
-
-//////////////////////////////////////////////////////////////////////////
-
-UNiagaraDataChannelLibrary::UNiagaraDataChannelLibrary(const FObjectInitializer& ObjectInitializer)
-	: Super(ObjectInitializer)
-{
-}
-
-UNiagaraDataChannelHandler* UNiagaraDataChannelLibrary::GetNiagaraDataChannel(const UObject* WorldContextObject, const UNiagaraDataChannelAsset* Channel)
-{
-	return FindDataChannelHandler(WorldContextObject, Channel->Get());
-}
-
-UNiagaraDataChannelWriter* UNiagaraDataChannelLibrary::WriteToNiagaraDataChannel(const UObject* WorldContextObject, const UNiagaraDataChannelAsset* Channel, FNiagaraDataChannelSearchParameters SearchParams, int32 Count, bool bVisibleToGame, bool bVisibleToCPU, bool bVisibleToGPU, const FString& DebugSource)
-{
-	return CreateDataChannelWriter(WorldContextObject, Channel->Get(), SearchParams, Count, bVisibleToGame, bVisibleToCPU, bVisibleToGPU, DebugSource);
-}
-
-UNiagaraDataChannelReader* UNiagaraDataChannelLibrary::ReadFromNiagaraDataChannel(const UObject* WorldContextObject, const UNiagaraDataChannelAsset* Channel, FNiagaraDataChannelSearchParameters SearchParams, bool bReadPreviousFrame)
-{
-	return CreateDataChannelReader(WorldContextObject, Channel->Get(), SearchParams, bReadPreviousFrame);
-}
-
-int32 UNiagaraDataChannelLibrary::GetDataChannelElementCount(const UObject* WorldContextObject, const UNiagaraDataChannelAsset* Channel, FNiagaraDataChannelSearchParameters SearchParams, bool bReadPreviousFrame)
-{
-	if (Channel && Channel->Get())
-	{
-		if (UNiagaraDataChannelReader* Reader = CreateDataChannelReader(WorldContextObject, Channel->Get(), SearchParams, bReadPreviousFrame))
-		{
-			return Reader->Num();
-		}
-	}
-	return 0;
-}
-
-void UNiagaraDataChannelLibrary::ReadFromNiagaraDataChannelSingle(const UObject*, const UNiagaraDataChannelAsset*, int32, FNiagaraDataChannelSearchParameters, bool, ENiagartaDataChannelReadResult&)
-{
-	// this function is just a placeholder and calls into CreateDataChannelReader and its individual read functions from the BP node
-}
-
-void UNiagaraDataChannelLibrary::WriteToNiagaraDataChannelSingle(const UObject*, const UNiagaraDataChannelAsset*, FNiagaraDataChannelSearchParameters, bool, bool, bool)
-{
-	// this function is just a placeholder and calls into CreateDataChannelWriter and its individual write functions from the BP node
-}
-
-//////////////////////////////////////////////////////////////////////////
-// Access Context based functions
-
-UNiagaraDataChannelWriter* UNiagaraDataChannelLibrary::WriteToNiagaraDataChannel_WithContext(const UObject* WorldContextObject, const UNiagaraDataChannelAsset* Channel, FNDCAccessContextInst& AccessContext, int32 Count, bool bVisibleToBlueprint, bool bVisibleToNiagaraCPU, bool bVisibleToNiagaraGPU, const FString& DebugSource)
-{
-	return CreateDataChannelWriter_WithContext(WorldContextObject, Channel->Get(), AccessContext, Count, bVisibleToBlueprint, bVisibleToNiagaraCPU, bVisibleToNiagaraGPU, DebugSource);
-}
-
-UNiagaraDataChannelReader* UNiagaraDataChannelLibrary::ReadFromNiagaraDataChannel_WithContext(const UObject* WorldContextObject, const UNiagaraDataChannelAsset* Channel, FNDCAccessContextInst& AccessContext, bool bReadPreviousFrame)
-{
-	return CreateDataChannelReader_WithContext(WorldContextObject, Channel->Get(), AccessContext, bReadPreviousFrame);
-}
-
-int32 UNiagaraDataChannelLibrary::GetDataChannelElementCount_WithContext(const UObject* WorldContextObject, const UNiagaraDataChannelAsset* Channel, FNDCAccessContextInst& AccessContext, bool bReadPreviousFrame)
-{
-	if (Channel && Channel->Get())
-	{
-		if (UNiagaraDataChannelReader* Reader = CreateDataChannelReader_WithContext(WorldContextObject, Channel->Get(), AccessContext, bReadPreviousFrame))
-		{
-			return Reader->Num();
-		}
-	}
-	return 0;
-}
-
-void UNiagaraDataChannelLibrary::ReadFromNiagaraDataChannelSingle_WithContext(const UObject*, const UNiagaraDataChannelAsset*, int32, FNDCAccessContextInst&, bool, ENiagartaDataChannelReadResult&)
-{
-	// this function is just a placeholder and calls into CreateDataChannelReader_WithContext and its individual read functions from the BP node
-}
-
-void UNiagaraDataChannelLibrary::WriteToNiagaraDataChannelSingle_WithContext(const UObject*, const UNiagaraDataChannelAsset*, FNDCAccessContextInst& AccessContext, bool, bool, bool)
-{
-	// this function is just a placeholder and calls into CreateDataChannelWriter_WithContext and its individual write functions from the BP node
-}
-
-//////////////////////////////////////////////////////////////////////////
-//New Access context nodes
-
-FNDCAccessContextInst UNiagaraDataChannelLibrary::MakeNDCAccessContextInstance(UScriptStruct* ContextStruct)
-{
-	// We should never hit this! stubs to avoid NoExport on the class.
-	checkNoEntry();
-	return FNDCAccessContextInst();
-}
-
-void UNiagaraDataChannelLibrary::SetMembersInNDCAccessContextInstance(FNDCAccessContextInst& AccessContext, UScriptStruct* ContextStruct)
-{
-	// We should never hit this! stubs to avoid NoExport on the class.
-	checkNoEntry();
-}
-
-void UNiagaraDataChannelLibrary::GetMembersInNDCAccessContextInstance(const FNDCAccessContextInst& AccessContext, UScriptStruct* ContextStruct)
-{
-	// We should never hit this! stubs to avoid NoExport on the class.
-	checkNoEntry();
-}
-
-void UNiagaraDataChannelLibrary::GetSinglePropertyInNDCAccessContextInstance(const FNDCAccessContextInst& AccessContext, UScriptStruct* ContextStruct, FName PropertyName, int32& Value)
-{
-	// We should never hit this! stubs to avoid NoExport on the class.
-	checkNoEntry();
-}
-
-void UNiagaraDataChannelLibrary::SetSinglePropertyInNDCAccessContextInstance(FNDCAccessContextInst& AccessContext, UScriptStruct* ContextStruct, FName PropertyName, const int32& Value)
-{
-	// We should never hit this! stubs to avoid NoExport on the class.
-	checkNoEntry();
-}
-
-DEFINE_FUNCTION(UNiagaraDataChannelLibrary::execMakeNDCAccessContextInstance)
-{
-	P_GET_OBJECT(UScriptStruct, ContextStruct);
-	P_FINISH
-
-	if(ContextStruct && ContextStruct->IsChildOf(FNDCAccessContextBase::StaticStruct()))
-	{
-		P_NATIVE_BEGIN;
-		(*(FNDCAccessContextInst*)RESULT_PARAM).Init(TNDCAccessContextType(ContextStruct));		
-		P_NATIVE_END;
-	}
-	else
-	{
-		FBlueprintExceptionInfo ExceptionInfo(
-			EBlueprintExceptionType::AbortExecution,
-			LOCTEXT("NDC_MakeAccessContext_InvalidTypeWarning", "Invalid Value type for Make Niagara Data Channel Access Context. You must use a child of FNDCAccessContextBase"));
-		FBlueprintCoreDelegates::ThrowScriptException(P_THIS, Stack, ExceptionInfo);
-	}
-}
-
-DEFINE_FUNCTION(UNiagaraDataChannelLibrary::execGetSinglePropertyInNDCAccessContextInstance)
-{
-	P_GET_STRUCT_REF(FNDCAccessContextInst, AccessContext);
-	P_GET_OBJECT(UScriptStruct, ContextStruct);
-	P_GET_STRUCT(FName, PropertyName);
-
-	// Read wildcard Value input.
-	Stack.MostRecentPropertyAddress = nullptr;
-	Stack.MostRecentPropertyContainer = nullptr;
-	Stack.StepCompiledIn<FProperty>(nullptr);
-
-	P_FINISH
-
-	const FProperty* ValueProp = CastField<FProperty>(Stack.MostRecentProperty);
-	void* ValuePtr = Stack.MostRecentPropertyAddress;
-
-	const UScriptStruct* ActualContextStruct = AccessContext.GetScriptStruct();
-	const uint8* SrcPtr = AccessContext.AccessContext.GetMemory();
-	
-	if (ActualContextStruct && ActualContextStruct->IsChildOf(FNDCAccessContextBase::StaticStruct()) && ActualContextStruct->IsChildOf(ContextStruct))
-	{
-		P_NATIVE_BEGIN
-		if(FProperty* SrcProperty = ActualContextStruct->FindPropertyByName(PropertyName))
-		{
-			SrcProperty->CopyCompleteValue(ValuePtr, SrcProperty->ContainerPtrToValuePtr<uint8>(SrcPtr));
-		}
-		else
-		{
-			FBlueprintExceptionInfo ExceptionInfo(
-				EBlueprintExceptionType::AbortExecution,
-				FText::Format(LOCTEXT("NDC_AccessContext_MissingSourceProperty_Fmt", "Cannot read {0} from {1}. Property does not exist on this struct type."), FText::FromName(PropertyName), FText::FromString(ActualContextStruct->GetName())));
-			FBlueprintCoreDelegates::ThrowScriptException(P_THIS, Stack, ExceptionInfo);
-		}
-		P_NATIVE_END
-	}
-	else
-	{
-		FBlueprintExceptionInfo ExceptionInfo(
-			EBlueprintExceptionType::NonFatalError,
-			LOCTEXT("NDC_GetAccessContext_InvalidTypeWarning", "Invalid Value type for Get Members In Niagara Data Channel Access Context. You must use a child of FNDCAccessContextBase. Possible that calls with Incompatible types have been mixed."));
-		FBlueprintCoreDelegates::ThrowScriptException(P_THIS, Stack, ExceptionInfo);
-	}
-}
-
-DEFINE_FUNCTION(UNiagaraDataChannelLibrary::execSetSinglePropertyInNDCAccessContextInstance)
-{
-	P_GET_STRUCT_REF(FNDCAccessContextInst, AccessContext);
-	P_GET_OBJECT(UScriptStruct, ContextStruct);
-	P_GET_STRUCT(FName, PropertyName);
-
-	// Read wildcard Value input.
-	Stack.MostRecentPropertyAddress = nullptr;
-	Stack.MostRecentPropertyContainer = nullptr;
-	Stack.StepCompiledIn<FProperty>(nullptr);
-
-	P_FINISH
-
-	const FProperty* ValueProp = CastField<FProperty>(Stack.MostRecentProperty);
-	void* ValuePtr = Stack.MostRecentPropertyAddress;
-
-	const UScriptStruct* ActualContextStruct = AccessContext.GetScriptStruct();
-	uint8* DestPtr = AccessContext.AccessContext.GetMutableMemory();
-
-	if (ActualContextStruct && ActualContextStruct->IsChildOf(FNDCAccessContextBase::StaticStruct()) && ActualContextStruct->IsChildOf(ContextStruct))
-	{
-		P_NATIVE_BEGIN
-			if (FProperty* DestProperty = ActualContextStruct->FindPropertyByName(PropertyName))
-			{
-				DestProperty->CopyCompleteValue(DestProperty->ContainerPtrToValuePtr<uint8>(DestPtr), ValuePtr);
-			}
-			else
-			{
-				FBlueprintExceptionInfo ExceptionInfo(
-					EBlueprintExceptionType::AbortExecution,
-					FText::Format(LOCTEXT("NDC_AccessContext_MissingDestProperty_Fmt", "Cannot write {0} to {1}. Property does not exist on this struct type."), FText::FromName(PropertyName), FText::FromString(ActualContextStruct->GetName())));
-				FBlueprintCoreDelegates::ThrowScriptException(P_THIS, Stack, ExceptionInfo);
-			}
-		P_NATIVE_END
-	}
-	else
-	{
-		FBlueprintExceptionInfo ExceptionInfo(
-			EBlueprintExceptionType::AbortExecution,
-			LOCTEXT("NDC_SetAccessContext_InvalidTypeWarning", "Invalid Value type for Set Niagara Data Channel Access Context Members. You must use a child of FNDCAccessContextBase"));
-		FBlueprintCoreDelegates::ThrowScriptException(P_THIS, Stack, ExceptionInfo);
-	}
-}
-
-//////////////////////////////////////////////////////////////////////////
-// 
-void UNiagaraDataChannelLibrary::SubscribeToNiagaraDataChannel(const UObject* WorldContextObject, const UNiagaraDataChannelAsset* Channel, FNiagaraDataChannelSearchParameters SearchParams, const FOnNewNiagaraDataChannelPublish& UpdateDelegate, int32& UnsubscribeToken)
-{
-	if (UNiagaraDataChannelHandler* NiagaraDataChannelHandler = GetNiagaraDataChannel(WorldContextObject, Channel))
-	{
-		NiagaraDataChannelHandler->SubscribeToDataChannelUpdates(UpdateDelegate, SearchParams, UnsubscribeToken);
-	}
-}
-
-void UNiagaraDataChannelLibrary::SubscribeToNiagaraDataChannel_WithContext(const UObject* WorldContextObject, const UNiagaraDataChannelAsset* Channel, FNDCAccessContextInst& AccessContext, const FOnNewNiagaraDataChannelPublish& UpdateDelegate, int32& UnsubscribeToken)
-{
-	if (UNiagaraDataChannelHandler* NiagaraDataChannelHandler = GetNiagaraDataChannel(WorldContextObject, Channel))
-	{
-		NiagaraDataChannelHandler->SubscribeToDataChannelUpdates_WithContext(UpdateDelegate, AccessContext, UnsubscribeToken);
-	}
-}
-
-void UNiagaraDataChannelLibrary::UnsubscribeFromNiagaraDataChannel(const UObject* WorldContextObject, const UNiagaraDataChannelAsset* Channel, const int32& UnsubscribeToken)
-{
-	if (UNiagaraDataChannelHandler* NiagaraDataChannelHandler = GetNiagaraDataChannel(WorldContextObject, Channel))
-    {
-    	NiagaraDataChannelHandler->UnsubscribeFromDataChannelUpdates(UnsubscribeToken);
-    }
-}
-
-UNiagaraDataChannelHandler* UNiagaraDataChannelLibrary::FindDataChannelHandler(const UObject* WorldContextObject, const UNiagaraDataChannel* Channel)
-{
-	UWorld* World = (WorldContextObject != nullptr) ? GEngine->GetWorldFromContextObject(WorldContextObject, EGetWorldErrorMode::LogAndReturnNull) : nullptr;
-	if (Channel && World)
-	{
-		if (FNiagaraWorldManager* WorldMan = FNiagaraWorldManager::Get(World))
-		{
-			return WorldMan->GetDataChannelManager().FindDataChannelHandler(Channel);
-		}
-	}
-	return nullptr;
-}
-
-UNiagaraDataChannelWriter* UNiagaraDataChannelLibrary::CreateDataChannelWriter(const UObject* WorldContextObject, const UNiagaraDataChannel* Channel, FNiagaraDataChannelSearchParameters SearchParams, int32 Count, bool bVisibleToGame, bool bVisibleToCPU, bool bVisibleToGPU, const FString& DebugSource)
-{
-	check(IsInGameThread());
-	UWorld* World = (WorldContextObject != nullptr) ? GEngine->GetWorldFromContextObject(WorldContextObject, EGetWorldErrorMode::LogAndReturnNull) : nullptr;
-	if (Channel && World && Count > 0)
-	{
-		if (FNiagaraWorldManager* WorldMan = FNiagaraWorldManager::Get(World))
-		{
-			if(UNiagaraDataChannelHandler* Handler = WorldMan->GetDataChannelManager().FindDataChannelHandler(Channel))
-			{
-				if(UNiagaraDataChannelWriter* Writer = Handler->GetDataChannelWriter())
-				{
-					if(Writer->InitWrite(SearchParams, Count, bVisibleToGame, bVisibleToCPU, bVisibleToGPU, DebugSource))
-					{
-						return Writer;
-					}
-				}
-			}
-		}
-	}
-	return nullptr;
-}
-
-UNiagaraDataChannelReader* UNiagaraDataChannelLibrary::CreateDataChannelReader(const UObject* WorldContextObject, const UNiagaraDataChannel* Channel, FNiagaraDataChannelSearchParameters SearchParams, bool bReadPreviousFrame)
-{
-	check(IsInGameThread());
-	UWorld* World = (WorldContextObject != nullptr) ? GEngine->GetWorldFromContextObject(WorldContextObject, EGetWorldErrorMode::LogAndReturnNull) : nullptr;
-	if (Channel && World)
-	{
-		if (FNiagaraWorldManager* WorldMan = FNiagaraWorldManager::Get(World))
-		{
-			if (UNiagaraDataChannelHandler* Handler = WorldMan->GetDataChannelManager().FindDataChannelHandler(Channel))
-			{
-				if (UNiagaraDataChannelReader* Reader = Handler->GetDataChannelReader())
-				{
-					if(Reader->InitAccess(SearchParams, bReadPreviousFrame))
-					{
-						return Reader;
-					}
-				}
-			}
-		}
-	}
-	return nullptr;
-}
-
-UNiagaraDataChannelWriter* UNiagaraDataChannelLibrary::CreateDataChannelWriter_WithContext(const UObject* WorldContextObject, const UNiagaraDataChannel* Channel, FNDCAccessContextInst& AccessContext, int32 Count, bool bVisibleToGame, bool bVisibleToCPU, bool bVisibleToGPU, const FString& DebugSource)
-{
-	check(IsInGameThread());
-	UWorld* World = (WorldContextObject != nullptr) ? GEngine->GetWorldFromContextObject(WorldContextObject, EGetWorldErrorMode::LogAndReturnNull) : nullptr;
-	if (Channel && World && Count > 0 && AccessContext.IsValid())
-	{
-		if (FNiagaraWorldManager* WorldMan = FNiagaraWorldManager::Get(World))
-		{
-			if (UNiagaraDataChannelHandler* Handler = WorldMan->GetDataChannelManager().FindDataChannelHandler(Channel))
-			{
-				if (UNiagaraDataChannelWriter* Writer = Handler->GetDataChannelWriter())
-				{					
-					if (Writer->BeginWrite(AccessContext, Count, bVisibleToGame, bVisibleToCPU, bVisibleToGPU, DebugSource))
-					{
-						return Writer;
-					}
-				}
-			}
-		}
-	}
-	return nullptr;	
-}
-
-UNiagaraDataChannelReader* UNiagaraDataChannelLibrary::CreateDataChannelReader_WithContext(const UObject* WorldContextObject, const UNiagaraDataChannel* Channel, FNDCAccessContextInst& AccessContext, bool bReadPreviousFrame)
-{
-	check(IsInGameThread());
-	UWorld* World = (WorldContextObject != nullptr) ? GEngine->GetWorldFromContextObject(WorldContextObject, EGetWorldErrorMode::LogAndReturnNull) : nullptr;
-	if (Channel && World && AccessContext.IsValid())
-	{
-		if (FNiagaraWorldManager* WorldMan = FNiagaraWorldManager::Get(World))
-		{
-			if (UNiagaraDataChannelHandler* Handler = WorldMan->GetDataChannelManager().FindDataChannelHandler(Channel))
-			{
-				if (UNiagaraDataChannelReader* Reader = Handler->GetDataChannelReader())
-				{
-					if (Reader->BeginRead(AccessContext, bReadPreviousFrame))
-					{
-						return Reader;
-					}
-				}
-			}
-		}
-	}
-	return nullptr;
-}
-
-FNDCAccessContextInst& UNiagaraDataChannelLibrary::GetUsableAccessContextFromNDC(const UNiagaraDataChannelAsset* DataChannel)
-{
-	if(DataChannel && DataChannel->Get())
-	{
-		return DataChannel->Get()->GetTransientAccessContext();
-	}
-	static FNDCAccessContextInst Dummy;
-	return Dummy;
-}
-
-FNDCAccessContextInst& UNiagaraDataChannelLibrary::GetUsableAccessContextFromNDCRef(const FNiagaraDataChannelReference& NDCRef)
-{
-	return NDCRef.GetUsableAccessContext();
-}
-
-FNDCAccessContextInst& UNiagaraDataChannelLibrary::PrepareAccessContextFromNDCRef(FNiagaraDataChannelReference& NDCRef)
-{
-	check(0);//Dummy function replaced by BP compiler. We should not get here.
-	static FNDCAccessContextInst Dummy;
-	return Dummy;
-}
+#include UE_INLINE_GENERATED_CPP_BY_NAME(NiagaraDataChannel)
 
-//////////////////////////////////////////////////////////////////////////
-// Functions for modifying dynamic objects.
-// 
+#define LOCTEXT_NAMESPACE "NiagaraDataChannels"
 
-//Helper to reduce boiler plate of finding an NDC and calling a function on it from BP callable UFunctions.
-template<typename TAction>
-void ExecNDCHandlerAction(const UObject* WorldContextObject, const UNiagaraDataChannelAsset* Channel, TAction&& Func)
+void UNiagaraDataChannel::PostInitProperties()
 {
-	UWorld* World = (WorldContextObject != nullptr) ? GEngine->GetWorldFromContextObject(WorldContextObject, EGetWorldErrorMode::LogAndReturnNull) : nullptr;
-	if (Channel && World)
-	{
-		if (FNiagaraWorldManager* WorldMan = FNiagaraWorldManager::Get(World))
-		{
-			if (UNiagaraDataChannelHandler* Handler = WorldMan->GetDataChannelManager().FindDataChannelHandler(Channel))
-			{
-				Func(Handler);
-			}
-		}
-	}
+	Super::PostInitProperties();
+	INiagaraModule::RequestRefreshDataChannels();
 }
 
-//////////////////////////////////////////////////////////////////////////
-
-#if WITH_EDITOR
-
-void UNiagaraDataChannelAsset::PreEditChange(class FProperty* PropertyAboutToChange)
+void AddVarToHash(const FNiagaraVariable& Var, FBlake3& Builder)
 {
-	Super::PreEditChange(PropertyAboutToChange);
-
-	if(PropertyAboutToChange->GetFName() == GET_MEMBER_NAME_CHECKED(UNiagaraDataChannelAsset, DataChannel) && DataChannel)
-	{
-		CachedPreChangeDataChannel = DataChannel;
-	}
+	FNameBuilder VarName(Var.GetName());
+	FStringView Name = VarName.ToView();
+	uint32 ClassHash = GetTypeHash(Var.GetType().ClassStructOrEnum);
+	Builder.Update(Name.GetData(), Name.Len() * sizeof(TCHAR));
+	Builder.Update(&ClassHash, sizeof(uint32));
+	Builder.Update(&Var.GetType().UnderlyingType, sizeof(uint16));
 }
 
-void UNiagaraDataChannelAsset::PostEditChangeProperty(FPropertyChangedEvent& PropertyChangedEvent)
+void UNiagaraDataChannel::PostLoad()
 {
-	Super::PostEditChangeProperty(PropertyChangedEvent);
+	Super::PostLoad();
 
-	if(CachedPreChangeDataChannel && DataChannel)
+#if WITH_EDITORONLY_DATA
+	static FGuid BaseVersion(TEXT("182b8dd3-f963-477f-a57d-70a449d922d8"));
+	for (const FNiagaraVariable& Var : Variables_DEPRECATED)
 	{
-		UEngine::FCopyPropertiesForUnrelatedObjectsParams Params;
-		UEngine::CopyPropertiesForUnrelatedObjects(CachedPreChangeDataChannel, DataChannel, Params);
-		CachedPreChangeDataChannel = nullptr;
-	}
-}
-
-#endif
-
-//////////////////////////////////////////////////////////////////////////
-
-const FNDCAccessContextInst& FNiagaraDataChannelReference::GetAccessContext()const
-{
-	InitAccessContext();
-	return AccessContext;
-}
+		FNiagaraDataChannelVariable ChannelVar;
+		ChannelVar.SetName(Var.GetName());
+		ChannelVar.SetType(FNiagaraDataChannelVariable::ToDataChannelType(Var.GetType()));
 
-FNDCAccessContextInst& FNiagaraDataChannelReference::GetUsableAccessContext()const
-{
-	if (DataChannel)
-	{
-		if (UNiagaraDataChannel* NDC = DataChannel->Get())
-		{
-			InitAccessContext();
-			FNDCAccessContextInst& TransientContext = NDC->GetTransientAccessContext();
-			if(bCustomAccessContext && AccessContext.IsValid())
-			{
-				TransientContext = AccessContext;
-			}
-			return TransientContext;
-		}
+		FBlake3 VarHashBuilder;
+		VarHashBuilder.Update(&BaseVersion, sizeof(FGuid));
+		AddVarToHash(Var, VarHashBuilder);
+		FBlake3Hash VarHash = VarHashBuilder.Finalize();
+		ChannelVar.Version = FGuid::NewGuidFromHash(VarHash);
+		
+		ChannelVariables.Add(ChannelVar);
 	}
-	static FNDCAccessContextInst Dummy;
-	return Dummy;
-}
-
-void FNiagaraDataChannelReference::InitAccessContext(bool bTransferExisting/* =true */)const
-{
-	//We need to verify that the access context has the correct type.
-	//We can get bad types if someone changes the NDC type for example.
-	if(DataChannel)
+	Variables_DEPRECATED.Empty();
+	
+	if (!VersionGuid.IsValid())
 	{
-		if (UNiagaraDataChannel* NDC = DataChannel->Get())
+		// If we don't have a guid yet we create one by hashing the existing variables to get a deterministic start guid
+		FBlake3 Builder;
+		Builder.Update(&BaseVersion, sizeof(FGuid));
+		for (const FNiagaraDataChannelVariable& Var : ChannelVariables)
 		{
-			const UScriptStruct* ExpectedType = NDC->GetAccessContextType().Get();
-			if(bCustomAccessContext && ExpectedType != FNDCAccessContextLegacy::StaticStruct())//Special case legacy type indicating we should not create a context in most places.
-			{
-				if (AccessContext.IsValid() && bTransferExisting)
-				{
-					const UScriptStruct* ActualType = AccessContext.GetScriptStruct();
-					if (ActualType != ExpectedType)
-					{
-						//If we do have the wrong type. 
-						//Initialize to the correct type but try to pull any matching properties from the old type.
-						FInstancedStruct OldContext = AccessContext.AccessContext;
-						AccessContext.Init(NDC->GetAccessContextType());
-
-						const uint8* Src = OldContext.GetMemory();
-						uint8* Dst = AccessContext.AccessContext.GetMutableMemory();
-						for (TFieldIterator<FProperty> It(ExpectedType); It; ++It)
-						{
-							FProperty* DestProp = *It;
-							FProperty* SrcProp = ActualType->FindPropertyByName(DestProp->GetFName());
-							if (SrcProp && SrcProp->SameType(DestProp))
-							{
-								const uint8* PropSrcPtr = SrcProp->ContainerPtrToValuePtr<const uint8>(Src);
-								uint8* PropDstPtr = DestProp->ContainerPtrToValuePtr<uint8>(Dst);
-								DestProp->CopyCompleteValue(PropDstPtr, PropSrcPtr);
-							}
-						}
-					}
-				}
-				else
-				{
-					AccessContext.Init(NDC->GetAccessContextType());
-				}
-				return;
-			}
+			AddVarToHash(Var, Builder);
 		}
-	}
 
-	//If we fail to init properly then clear the access context to avoid memory bloat.
-	AccessContext.Reset();
-}
-
-void FNiagaraDataChannelReference::SetDataChannel(UNiagaraDataChannelAsset* InDataChannel)
-{
-	if(DataChannel != InDataChannel)
-	{
-		DataChannel = InDataChannel;
-		OnValueChanged();
+		FBlake3Hash Hash = Builder.Finalize();
+		VersionGuid = FGuid::NewGuidFromHash(Hash);
 	}
-}
+#endif
 
-void FNiagaraDataChannelReference::OnValueChanged()
-{
-	if(DataChannel)
-	{
-		InitAccessContext();
-	}
-	else
-	{
-		AccessContext.Reset();
-	}
-}
+	GetLayoutInfo();
 
-//////////////////////////////////////////////////////////////////////////
+	TransientAccessContext.Init(GetAccessContextType());
 
-bool FNiagaraDataChannelVariable::Serialize(FArchive& Ar)
-{
-	if (FNiagaraVariableBase::Serialize(Ar))
-	{
-		if (Ar.IsLoading())
-		{
-			// fix up variables serialized with wrong type
-			// this happens because we only save swc types
-			SetType(ToDataChannelType(GetType()));
-		}
-		return true;
-	}
-	return false;
+	INiagaraModule::RequestRefreshDataChannels();
 }
 
-#if WITH_EDITORONLY_DATA
-bool FNiagaraDataChannelVariable::IsAllowedType(const FNiagaraTypeDefinition& Type)
+void UNiagaraDataChannel::BeginDestroy()
 {
-	return !(Type.IsDataInterface() || Type.GetClass() || Type == FNiagaraTypeDefinition::GetParameterMapDef() || Type == FNiagaraTypeDefinition::GetGenericNumericDef() || Type == FNiagaraTypeDefinition::GetHalfDef() || Type == FNiagaraTypeDefinition::GetMatrix4Def());
-}
-#endif
+	Super::BeginDestroy();
+	INiagaraModule::RequestRefreshDataChannels();
 
-FNiagaraTypeDefinition FNiagaraDataChannelVariable::ToDataChannelType(const FNiagaraTypeDefinition& Type)
-{
-	if (Type == FNiagaraTypeDefinition::GetVec3Def())
-	{
-		return FNiagaraTypeHelper::GetVectorDef();
-	}
-	if (Type == FNiagaraTypeDefinition::GetFloatDef())
-	{
-		return FNiagaraTypeHelper::GetDoubleDef();
-	}
-	if (Type == FNiagaraTypeDefinition::GetQuatDef())
-	{
-		return FNiagaraTypeHelper::GetQuatDef();
-	}
-	if (Type == FNiagaraTypeDefinition::GetVec2Def())
-	{
-		return FNiagaraTypeHelper::GetVector2DDef();
-	}
-	if (Type == FNiagaraTypeDefinition::GetVec4Def())
-	{
-		return FNiagaraTypeHelper::GetVector4Def();
-	}
-	return Type;
+	RTFence.BeginFence();
 }
 
-//////////////////////////////////////////////////////////////////////////
-
-#if WITH_NIAGARA_DEBUGGER
-void FNiagaraDataChannelDebugUtilities::BeginFrame(FNiagaraWorldManager* WorldMan, float DeltaSeconds)
+bool UNiagaraDataChannel::IsReadyForFinishDestroy()
 {
-	check(WorldMan);
-	WorldMan->GetDataChannelManager().BeginFrame(DeltaSeconds);
+	return RTFence.IsFenceComplete() && Super::IsReadyForFinishDestroy();
 }
 
-void FNiagaraDataChannelDebugUtilities::EndFrame(FNiagaraWorldManager* WorldMan, float DeltaSeconds)
-{
-	check(WorldMan);
-	WorldMan->GetDataChannelManager().EndFrame(DeltaSeconds);
-}
+#if WITH_EDITOR
 
-void FNiagaraDataChannelDebugUtilities::Tick(FNiagaraWorldManager* WorldMan, float DeltaSeconds, ETickingGroup TickGroup)
+void UNiagaraDataChannel::PreEditChange(FProperty* PropertyAboutToChange)
 {
-	check(WorldMan);
-	WorldMan->GetDataChannelManager().Tick(DeltaSeconds, TickGroup);
-}
+	Super::PreEditChange(PropertyAboutToChange);
 
-UNiagaraDataChannelHandler* FNiagaraDataChannelDebugUtilities::FindDataChannelHandler(FNiagaraWorldManager* WorldMan, UNiagaraDataChannel* DataChannel)
-{
-	check(WorldMan);
-	return WorldMan->GetDataChannelManager().FindDataChannelHandler(DataChannel);
+	FNiagaraWorldManager::ForAllWorldManagers(
+		[DataChannel = this](FNiagaraWorldManager& WorldMan)
+		{
+			WorldMan.RemoveDataChannel(DataChannel);
+		});
 }
 
-void FNiagaraDataChannelDebugUtilities::LogWrite(const FNiagaraDataChannelPublishRequest& WriteRequest, const UNiagaraDataChannel* DataChannel, const ETickingGroup& TickGroup)
+void UNiagaraDataChannel::PostEditChangeProperty(FPropertyChangedEvent& PropertyChangedEvent)
 {
-	FNiagaraDataChannelDebugUtilities& Debugger = Get();
-	if (NDCCVars::FrameDataToCapture > 0)
+	FName VariablesMemberName = GET_MEMBER_NAME_CHECKED(UNiagaraDataChannel, ChannelVariables);
+	if (PropertyChangedEvent.ChangeType == EPropertyChangeType::ArrayAdd && PropertyChangedEvent.GetPropertyName() == VariablesMemberName)
 	{
-		if (Debugger.FrameData.Num() > NDCCVars::FrameDataToCapture)
-		{
-			Debugger.FrameData.SetNum(NDCCVars::FrameDataToCapture);
-		}
-		FFrameDebugData* Data;
-		if (Debugger.FrameData.Num() && Debugger.FrameData.Last().FrameNumber == GFrameCounter)
-		{
-			Data = &Debugger.FrameData.Last();
-		} 
-		else
+		int32 ArrayIndex = PropertyChangedEvent.GetArrayIndex(VariablesMemberName.ToString());
+		if(ChannelVariables.IsValidIndex(ArrayIndex))
 		{
-			if (Debugger.FrameData.Num() == NDCCVars::FrameDataToCapture)
+			TSet<FName> ExistingNames;
+			for (const FNiagaraDataChannelVariable& Var : ChannelVariables)
 			{
-				Debugger.FrameData.RemoveAt(0);
+				ExistingNames.Add(Var.GetName());
 			}
-			Data = &Debugger.FrameData.AddDefaulted_GetRef();
-			Data->FrameNumber = GFrameCounter;
-		}
-		FChannelWriteRequest DebugData;
-		DebugData.Channel = DataChannel;
-		DebugData.DebugSource = WriteRequest.DebugSource;
-		DebugData.bVisibleToGame = WriteRequest.bVisibleToGame;
-		DebugData.bVisibleToCPUSims = WriteRequest.bVisibleToCPUSims;
-		DebugData.bVisibleToGPUSims = WriteRequest.bVisibleToGPUSims;
-		DebugData.TickGroup = TickGroup;
-		if (WriteRequest.GameData.IsValid())
-		{
-			DebugData.Data = WriteRequest.GameData;
-		}
-		else if (ensure(WriteRequest.Data))
-		{
-			DebugData.Data = MakeShared<FNiagaraDataChannelGameData>(DataChannel->GetLayoutInfo());
-			DebugData.Data->AppendFromDataSet(WriteRequest.Data, WriteRequest.LwcTile);
-		}
-		Data->WriteRequests.Add(DebugData);
-	}
-	else
-	{
-		if (Debugger.FrameData.Num())
-		{
-			Debugger.FrameData.Empty();
+			FName UniqueName = FNiagaraUtilities::GetUniqueName(FName("MyNewVar"), ExistingNames);
+			ChannelVariables[ArrayIndex].SetName(UniqueName);
 		}
 	}
-	
-	if (NDCCVars::LogWritesToOutputLog > 0)
+	if (PropertyChangedEvent.ChangeType == EPropertyChangeType::Duplicate && PropertyChangedEvent.GetPropertyName() == VariablesMemberName)
 	{
-		FString DataString;
-		uint32 NumInsts = 0;
-		if (FNiagaraDataChannelGameData* RequestGameData = WriteRequest.GameData.Get())
-		{
-			NumInsts = RequestGameData->Num();
-			if (NDCCVars::LogWritesToOutputLog > 1)
-			{
-				DataString = ToJson(RequestGameData);
-			}
-		}
-		else if (ensure(WriteRequest.Data))
+		int32 ArrayIndex = PropertyChangedEvent.GetArrayIndex(VariablesMemberName.ToString());
+		if (ChannelVariables.IsValidIndex(ArrayIndex + 1))
 		{
-			NumInsts = WriteRequest.Data->GetNumInstances();
-			if (NDCCVars::LogWritesToOutputLog > 1)
+			TSet<FName> ExistingNames;
+			for (const FNiagaraDataChannelVariable& Var : ChannelVariables)
 			{
-				FNiagaraDataChannelGameData TempData(DataChannel->GetLayoutInfo());
-				TempData.AppendFromDataSet(WriteRequest.Data, WriteRequest.LwcTile);
-				DataString = ToJson(&TempData);
+				ExistingNames.Add(Var.GetName());
 			}
+			FNiagaraDataChannelVariable& NewEntry = ChannelVariables[ArrayIndex + 1];
+			FName UniqueName = FNiagaraUtilities::GetUniqueName(NewEntry.GetName(), ExistingNames);
+			NewEntry.SetName(UniqueName);
+			NewEntry.Version = FGuid::NewGuid();
 		}
-		
-		UE_LOG(LogNiagara, Log, TEXT("Frame %llu, TG %s, NDC write by %s (BP[%s]/CPU[%s]/GPU[%s]): %i entries to data channel %s %s%s")
-			, GFrameCounter
-			, *TickGroupToString(TickGroup)
-			, *WriteRequest.DebugSource
-			, WriteRequest.bVisibleToGame ? TEXT("X") : TEXT(" ")
-			, WriteRequest.bVisibleToCPUSims ? TEXT("X") : TEXT(" ")
-			, WriteRequest.bVisibleToGPUSims ? TEXT("X") : TEXT(" ")
-			, NumInsts
-			, *GetPathNameSafe(DataChannel)
-			, DataString.IsEmpty() ? TEXT("") : TEXT("\n")
-			, *DataString);
 	}
-}
-
-void FNiagaraDataChannelDebugUtilities::DumpAllWritesToLog()
-{
-	FNiagaraDataChannelDebugUtilities& Debugger = Get();
-	if (Debugger.FrameData.Num() == 0)
+	if (PropertyChangedEvent.GetPropertyName() == VariablesMemberName || PropertyChangedEvent.GetMemberPropertyName() == VariablesMemberName)
 	{
-		UE_LOG(LogNiagara, Log, TEXT("No writes are currently stored in the log. fx.Niagara.DataChannels.FrameDataToCapture = %i"), NDCCVars::FrameDataToCapture);
-		return;
+		VersionGuid = FGuid::NewGuid();
+		// the guid of the variable is updated by the details customization, as we don't want to change it when just renaming a variable
 	}
-	UE_LOG(LogNiagara, Log, TEXT("Current Frame is %llu, logging data from oldest to newest:"), GFrameCounter);
-	for (const FFrameDebugData& FrameData : Debugger.FrameData)
-	{
-		UE_LOG(LogNiagara, Log, TEXT("Frame %llu: %i entries"), FrameData.FrameNumber, FrameData.WriteRequests.Num());
-		for (const FChannelWriteRequest& Request : FrameData.WriteRequests)
-		{
-			FString DataString = ToJson(Request.Data.Get());
-			UE_LOG(LogNiagara, Log, TEXT("Write by %s (BP[%s]/CPU[%s]/GPU[%s], TG %s): %i entries to data channel %s \n%s")
-			, *Request.DebugSource
-			, Request.bVisibleToGame ? TEXT("X") : TEXT(" ")
-			, Request.bVisibleToCPUSims ? TEXT("X") : TEXT(" ")
-			, Request.bVisibleToGPUSims ? TEXT("X") : TEXT(" ")
-			, *TickGroupToString(Request.TickGroup)
-			, Request.Data->Num()
-			, *GetPathNameSafe(Request.Channel.Get())
-			, *DataString);
-		}
-		UE_LOG(LogNiagara, Log, TEXT("----------------------------------------------"));
-	} 
-}
+	
+	Super::PostEditChangeProperty(PropertyChangedEvent);
+	
+	check(IsInGameThread());
 
-FNiagaraDataChannelDebugUtilities& FNiagaraDataChannelDebugUtilities::Get()
-{
-	return TLazySingleton<FNiagaraDataChannelDebugUtilities>::Get();
-}
+	//Refresh compiled data
+	LayoutInfo = nullptr;
+	GetLayoutInfo();
 
-void FNiagaraDataChannelDebugUtilities::TearDown()
-{
-	return TLazySingleton<FNiagaraDataChannelDebugUtilities>::TearDown();
-}
+	TransientAccessContext.Init(GetAccessContextType());
 
-FString FNiagaraDataChannelDebugUtilities::ToJson(FNiagaraDataChannelGameData* Data)
-{
-	FString JsonString;
-	TSharedRef<TJsonWriter<>> JsonWriter = TJsonWriterFactory<>::Create(&JsonString);
-	JsonWriter->WriteArrayStart();
+	INiagaraModule::RequestRefreshDataChannels();
+}
 
-	TConstArrayView<FNiagaraDataChannelVariableBuffer> VaraibleBuffers = Data->GetVariableBuffers();
+#endif//WITH_EDITOR
 
-	const FNiagaraDataChannelGameDataLayout& GameDataLayout = Data->GetLayoutInfo()->GetGameDataLayout();
-	for (int32 i = 0; i < Data->Num(); i++)
+const FNiagaraDataChannelLayoutInfoPtr UNiagaraDataChannel::GetLayoutInfo()const
+{
+	if(LayoutInfo == nullptr)
 	{
-		TSharedPtr<FJsonObject> JsonObject = MakeShared<FJsonObject>();
-		for (auto& LayoutVarPair : GameDataLayout.VariableIndices)
-		{
-			const FNiagaraVariableBase& Var = LayoutVarPair.Key;
-			int32 VarIndex = LayoutVarPair.Value;
-
-			FString VarName = Var.GetName().ToString();
-			const FNiagaraDataChannelVariableBuffer* Buffer = nullptr;
-			
-			if(VaraibleBuffers.IsValidIndex(VarIndex))
-			{
-				Buffer = &VaraibleBuffers[VarIndex];
-			}
-			
-			if (Var.GetType() == FNiagaraTypeHelper::GetDoubleDef() && Buffer)
-			{
-				double Value;
-				Buffer->Read<double>(i, Value, false);
-				JsonObject->SetNumberField(VarName, Value);
-			}
-			else if (Var.GetType() == FNiagaraTypeDefinition::GetBoolDef() && Buffer)
-			{
-				FNiagaraBool Value;
-				Buffer->Read<FNiagaraBool>(i, Value, false);
-				JsonObject->SetBoolField(VarName, Value);
-			}
-			else if (Var.GetType() == FNiagaraTypeDefinition::GetIntDef() && Buffer)
-			{
-				int32 Value;
-				Buffer->Read<int32>(i, Value, false);
-				JsonObject->SetNumberField(VarName, Value);
-			}
-			else if ((Var.GetType() == FNiagaraTypeHelper::GetVectorDef() || Var.GetType() == FNiagaraTypeDefinition::GetPositionDef()) && Buffer)
-			{
-				FVector Value;
-				Buffer->Read<FVector>(i, Value, false);
-				JsonObject->SetStringField(VarName, Value.ToString());
-			}
-			else if (Var.GetType() == FNiagaraTypeHelper::GetVector2DDef() && Buffer)
-			{
-				FVector2D Value;
-				Buffer->Read<FVector2D>(i, Value, false);
-				JsonObject->SetStringField(VarName, Value.ToString());
-			}
-			else if (Var.GetType() == FNiagaraTypeHelper::GetVector4Def() && Buffer)
-			{
-				FVector4 Value;
-				Buffer->Read<FVector4>(i, Value, false);
-				JsonObject->SetStringField(VarName, Value.ToString());
-			}
-			else if (Var.GetType() == FNiagaraTypeHelper::GetQuatDef() && Buffer)
-			{
-				FQuat Value;
-				Buffer->Read<FQuat>(i, Value, false);
-				JsonObject->SetStringField(VarName, Value.ToString());
-			}
-			else if (Var.GetType() == FNiagaraTypeDefinition::GetColorDef() && Buffer)
-			{
-				FLinearColor Value;
-				Buffer->Read<FLinearColor>(i, Value, false);
-				JsonObject->SetStringField(VarName, Value.ToString());
-			}
-			else if (Var.GetType() == FNiagaraTypeDefinition::GetIDDef() && Buffer)
-			{
-				FNiagaraID Value;
-				Buffer->Read<FNiagaraID>(i, Value, false);
-				JsonObject->SetStringField(VarName, FString::FromInt(Value.Index) + "/" + FString::FromInt(Value.AcquireTag));
-			}
-			else if (Var.GetType().IsEnum() && Buffer)
-			{
-				int32 Value;
-				Buffer->Read<int32>(i, Value, false);
-				JsonObject->SetStringField(VarName, Var.GetType().GetEnum()->GetNameByValue(Value).ToString());
-			}
-			else
-			{
-				JsonObject->SetStringField(VarName, "???");
-			}
-		}
-		FJsonSerializer::Serialize(JsonObject.ToSharedRef(), JsonWriter);
+		LayoutInfo =  MakeShared<FNiagaraDataChannelLayoutInfo>(this);
 	}
-	JsonWriter->WriteArrayEnd();
-	JsonWriter->Close();
-	return JsonString;
+	return LayoutInfo;
 }
 
-FString FNiagaraDataChannelDebugUtilities::TickGroupToString(const ETickingGroup& TickGroup)
+FNiagaraDataChannelGameDataPtr UNiagaraDataChannel::CreateGameData()const
 {
-	static UEnum* TGEnum = StaticEnum<ETickingGroup>();
-	return TGEnum->GetDisplayNameTextByValue(TickGroup).ToString();
+	return MakeShared<FNiagaraDataChannelGameData>(GetLayoutInfo());
 }
-#endif //WITH_NIAGARA_DEBUGGER
 
-#undef LOCTEXT_NAMESPACE
+bool UNiagaraDataChannel::IsValid()const
+{
+	return LayoutInfo && ChannelVariables.Num() > 0 && LayoutInfo->GetDataSetCompiledData().Variables.Num() == ChannelVariables.Num() && LayoutInfo->GetDataSetCompiledDataGPU().Variables.Num() == ChannelVariables.Num();
+}
 
+#undef LOCTEXT_NAMESPACE
\ No newline at end of file
diff --git a/Engine/Plugins/FX/Niagara/Source/Niagara/Private/NiagaraDataChannelAccessor.cpp b/Engine/Plugins/FX/Niagara/Source/Niagara/Private/NiagaraDataChannelAccessor.cpp
--- a/Engine/Plugins/FX/Niagara/Source/Niagara/Private/NiagaraDataChannelAccessor.cpp
+++ b/Engine/Plugins/FX/Niagara/Source/Niagara/Private/NiagaraDataChannelAccessor.cpp
@@ -6,18 +6,16 @@
 #include "UObject/UObjectGlobals.h"
 #include "UObject/Package.h"
 
-#include "NiagaraDataChannelPublic.h"
-#include "NiagaraDataChannelCommon.h"
 #include "NiagaraDataChannel.h"
 #include "NiagaraDataChannelHandler.h"
 #include "NiagaraDataChannelManager.h"
+#include "NiagaraDataChannelData.h"
 
 #include "NiagaraWorldManager.h"
 #include "NiagaraDataChannelAccessorImpl.h"
 
 #include UE_INLINE_GENERATED_CPP_BY_NAME(NiagaraDataChannelAccessor)
 
-
 namespace NDCCVars
 {
 #if DEBUG_NDC_ACCESS
diff --git a/Engine/Plugins/FX/Niagara/Source/Niagara/Private/NiagaraDataChannelCommon.cpp b/Engine/Plugins/FX/Niagara/Source/Niagara/Private/NiagaraDataChannelCommon.cpp
--- a/Engine/Plugins/FX/Niagara/Source/Niagara/Private/NiagaraDataChannelCommon.cpp
+++ b/Engine/Plugins/FX/Niagara/Source/Niagara/Private/NiagaraDataChannelCommon.cpp
@@ -1,20 +1,3 @@
 // Copyright Epic Games, Inc. All Rights Reserved.
 
 #include "NiagaraDataChannelCommon.h"
-#include "Components/SceneComponent.h"
-#include "NiagaraComponent.h"
-
-FNDCSpawnedSystemRef::FNDCSpawnedSystemRef(UNiagaraComponent* Component)
-	: SpawnedSystem(Component)
-{
-}
-
-FVector FNDCAccessContext::GetLocation()const
-{
-	return OwningComponent && !bOverrideLocation ? OwningComponent->GetComponentLocation() : Location;
-}
-
-FVector FNiagaraDataChannelSearchParameters::GetLocation()const
-{
-	return OwningComponent && !bOverrideLocation ? OwningComponent->GetComponentLocation() : Location;
-}
diff --git a/Engine/Plugins/FX/Niagara/Source/Niagara/Private/NiagaraDataChannelHandler.cpp b/Engine/Plugins/FX/Niagara/Source/Niagara/Private/NiagaraDataChannelHandler.cpp
--- a/Engine/Plugins/FX/Niagara/Source/Niagara/Private/NiagaraDataChannelHandler.cpp
+++ b/Engine/Plugins/FX/Niagara/Source/Niagara/Private/NiagaraDataChannelHandler.cpp
@@ -3,6 +3,7 @@
 #include "NiagaraDataChannelHandler.h"
 #include "NiagaraDataChannelAccessor.h"
 #include "NiagaraDataChannelCommon.h"
+#include "NiagaraDataChannelData.h"
 #include "Logging/StructuredLog.h"
 #include "NiagaraGpuComputeDispatchInterface.h"
 
diff --git a/Engine/Plugins/FX/Niagara/Source/Niagara/Private/NiagaraDataChannelManager.h b/Engine/Plugins/FX/Niagara/Source/Niagara/Private/NiagaraDataChannelManager.h
--- a/Engine/Plugins/FX/Niagara/Source/Niagara/Private/NiagaraDataChannelManager.h
+++ b/Engine/Plugins/FX/Niagara/Source/Niagara/Private/NiagaraDataChannelManager.h
@@ -3,7 +3,7 @@
 #pragma once
 
 #include "UObject/ObjectPtr.h"
-#include "NiagaraDataChannelPublic.h"
+#include "NiagaraDataChannelAsset.h"
 
 class FNiagaraWorldManager;
 class UNiagaraDataChannel;
diff --git a/Engine/Plugins/FX/Niagara/Source/Niagara/Private/NiagaraDataChannel_Map.cpp b/Engine/Plugins/FX/Niagara/Source/Niagara/Private/NiagaraDataChannel_Map.cpp
--- a/Engine/Plugins/FX/Niagara/Source/Niagara/Private/NiagaraDataChannel_Map.cpp
+++ b/Engine/Plugins/FX/Niagara/Source/Niagara/Private/NiagaraDataChannel_Map.cpp
@@ -6,7 +6,7 @@
 #include "NiagaraWorldManager.h"
 #include "NiagaraComponent.h"
 #include "NiagaraDataChannelCommon.h"
-
+#include "NiagaraDataChannelData.h"
 #include "NiagaraFunctionLibrary.h"
 
 namespace NDCMapLocal
diff --git a/Engine/Plugins/FX/Niagara/Source/Niagara/Private/NiagaraGpuComputeDispatch.cpp b/Engine/Plugins/FX/Niagara/Source/Niagara/Private/NiagaraGpuComputeDispatch.cpp
--- a/Engine/Plugins/FX/Niagara/Source/Niagara/Private/NiagaraGpuComputeDispatch.cpp
+++ b/Engine/Plugins/FX/Niagara/Source/Niagara/Private/NiagaraGpuComputeDispatch.cpp
@@ -38,7 +38,7 @@
 #include "GPUSkinCache.h"
 #include "SkeletalMeshUpdater.h"
 
-#include "NiagaraDataChannel.h"
+#include "NiagaraDataChannelData.h"
 
 DECLARE_CYCLE_STAT(TEXT("GPU Dispatch Setup [RT]"), STAT_NiagaraGPUDispatchSetup_RT, STATGROUP_Niagara);
 DECLARE_CYCLE_STAT(TEXT("GPU Emitter Dispatch [RT]"), STAT_NiagaraGPUSimTick_RT, STATGROUP_Niagara);
diff --git a/Engine/Plugins/FX/Niagara/Source/Niagara/Private/NiagaraModule.cpp b/Engine/Plugins/FX/Niagara/Source/Niagara/Private/NiagaraModule.cpp
--- a/Engine/Plugins/FX/Niagara/Source/Niagara/Private/NiagaraModule.cpp
+++ b/Engine/Plugins/FX/Niagara/Source/Niagara/Private/NiagaraModule.cpp
@@ -11,6 +11,7 @@
 #include "NiagaraEvents.h"
 #include "NiagaraSettings.h"
 #include "NiagaraDataChannelManager.h"
+#include "NiagaraDataChannelDebugUtilities.h"
 #include "NiagaraSimCache.h"
 #include "NiagaraSimulationTaskPriority.h"
 #include "NiagaraUseOpenVDB.h"
