Bit of refactoring classes out of NiagaraDataChannels.cpp/h

diff --git a/Engine/Plugins/FX/Niagara/Source/Niagara/Public/NiagaraDataChannel.h b/Engine/Plugins/FX/Niagara/Source/Niagara/Public/NiagaraDataChannel.h
--- a/Engine/Plugins/FX/Niagara/Source/Niagara/Public/NiagaraDataChannel.h
+++ b/Engine/Plugins/FX/Niagara/Source/Niagara/Public/NiagaraDataChannel.h
@@ -27,136 +27,22 @@ Reading the previous frames data introduces a frame of latency but removes the n
 
 #pragma once
 
-#include "NiagaraDataChannelPublic.h"
-#include "Kismet/BlueprintFunctionLibrary.h"
-#include "NiagaraDataSetCompiledData.h"
-#include "NiagaraCommon.h"
-#include "NiagaraDataSet.h"
 #include "UObject/UObjectIterator.h"
 #include "RenderCommandFence.h"
+
+#include "NiagaraDataChannelCommon.h"
+#include "NiagaraDataChannelAsset.h"
+#include "NiagaraDataChannelVariable.h"
+#include "NiagaraDataChannelAccessContext.h"
+
 #include "NiagaraDataChannel.generated.h"
 
 #define UE_API NIAGARA_API
 
 DECLARE_STATS_GROUP(TEXT("Niagara Data Channels"), STATGROUP_NiagaraDataChannels, STATCAT_Niagara);
 
-class FRHICommandListImmediate;
-
-class FNiagaraGpuComputeDispatchInterface;
-struct FNiagaraDataSetCompiledData;
-
-class UNiagaraDataChannelHandler;
-class UNiagaraDataChannelWriter;
-class UNiagaraDataChannelReader;
-struct FNiagaraDataChannelPublishRequest;
-struct FNiagaraDataChannelGameDataLayout;
-class FNiagaraGpuReadbackManager;
-class FRDGBuilder;
-class UNiagaraDataChannel;
-
 //////////////////////////////////////////////////////////////////////////
 
-struct FNDCGpuReadbackInfo
-{
-	FNiagaraDataBufferRef Buffer;
-	bool bPublishToCPU = false;
-	bool bPublishToGame = false;
-	FVector3f LWCTile;
-};
-
-using FNiagaraDataChannelDataProxyPtr = TSharedPtr<struct FNiagaraDataChannelDataProxy>;
-
-/** Render thread proxy of FNiagaraDataChannelData. */
-struct FNiagaraDataChannelDataProxy : public TSharedFromThis<FNiagaraDataChannelDataProxy>
-{
-	~FNiagaraDataChannelDataProxy();
-
-	TWeakPtr<FNiagaraDataChannelData> Owner;
-	FNiagaraDataSet* GPUDataSet = nullptr;
-	FNiagaraDataBufferRef CurrFrameData = nullptr;
-	FNiagaraDataBufferRef PrevFrameData = nullptr;
-	bool bNeedsPrevFrameData = false;
-
-	//Keeping layout info ref to ensure lifetime for GPUDataSet.
-	FNiagaraDataChannelLayoutInfoPtr LayoutInfo;
-
-	//Buffers coming from the CPU that we're going to copy up for reading on the GPU
-	TArray<FNiagaraDataBufferRef> PendingCPUBuffers;
-
-	//Buffers written from the GPU that we must send back to the CPU.
-	TArray<FNDCGpuReadbackInfo> PendingGPUReadbackBuffers;
-	
-	//Users that need space in this NDC Data add to this for each tick via AddGPUAllocationForNextTick().
-	int32 PendingGPUAllocations = 0;
-
-	//Track current read/write counts +ve for readers, -ve for writers. We cannot mix readers and writers in the same buffer in the same stage.
-	int32 CurrBufferAccessCounts = 0;
-
-	#if !UE_BUILD_SHIPPING
-	bool bWarnedAboutSameStageRW = false;
-	FNiagaraGpuComputeDispatchInterface* DispatchInterfaceForDebuggingOnly = nullptr;
-	
-	FString DebugName;
-	const TCHAR* GetDebugName()const{return *DebugName;}
-	#else
-	const TCHAR* GetDebugName()const{return nullptr;}
-	#endif
-
-	void BeginFrame(FNiagaraGpuComputeDispatchInterface* DispatchInterface, FRHICommandListImmediate& RHICmdList);
-	void EndFrame(FNiagaraGpuComputeDispatchInterface* DispatchInterface, FRHICommandListImmediate& RHICmdList);
-	void Reset();
-
-	FNiagaraDataBufferRef PrepareForWriteAccess(FRDGBuilder& GraphBuilder);
-	void EndWriteAccess(FRDGBuilder& GraphBuilder);
-	
-	FNiagaraDataBufferRef PrepareForReadAccess(FRDGBuilder& GraphBuilder, bool bCurrentFrame);
-	void EndReadAccess(FRDGBuilder& GraphBuilder, bool bCurrentFrame);
-
-
-	FNiagaraDataBufferRef AllocateBufferForCPU(FRDGBuilder& GraphBuilder, ERHIFeatureLevel::Type FeatureLevel, int32 AllocationSize, bool bPublishToGame, bool bPublishToCPU, FVector3f LWCTile);
-	void AddBuffersFromCPU(const TArray<FNiagaraDataBufferRef>& BuffersFromCPU);	
-	void AddGPUAllocationForNextTick(int32 AllocationCount);
-
-	FNiagaraDataBufferRef GetCurrentData()const { return CurrFrameData; }
-	FNiagaraDataBufferRef GetPrevFrameData()const { return PrevFrameData; }
-	
-	void AddTransition(FRDGBuilder& GraphBuilder, ERHIAccess AccessBefore, ERHIAccess AccessAfter, FNiagaraDataBuffer* Buffer);
-
-	//Perform and bookkeeping required when we remove a proxy from a dispatcher.
-	void OnAddedToDispatcher(FNiagaraGpuComputeDispatchInterface* ComputeDispatchInterface);
-	void OnRemovedFromDispatcher(FNiagaraGpuComputeDispatchInterface* ComputeDispatchInterface);
-};
-
-using FNiagaraDataChannelLayoutInfoPtr = TSharedPtr<FNiagaraDataChannelLayoutInfo>;
-
-
-/** Data describing the layout of Niagara Data channel buffers that is used in multiple places and must live beyond it's owning Data Channel. */
-struct FNiagaraDataChannelLayoutInfo : public TSharedFromThis<FNiagaraDataChannelLayoutInfo>
-{
-	FNiagaraDataChannelLayoutInfo(const UNiagaraDataChannel* DataChannel);
-	~FNiagaraDataChannelLayoutInfo();
-
-	const FNiagaraDataSetCompiledData& GetDataSetCompiledData()const{ return CompiledData; }
-	const FNiagaraDataSetCompiledData& GetDataSetCompiledDataGPU()const { return CompiledDataGPU; }
-	const FNiagaraDataChannelGameDataLayout& GetGameDataLayout()const { return GameDataLayout; }
-
-	/** If true, we keep our previous frame's data. Some users will prefer a frame of latency to tick dependency. */
-	bool KeepPreviousFrameData() const { return bKeepPreviousFrameData; }
-private:
-
-	/**
-	Data layout for payloads in Niagara datasets.
-	*/
-	FNiagaraDataSetCompiledData CompiledData;
-
-	FNiagaraDataSetCompiledData CompiledDataGPU;
-
-	/** Layout information for any data stored at the "Game" level. i.e. From game code/BP. AoS layout and LWC types. */
-	FNiagaraDataChannelGameDataLayout GameDataLayout;
-
-	bool bKeepPreviousFrameData = false;
-};
-
 UCLASS(abstract, EditInlineNew, MinimalAPI, prioritizeCategories=("Data Channel"))
 class UNiagaraDataChannel : public UObject
 {
@@ -279,267 +165,4 @@ void UNiagaraDataChannel::ForEachDataChannel(TAction Func)
 	}
 }
 
-UENUM(BlueprintType)
-enum class ENiagartaDataChannelReadResult : uint8
-{
-	Success,
-	Failure
-};
-
-/**
-* A C++ and Blueprint accessible library of utility functions for accessing Niagara DataChannel
-*/
-UCLASS(MinimalAPI)
-class UNiagaraDataChannelLibrary : public UBlueprintFunctionLibrary
-{
-	GENERATED_UCLASS_BODY()
-
-	UFUNCTION(BlueprintInternalUseOnly, Category = NiagaraDataChannel, meta = (Keywords = "niagara DataChannel", WorldContext = "WorldContextObject", UnsafeDuringActorConstruction = "true"))
-	static UE_API UNiagaraDataChannelHandler* GetNiagaraDataChannel(const UObject* WorldContextObject, const UNiagaraDataChannelAsset* Channel);
-
-	
-	//////////////////////////////////////////////////////////////////////////
-	//These functions are soon to be deprecated in favor of new functions that take a more general AccessContext instead of Search Parameters.
-
-	/**
-	 * LEGACY FUNCTION: Please use non-legacy version.
-	 * Initializes and returns the Niagara Data Channel writer to write N elements to the given data channel.
-	 * This function is now legacy and and will soon be deprecated. Please use the new non-legacy function.
-	 *
-	 * @param WorldContextObject	World to execute in
-	 * @param Channel				The channel to write to
-	 * @param SearchParams			Parameters used when retrieving a specific set of Data Channel Data to read or write like the islands data channel type.
-	 * @param Count					The number of elements to write 
-	 * @param bVisibleToGame	If true, the data written to this data channel is visible to Blueprint and C++ logic reading from it
-	 * @param bVisibleToCPU	If true, the data written to this data channel is visible to Niagara CPU emitters
-	 * @param bVisibleToGPU	If true, the data written to this data channel is visible to Niagara GPU emitters
-	 * @param DebugSource	Instigator for this write, used in the debug hud to track writes to the data channel from different sources
-	 */
-	UFUNCTION(BlueprintCallable, Category = NiagaraDataChannel, DisplayName="Write To Niagara Data Channel (Batch) (Legacy)", meta = (AdvancedDisplay = "SearchParams, DebugSource", Keywords = "niagara DataChannel", WorldContext = "WorldContextObject", UnsafeDuringActorConstruction = "true", AutoCreateRefTerm="DebugSource"))
-	static UE_API UNiagaraDataChannelWriter* WriteToNiagaraDataChannel(const UObject* WorldContextObject, const UNiagaraDataChannelAsset* Channel, FNiagaraDataChannelSearchParameters SearchParams, int32 Count, UPARAM(DisplayName = "Visible to Blueprint") bool bVisibleToGame, UPARAM(DisplayName = "Visible to Niagara CPU") bool bVisibleToCPU, UPARAM(DisplayName = "Visible to Niagara GPU") bool bVisibleToGPU, const FString& DebugSource);
-
-	/**
-	 * LEGACY FUNCTION: Please use non-legacy version.
-	 * Initializes and returns the Niagara Data Channel reader for the given data channel.
-	 * This function is now legacy and and will soon be deprecated. Please use the new non-legacy function.
-	 *
-	 * @param WorldContextObject	World to execute in
-	 * @param Channel				The channel to read from
-	 * @param SearchParams			Parameters used when retrieving a specific set of Data Channel Data to read or write like the islands data channel type.
-	 * @param bReadPreviousFrame	True if this reader will read the previous frame's data. If false, we read the current frame.
-	 *								Reading the current frame allows for zero latency reads, but any data elements that are generated after this reader is used are missed.
-	 *								Reading the previous frame's data introduces a frame of latency but ensures we never miss any data as we have access to the whole frame.
-	 */
-	UFUNCTION(BlueprintCallable, Category = NiagaraDataChannel, DisplayName="Read From Niagara Data Channel (Batch) (Legacy)", meta = (AdvancedDisplay = "SearchParams", Keywords = "niagara DataChannel", WorldContext = "WorldContextObject", UnsafeDuringActorConstruction = "true"))
-	static UE_API UNiagaraDataChannelReader* ReadFromNiagaraDataChannel(const UObject* WorldContextObject, const UNiagaraDataChannelAsset* Channel, FNiagaraDataChannelSearchParameters SearchParams, bool bReadPreviousFrame);
-
-	/**
-	 * LEGACY FUNCTION: Please use non-legacy version.
-	 * Returns the number of readable elements in the given data channel
-	 * This function is now legacy and and will soon be deprecated. Please use the new non-legacy function.
-	 *
-	 * @param WorldContextObject	World to execute in
-	 * @param Channel				The channel to read from
-	 * @param SearchParams			Parameters used when retrieving a specific set of Data Channel Data to read or write like the islands data channel type.
-	 * @param bReadPreviousFrame	True if this reader will read the previous frame's data. If false, we read the current frame.
-	 *								Reading the current frame allows for zero latency reads, but any data elements that are generated after this reader is used are missed.
-	 *								Reading the previous frame's data introduces a frame of latency but ensures we never miss any data as we have access to the whole frame.
-	 */
-	UFUNCTION(BlueprintCallable, Category = NiagaraDataChannel, meta = (bReadPreviousFrame="true", DisplayName="Get Data Channel Element Count (Legacy)", AdvancedDisplay = "SearchParams, bReadPreviousFrame", Keywords = "niagara DataChannel num size", WorldContext = "WorldContextObject", UnsafeDuringActorConstruction = "true"))
-	static UE_API int32 GetDataChannelElementCount(const UObject* WorldContextObject, const UNiagaraDataChannelAsset* Channel, FNiagaraDataChannelSearchParameters SearchParams, bool bReadPreviousFrame);
-
-	/**
-	 * LEGACY FUNCTION: Please use non-legacy version.
-	 * Subscribes to a single data channel and calls a delegate every times new data is written to the data channel.
-	 *
-	 * @param WorldContextObject	World to execute in
-	 * @param Channel				The channel to subscribe to for updates
-	 * @param SearchParams			Parameters used when retrieving a specific set of Data Channel Data to read - only used by some types of data channels, like the island types.
-	 * @param UpdateDelegate		The delegate to be called when new data is available in the data channel. Can be called multiple times per tick.
-	 * @param UnsubscribeToken		This token can be used to unsubscribe from the data channel.
-	 */
-	UFUNCTION(BlueprintCallable, Category = NiagaraDataChannel, meta = (DisplayName="Subscribe To Niagara Data Channel (Legacy)", Keywords = "niagara DataChannel event reader subscription delegate listener updates", WorldContext = "WorldContextObject", UnsafeDuringActorConstruction = "true"))
-	static UE_API void SubscribeToNiagaraDataChannel(const UObject* WorldContextObject, const UNiagaraDataChannelAsset* Channel, FNiagaraDataChannelSearchParameters SearchParams, const FOnNewNiagaraDataChannelPublish& UpdateDelegate, int32& UnsubscribeToken);
-
-	/**
-	 * Removes a prior registration from a data channel
-	 *
-	 * @param WorldContextObject	World to execute in
-	 * @param Channel				The channel to unsubscribe from
-	 * @param UnsubscribeToken		The token returned from the subscription call
-	 */
-	UFUNCTION(BlueprintCallable, Category = NiagaraDataChannel, meta = (Keywords = "niagara DataChannel event reader subscription delegate listener updates", WorldContext = "WorldContextObject", UnsafeDuringActorConstruction = "true"))
-	static UE_API void UnsubscribeFromNiagaraDataChannel(const UObject* WorldContextObject, const UNiagaraDataChannelAsset* Channel, const int32& UnsubscribeToken);
-	
-	/**
-	 * LEGACY FUNCTION: Please use non-legacy version.
-	 * Reads a single entry from the given data channel, if possible.
-	 *
-	 * @param WorldContextObject	World to execute in
-	 * @param Channel				The channel to read from
-	 * @param Index					The data index to read from
-	 * @param SearchParams			Parameters used when retrieving a specific set of Data Channel Data to read or write like the islands data channel type.
-	 * @param bReadPreviousFrame	True if this reader will read the previous frame's data. If false, we read the current frame.
-	 *								Reading the current frame allows for zero latency reads, but any data elements that are generated after this reader is used are missed.
-	 *								Reading the previous frame's data introduces a frame of latency but ensures we never miss any data as we have access to the whole frame.
-	 * @param ReadResult			Used by Blueprint for the return value
-	 */
-	UFUNCTION(BlueprintInternalUseOnly, Category = NiagaraDataChannel, DisplayName="Read From Niagara Data Channel (Legacy)", meta = (bReadPreviousFrame="true", AdvancedDisplay = "SearchParams, bReadPreviousFrame", ExpandEnumAsExecs="ReadResult", Keywords = "niagara DataChannel", WorldContext = "WorldContextObject", UnsafeDuringActorConstruction = "true"))
-	static UE_API void ReadFromNiagaraDataChannelSingle(const UObject* WorldContextObject, const UNiagaraDataChannelAsset* Channel, int32 Index, FNiagaraDataChannelSearchParameters SearchParams, bool bReadPreviousFrame, ENiagartaDataChannelReadResult& ReadResult);
-
-	/**
-	 * LEGACY FUNCTION: Please use non-legacy version.
-	 * Writes a single element to a Niagara Data Channel. The element won't be immediately visible to readers, as it needs to be processed first. The earliest point it can be read is in the next tick group.
-	 *
-	 * @param WorldContextObject	World to execute in
-	 * @param Channel				The channel to write to
-	 * @param SearchParams			Parameters used when retrieving a specific set of Data Channel Data to read or write like the islands data channel type.
-	 * @param bVisibleToBlueprint	If true, the data written to this data channel is visible to Blueprint and C++ logic reading from it
-	 * @param bVisibleToNiagaraCPU	If true, the data written to this data channel is visible to Niagara CPU emitters
-	 * @param bVisibleToNiagaraGPU	If true, the data written to this data channel is visible to Niagara GPU emitters
-	 */
-	UFUNCTION(BlueprintInternalUseOnly, Category = NiagaraDataChannel, DisplayName="Write To Niagara Data Channel (Legacy)", meta = (bVisibleToBlueprint="true", bVisibleToNiagaraCPU="true", bVisibleToNiagaraGPU="true", AdvancedDisplay = "bVisibleToBlueprint, bVisibleToNiagaraCPU, bVisibleToNiagaraGPU, SearchParams", Keywords = "niagara DataChannel event writer", WorldContext = "WorldContextObject", UnsafeDuringActorConstruction = "true"))
-	static UE_API void WriteToNiagaraDataChannelSingle(const UObject* WorldContextObject, const UNiagaraDataChannelAsset* Channel, FNiagaraDataChannelSearchParameters SearchParams, bool bVisibleToBlueprint, bool bVisibleToNiagaraCPU, bool bVisibleToNiagaraGPU);
-
-	//////////////////////////////////////////////////////////////////////////
-
-	/**
-	 * Initializes and returns the Niagara Data Channel reader for the given data channel.
-	 *
-	 * @param WorldContextObject	World to execute in
-	 * @param Channel				The channel to read from
-	 * @param AccessContext			Context allowing passing of NDC type specific information between the caller and internal Data Channel that can control internal routing and help feedback information to the caller.
-	 * @param bReadPreviousFrame	True if this reader will read the previous frame's data. If false, we read the current frame.
-	 *								Reading the current frame allows for zero latency reads, but any data elements that are generated after this reader is used are missed.
-	 *								Reading the previous frame's data introduces a frame of latency but ensures we never miss any data as we have access to the whole frame.
-	 */
-	UFUNCTION(BlueprintCallable, BlueprintInternalUseOnly, Category = NiagaraDataChannel, DisplayName="Read From Niagara Data Channel (Batch)", meta = (Keywords = "niagara DataChannel", WorldContext = "WorldContextObject", UnsafeDuringActorConstruction = "true"))
-	static UE_API UNiagaraDataChannelReader* ReadFromNiagaraDataChannel_WithContext(const UObject* WorldContextObject, const UNiagaraDataChannelAsset* Channel, UPARAM(ref) FNDCAccessContextInst& AccessContext, bool bReadPreviousFrame);
-	
-	/**
-	 * Initializes and returns the Niagara Data Channel writer to write N elements to the given data channel.
-	 *
-	 * @param WorldContextObject	World to execute in
-	 * @param Channel				The channel to write to
-	 * @param AccessContext			Context allowing passing of NDC type specific information between the caller and internal Data Channel that can control internal routing and help feedback information to the caller.
-	 * @param Count					The number of elements to write 
-	 * @param bVisibleToGame	If true, the data written to this data channel is visible to Blueprint and C++ logic reading from it
-	 * @param bVisibleToCPU	If true, the data written to this data channel is visible to Niagara CPU emitters
-	 * @param bVisibleToGPU	If true, the data written to this data channel is visible to Niagara GPU emitters
-	 * @param DebugSource	Instigator for this write, used in the debug hud to track writes to the data channel from different sources
-	 */
-	UFUNCTION(BlueprintCallable, BlueprintInternalUseOnly, Category = NiagaraDataChannel, DisplayName = "Write To Niagara Data Channel (Batch)", meta = (AdvancedDisplay = "DebugSource", Keywords = "niagara DataChannel", WorldContext = "WorldContextObject", UnsafeDuringActorConstruction = "true", AutoCreateRefTerm = "DebugSource", bVisibleToBlueprint="true", bVisibleToNiagaraCPU="true", bVisibleToNiagaraGPU="true"))
-	static UE_API UNiagaraDataChannelWriter* WriteToNiagaraDataChannel_WithContext(const UObject* WorldContextObject, const UNiagaraDataChannelAsset* Channel, UPARAM(ref) FNDCAccessContextInst& AccessContext, int32 Count, UPARAM(DisplayName = "Visible to Blueprint") bool bVisibleToBlueprint, UPARAM(DisplayName = "Visible to Niagara CPU") bool bVisibleToNiagaraCPU, UPARAM(DisplayName = "Visible to Niagara GPU") bool bVisibleToNiagaraGPU, const FString& DebugSource);
-	
-	/**
-	 * Returns the number of readable elements in the given data channel
-	 *
-	 * @param WorldContextObject	World to execute in
-	 * @param Channel				The channel to read from
-	 * @param AccessContext			Context allowing passing of NDC type specific information between the caller and internal Data Channel that can control internal routing and help feedback information to the caller.
-	 * @param bReadPreviousFrame	True if this reader will read the previous frame's data. If false, we read the current frame.
-	 *								Reading the current frame allows for zero latency reads, but any data elements that are generated after this reader is used are missed.
-	 *								Reading the previous frame's data introduces a frame of latency but ensures we never miss any data as we have access to the whole frame.
-	 */
-	UFUNCTION(BlueprintCallable, BlueprintInternalUseOnly,Category = NiagaraDataChannel, meta = (bReadPreviousFrame="true", Keywords = "niagara DataChannel num size", WorldContext = "WorldContextObject", UnsafeDuringActorConstruction = "true"))
-	static UE_API int32 GetDataChannelElementCount_WithContext(const UObject* WorldContextObject, const UNiagaraDataChannelAsset* Channel, UPARAM(ref) FNDCAccessContextInst& AccessContext, bool bReadPreviousFrame);
-	
-	/**
-	 * Subscribes to a single data channel and calls a delegate every times new data is written to the data channel.
-	 *
-	 * @param WorldContextObject	World to execute in
-	 * @param Channel				The channel to subscribe to for updates
-	 * @param AccessContext			Context allowing passing of NDC type specific information between the caller and internal Data Channel that can control internal routing and help feedback information to the caller.
-	 * @param UpdateDelegate		The delegate to be called when new data is available in the data channel. Can be called multiple times per tick.
-	 * @param UnsubscribeToken		This token can be used to unsubscribe from the data channel.
-	 */
-	UFUNCTION(BlueprintCallable, Category = NiagaraDataChannel, meta = (Keywords = "niagara DataChannel event reader subscription delegate listener updates", WorldContext = "WorldContextObject", UnsafeDuringActorConstruction = "true"))
-	static UE_API void SubscribeToNiagaraDataChannel_WithContext(const UObject* WorldContextObject, const UNiagaraDataChannelAsset* Channel, UPARAM(ref) FNDCAccessContextInst& AccessContext, const FOnNewNiagaraDataChannelPublish& UpdateDelegate, int32& UnsubscribeToken);
-		
-	/**
-	 * Reads a single entry from the given data channel, if possible.
-	 *
-	 * @param WorldContextObject	World to execute in
-	 * @param Channel				The channel to read from
-	 * @param Index					The data index to read from
-	 * @param AccessContext			Context allowing passing of NDC type specific information between the caller and internal Data Channel that can control internal routing and help feedback information to the caller.
-	 * @param bReadPreviousFrame	True if this reader will read the previous frame's data. If false, we read the current frame.
-	 *								Reading the current frame allows for zero latency reads, but any data elements that are generated after this reader is used are missed.
-	 *								Reading the previous frame's data introduces a frame of latency but ensures we never miss any data as we have access to the whole frame.
-	 * @param ReadResult			Used by Blueprint for the return value
-	 */
-	UFUNCTION(BlueprintInternalUseOnly, Category = NiagaraDataChannel, DisplayName="Read From Niagara Data Channel", meta = (bReadPreviousFrame="true", ExpandEnumAsExecs="ReadResult", Keywords = "niagara DataChannel", WorldContext = "WorldContextObject", UnsafeDuringActorConstruction = "true"))
-	static UE_API void ReadFromNiagaraDataChannelSingle_WithContext(const UObject* WorldContextObject, const UNiagaraDataChannelAsset* Channel, int32 Index, UPARAM(ref) FNDCAccessContextInst& AccessContext, bool bReadPreviousFrame, ENiagartaDataChannelReadResult& ReadResult);
-
-	/**
-	 * Writes a single element to a Niagara Data Channel. The element won't be immediately visible to readers, as it needs to be processed first. The earliest point it can be read is in the next tick group.
-	 *
-	 * @param WorldContextObject	World to execute in
-	 * @param Channel				The channel to write to
-	 * @param AccessContext			Context allowing passing of NDC type specific information between the caller and internal Data Channel that can control internal routing and help feedback information to the caller.
-	 * @param bVisibleToBlueprint	If true, the data written to this data channel is visible to Blueprint and C++ logic reading from it
-	 * @param bVisibleToNiagaraCPU	If true, the data written to this data channel is visible to Niagara CPU emitters
-	 * @param bVisibleToNiagaraGPU	If true, the data written to this data channel is visible to Niagara GPU emitters
-	 */
-	UFUNCTION(BlueprintInternalUseOnly, Category = NiagaraDataChannel, DisplayName="Write To Niagara Data Channel", meta = (bVisibleToBlueprint="true", bVisibleToNiagaraCPU="true", bVisibleToNiagaraGPU="true", Keywords = "niagara DataChannel event writer", WorldContext = "WorldContextObject", UnsafeDuringActorConstruction = "true"))
-	static UE_API void WriteToNiagaraDataChannelSingle_WithContext(const UObject* WorldContextObject, const UNiagaraDataChannelAsset* Channel, UPARAM(ref) FNDCAccessContextInst& AccessContext, bool bVisibleToBlueprint, bool bVisibleToNiagaraCPU, bool bVisibleToNiagaraGPU);
-
-	static UE_API UNiagaraDataChannelHandler* FindDataChannelHandler(const UObject* WorldContextObject, const UNiagaraDataChannel* Channel);
-	static UE_API UNiagaraDataChannelWriter* CreateDataChannelWriter(const UObject* WorldContextObject, const UNiagaraDataChannel* Channel, FNiagaraDataChannelSearchParameters SearchParams, int32 Count, bool bVisibleToGame, bool bVisibleToCPU, bool bVisibleToGPU, const FString& DebugSource);
-	static UE_API UNiagaraDataChannelReader* CreateDataChannelReader(const UObject* WorldContextObject, const UNiagaraDataChannel* Channel, FNiagaraDataChannelSearchParameters SearchParams, bool bReadPreviousFrame);
-
-	static UE_API UNiagaraDataChannelWriter* CreateDataChannelWriter_WithContext(const UObject* WorldContextObject, const UNiagaraDataChannel* Channel, FNDCAccessContextInst& AccessContext, int32 Count, bool bVisibleToGame, bool bVisibleToCPU, bool bVisibleToGPU, const FString& DebugSource);
-	static UE_API UNiagaraDataChannelReader* CreateDataChannelReader_WithContext(const UObject* WorldContextObject, const UNiagaraDataChannel* Channel, FNDCAccessContextInst& AccessContext, bool bReadPreviousFrame);
-
-
-	/** Returns the access context for use when accessing a Niagara Data channel. */
-	UFUNCTION(BlueprintCallable, BlueprintInternalUseOnly, Category = "Niagara Data Channel")
-	static UE_API FNDCAccessContextInst& GetUsableAccessContextFromNDC(const UNiagaraDataChannelAsset* DataChannel);
-
-	/** Returns the access context for use when accessing a Niagara Data channel. */
-	UFUNCTION(BlueprintCallable, BlueprintInternalUseOnly, Category = "Niagara Data Channel")
-	static UE_API FNDCAccessContextInst& GetUsableAccessContextFromNDCRef(const FNiagaraDataChannelReference& NDCRef);
-
-	/** 
-	Prepares an Access Context ready for accessing a Niagara Data Channel. 
-	If the passed Data Channel Reference has a valid custom AccessContext then it's values will be applied. Otherwise the defaults will be used.
-	Allows optional setting for the members of the Access Context. Transient properties are shown by default but properties to modify can be selected in the details panel.
-	*/
-	UFUNCTION(BlueprintInternalUseOnly, BlueprintCallable, Category = "Niagara Data Channel")
-	static UE_API FNDCAccessContextInst& PrepareAccessContextFromNDCRef(UPARAM(Ref) FNiagaraDataChannelReference& NDCRef);
-
-	/** Creates a new NDCAccessContext with a specified internal struct type. Struct must be a child of FNDCAccessContextBase. */
-	UFUNCTION(BlueprintInternalUseOnly, BlueprintCallable, CustomThunk, Category = "Niagara Data Channel", meta = (HidePin = "ContextStruct"))
-	static UE_API FNDCAccessContextInst MakeNDCAccessContextInstance(UScriptStruct* ContextStruct);
-
-	/** 
-	Gets the members from the internal Struct of an FNDCAccessContextIns. Expanded on compile into individual per property calls to GetSinglePropertyInNDCAccessContextInstance.
-	*/
-	UFUNCTION(BlueprintInternalUseOnly, BlueprintCallable, Category = "Niagara Data Channel", meta = (HidePin = "ContextStruct"))
-	static UE_API void GetMembersInNDCAccessContextInstance(const FNDCAccessContextInst& AccessContext, UScriptStruct* ContextStruct);
-	
-	/** 
-	Sets the members in the internal Struct of an FNDCAccessContextIns. Expanded on compile into individual per property calls to SetSinglePropertyInNDCAccessContextInstance.
-	*/
-	UFUNCTION(BlueprintInternalUseOnly, BlueprintCallable, Category = "Niagara Data Channel", meta = (HidePin = "ContextStruct"))
-	static UE_API void SetMembersInNDCAccessContextInstance(UPARAM(Ref) FNDCAccessContextInst& AccessContext, UScriptStruct* ContextStruct);
-	
-	/** 
-	Gets a specific property within the internal data of an FNDCAccessContextInst.
-	*/
-	UFUNCTION(BlueprintInternalUseOnly, BlueprintCallable, CustomThunk, Category = "Niagara Data Channel", meta = (CustomStructureParam = "Value", AutoCreateRefTerm="Value"))
-	static UE_API void GetSinglePropertyInNDCAccessContextInstance(const FNDCAccessContextInst& AccessContext, UScriptStruct* ContextStruct, FName PropertyName, int32& Value);
-	
-	/** 
-	Sets a specific property within the internal data of an FNDCAccessContextInst.
-	*/
-	UFUNCTION(BlueprintInternalUseOnly, BlueprintCallable, CustomThunk, Category = "Niagara Data Channel", meta = (CustomStructureParam = "Value", AutoCreateRefTerm="Value"))
-	static UE_API void SetSinglePropertyInNDCAccessContextInstance(UPARAM(Ref) FNDCAccessContextInst& AccessContext, UScriptStruct* ContextStruct, FName PropertyName, const int32& Value);
-
-private:
-	DECLARE_FUNCTION(execMakeNDCAccessContextInstance);
-	DECLARE_FUNCTION(execGetSinglePropertyInNDCAccessContextInstance);
-	DECLARE_FUNCTION(execSetSinglePropertyInNDCAccessContextInstance);
-};
-
 #undef UE_API
diff --git a/Engine/Plugins/FX/Niagara/Source/Niagara/Public/NiagaraDataChannelAccessor.h b/Engine/Plugins/FX/Niagara/Source/Niagara/Public/NiagaraDataChannelAccessor.h
--- a/Engine/Plugins/FX/Niagara/Source/Niagara/Public/NiagaraDataChannelAccessor.h
+++ b/Engine/Plugins/FX/Niagara/Source/Niagara/Public/NiagaraDataChannelAccessor.h
@@ -2,16 +2,13 @@
 
 #pragma once
 
-#include "NiagaraDataChannelPublic.h"
+#include "NiagaraDataChannelCommon.h"
+#include "NiagaraDataChannelGameData.h"
 #include "NiagaraDataChannel.h"
 #include "NiagaraDataChannelAccessor.generated.h"
 
 #define UE_API NIAGARA_API
 
-class UNiagaraDataChannelHandler;
-struct FNiagaraSpawnInfo;
-struct FNiagaraVariableBase;
-
 /** 
 Initial simple API for reading and writing data in a data channel from game code / BP. 
 Likely to be replaced in the near future with a custom BP node and a helper struct.
diff --git a/Engine/Plugins/FX/Niagara/Source/Niagara/Public/NiagaraDataChannelAccessorImpl.h b/Engine/Plugins/FX/Niagara/Source/Niagara/Public/NiagaraDataChannelAccessorImpl.h
--- a/Engine/Plugins/FX/Niagara/Source/Niagara/Public/NiagaraDataChannelAccessorImpl.h
+++ b/Engine/Plugins/FX/Niagara/Source/Niagara/Public/NiagaraDataChannelAccessorImpl.h
@@ -6,9 +6,8 @@
 #include "UObject/UObjectGlobals.h"
 #include "UObject/Package.h"
 
-#include "NiagaraDataChannelPublic.h"
-#include "NiagaraDataChannelCommon.h"
 #include "NiagaraDataChannel.h"
+#include "NiagaraDataChannelData.h"
 #include "NiagaraDataChannelHandler.h"
 
 #include "NiagaraWorldManager.h"
diff --git a/Engine/Plugins/FX/Niagara/Source/Niagara/Public/NiagaraDataChannelCommon.h b/Engine/Plugins/FX/Niagara/Source/Niagara/Public/NiagaraDataChannelCommon.h
--- a/Engine/Plugins/FX/Niagara/Source/Niagara/Public/NiagaraDataChannelCommon.h
+++ b/Engine/Plugins/FX/Niagara/Source/Niagara/Public/NiagaraDataChannelCommon.h
@@ -2,166 +2,67 @@
 
 #pragma once
 
+/**
+* Common shared types, headers and fwd dcls for Niagara Data Channels.
+*/
+
+#include "Engine/EngineBaseTypes.h"
 #include "CoreMinimal.h"
 #include "UObject/ObjectMacros.h"
+#include "NiagaraCommon.h"
 #include "NiagaraDataSet.h"
-#include "NiagaraDataChannelPublic.h"
+
+#include "NiagaraDataChannelCommon.generated.h"
 
 class UWorld;
-class UActorComponentc;
+class UActorComponent;
 class FNiagaraWorldManager;
 class FNiagaraDataBuffer;
 class FNiagaraSystemInstance;
+struct FNiagaraSpawnInfo;
+struct FNiagaraVariableBase;
+
+class UNiagaraDataChannelAsset;
 class UNiagaraDataChannel;
 class UNiagaraDataChannelHandler;
+class FNiagaraDataChannelManager;
+
+class UNiagaraDataChannelWriter;
+class UNiagaraDataChannelReader;
 class UNiagaraDataInterfaceDataChannelWrite;
 class UNiagaraDataInterfaceDataChannelRead;
 
-using FNiagaraDataChannelDataProxyPtr = TSharedPtr<struct FNiagaraDataChannelDataProxy>;
-
-/** A request to publish data into a Niagara Data Channel.  */
-struct FNiagaraDataChannelPublishRequest
-{
-	/** The buffer containing the data to be published. This can come from a data channel DI or can be the direct contents on a Niagara simulation. */
-	FNiagaraDataBufferRef Data;
-
-	/** Game level data if this request comes from the game code. */
-	TSharedPtr<FNiagaraDataChannelGameData> GameData = nullptr;
-
-	/** If true, data in this request will be made visible to BP and C++ game code.*/
-	bool bVisibleToGame = false;
-
-	/** If true, data in this request will be made visible to Niagara CPU simulations. */
-	bool bVisibleToCPUSims = false;
-
-	/** If true, data in this request will be made visible to Niagara GPU simulations. */
-	bool bVisibleToGPUSims = false;
-
-	/** 
-	LWC Tile for the originator system of this request.
-	Allows us to convert from the Niagara Simulation space into LWC coordinates.
-	*/
-	FVector3f LwcTile = FVector3f::ZeroVector;
-
-#if WITH_NIAGARA_DEBUGGER
-	/** Instigator of this write, used for debug tracking */
-	FString DebugSource;
-#endif
-
-	FNiagaraDataChannelPublishRequest() = default;
-	explicit FNiagaraDataChannelPublishRequest(FNiagaraDataBufferRef InData)
-		: Data(InData)
-	{
-	}
-
-	explicit FNiagaraDataChannelPublishRequest(FNiagaraDataBufferRef InData, bool bInVisibleToGame, bool bInVisibleToCPUSims, bool bInVisibleToGPUSims, FVector3f InLwcTile)
-	: Data(InData), bVisibleToGame(bInVisibleToGame), bVisibleToCPUSims(bInVisibleToCPUSims), bVisibleToGPUSims(bInVisibleToGPUSims), LwcTile(InLwcTile)
-	{
-	}
-};
-
-
-/**
-Underlying storage class for data channel data.
-Some data channels will have many of these and can distribute them as needed to different accessing systems.
-For example, some data channel handlers may subdivide the scene such that distant systems are not interacting.
-In this case, each subdivision would have it's own FNiagaraDataChannelData and distribute these to the relevant NiagaraSystems.
-*/
-struct FNiagaraDataChannelData final : public TSharedFromThis<FNiagaraDataChannelData, ESPMode::ThreadSafe>
-{
-	UE_NONCOPYABLE(FNiagaraDataChannelData)
-	NIAGARA_API explicit FNiagaraDataChannelData();
-	NIAGARA_API ~FNiagaraDataChannelData();
-
-	NIAGARA_API void Init(UNiagaraDataChannelHandler* Owner);
-	NIAGARA_API void Reset();
-
-	NIAGARA_API void BeginFrame(UNiagaraDataChannelHandler* Owner);
-	NIAGARA_API void EndFrame(UNiagaraDataChannelHandler* Owner);
-	NIAGARA_API int32 ConsumePublishRequests(UNiagaraDataChannelHandler* Owner, const ETickingGroup& TickGroup);
-
-	NIAGARA_API FNiagaraDataChannelGameData* GetGameData();
-	NIAGARA_API FNiagaraDataBufferRef GetCPUData(bool bPreviousFrame);
-	FNiagaraDataChannelDataProxyPtr GetRTProxy(){ return RTProxy; }
-	
-	/** Adds a request to publish some data into the channel on the next tick. */
-	NIAGARA_API void Publish(const FNiagaraDataChannelPublishRequest& Request);
-
-	NIAGARA_API void PublishFromGPU(const FNiagaraDataChannelPublishRequest& Request);
-
-	NIAGARA_API const FNiagaraDataSetCompiledData& GetCompiledData(ENiagaraSimTarget SimTarget);
+struct FNiagaraDataChannelData;
+struct FNiagaraDataChannelDataProxy;
+struct FNiagaraDataChannelGameData;
+struct FNiagaraDataChannelPublishRequest;
 
-	void SetLwcTile(FVector3f InLwcTile){ LwcTile = InLwcTile; }
-	FVector3f GetLwcTile()const { return LwcTile; }
+struct FNiagaraDataChannelLayoutInfo;
+struct FNiagaraDataChannelGameDataLayout;
 
-	//This will get a buffer from the CPU dataset intended to be written to on the CPU.
-	FNiagaraDataBuffer* GetBufferForCPUWrite();
+struct FNDCAccessContextInst;
 
-	void DestroyRenderThreadProxy(FNiagaraGpuComputeDispatchInterface* ComputeDispatchInterface);
+using FNiagaraDataChannelDataPtr = TSharedPtr<FNiagaraDataChannelData>;
+using FNiagaraDataChannelDataProxyPtr = TSharedPtr<FNiagaraDataChannelDataProxy>;
+using FNiagaraDataChannelGameDataPtr = TSharedPtr<FNiagaraDataChannelGameData>;
+using FNiagaraDataChannelLayoutInfoPtr = TSharedPtr<FNiagaraDataChannelLayoutInfo>;
 
-	void RegisterGPUSpawningReader() { ++NumGPUSpawningReaders; }
-	void UnregisterGPUSpawningReader() { --NumGPUSpawningReaders; }
-	int32 NumRegisteredGPUSpawningReaders()const{ return NumGPUSpawningReaders; }
-
-	//Returns if this data is still valid. This can return false in cases where the owning data channel has been modified for example.
-	bool IsLayoutValid(UNiagaraDataChannelHandler* Owner)const;
-
-	//Return true if data has been written to this NDC data for the current frame.
-	bool HasData()const;
-
-	//Returns a game data buffer into which we can write Count valuse on the Game Thread.
-	FNiagaraDataChannelGameDataPtr GetGameDataForWriteGT(int32 Count, bool bVisibleToGame, bool bVisibleToCPU, bool bVisibleToGPU, const FString& DebugSource);
-
-private:
-
-	void CreateRenderThreadProxy(UNiagaraDataChannelHandler* Owner);
-
-	void FlushPendingGameData(int32 Index);
-	void FlushAllPendingGameData();
-
-	/** DataChannel data accessible from Game/BP. AoS Layout. LWC types. */
-	FNiagaraDataChannelGameDataPtr GameData;
-
-	/** DataChannel data accessible to Niagara CPU sims. SoA layout. Non LWC types. */
-	FNiagaraDataSet* CPUSimData = nullptr;
-
-	// Cached off buffer with the previous frame's CPU Sim accessible data.
-	// Some systems can choose to read this to avoid any current frame tick ordering issues.
-	FNiagaraDataBufferRef PrevCPUSimData = nullptr;
-
-	/** Dataset we use for staging game data for the consumption by RT/GPU sims. */
-	FNiagaraDataSet* GameDataStaging = nullptr;
-
-	/** Data buffers we'll be passing to the RT proxy for uploading to the GPU */
-	TArray<FNiagaraDataChannelPublishRequest> PublishRequestsForGPU;
-
-	/** Render thread proxy for this data. Owns all RT side data meant for GPU simulations. */
-	FNiagaraDataChannelDataProxyPtr RTProxy;
-
-	/** Pending requests to publish data into this data channel. These requests are consumed at tick tick group. */
-	TArray<FNiagaraDataChannelPublishRequest> PublishRequests;
-
-	/** Pending requests to publish data into this data channel from the GPU. To alleviate data race behavior with data coming back from the GPU, we always consume GPU requests at the start of the frame only. */
-	TArray<FNiagaraDataChannelPublishRequest> PublishRequestsFromGPU;
-
-	/** The world we were initialized with, used to get the compute interface. */
-	TWeakObjectPtr<UWorld> WeakOwnerWorld;
-
-	FVector3f LwcTile = FVector3f::ZeroVector;
+USTRUCT(BlueprintType)
+struct FNiagaraDataChannelUpdateContext
+{
+	GENERATED_BODY();
 
-	/** Critical section protecting shared state for multiple writers publishing from different threads. */
-	FCriticalSection PublishCritSec;
+	UPROPERTY(BlueprintReadOnly, Category = "Data Channel")
+	TObjectPtr<class UNiagaraDataChannelReader> Reader;
 
-	//Keep reference to the layout this data was built with.
-	FNiagaraDataChannelLayoutInfoPtr LayoutInfo;
+	UPROPERTY(BlueprintReadOnly, Category = "Data Channel")
+	int32 FirstNewDataIndex = 0;
 
-	/** 
-	Track number of explicitly registered readers that spawn GPU particles from this data.
-	If we're spawning GPU particles using the CPU data (Spawn Conditional etc) then we have to send all CPU data to the GPU every frame.
-	Can possibly extend this to be a more automatic, registration based approach to shipping NDC data around rather than explicit flags on write.
-	*/
-	std::atomic<int32> NumGPUSpawningReaders;
+	UPROPERTY(BlueprintReadOnly, Category = "Data Channel")
+	int32 LastNewDataIndex = 0;
 
-	//We keep a set of incoming game data for all flag combinations and accumulate data into these rather than keeping all as separate
-	TArray<FNiagaraDataChannelGameDataPtr, TInlineAllocator<8>> PendingDestGameData;
+	UPROPERTY(BlueprintReadOnly, Category = "Data Channel")
+	int32 NewElementCount = 0;
 };
+
+DECLARE_DYNAMIC_DELEGATE_OneParam(FOnNewNiagaraDataChannelPublish, const FNiagaraDataChannelUpdateContext&, NewDataContext);
\ No newline at end of file
diff --git a/Engine/Plugins/FX/Niagara/Source/Niagara/Public/NiagaraDataChannelHandler.h b/Engine/Plugins/FX/Niagara/Source/Niagara/Public/NiagaraDataChannelHandler.h
--- a/Engine/Plugins/FX/Niagara/Source/Niagara/Public/NiagaraDataChannelHandler.h
+++ b/Engine/Plugins/FX/Niagara/Source/Niagara/Public/NiagaraDataChannelHandler.h
@@ -11,13 +11,13 @@ Some more complex handlers may want to divide up the scene in various different
 
 */
 
-#include "NiagaraDataChannelPublic.h"
+#include "NiagaraDataChannelCommon.h"
+#include "NiagaraDataChannelAccessContext.h"
+
 #include "RenderCommandFence.h"
 #include "NiagaraDataChannelHandler.generated.h"
 
 class FNiagaraGpuComputeDispatchInterface;
-struct FNiagaraDataChannelData;
-class FNiagaraDataChannelManager;
 
 UCLASS(abstract, BlueprintType, MinimalAPI)
 class UNiagaraDataChannelHandler : public UObject
diff --git a/Engine/Plugins/FX/Niagara/Source/Niagara/Public/NiagaraDataChannelPublic.h b/Engine/Plugins/FX/Niagara/Source/Niagara/Public/NiagaraDataChannelPublic.h
--- a/Engine/Plugins/FX/Niagara/Source/Niagara/Public/NiagaraDataChannelPublic.h
+++ b/Engine/Plugins/FX/Niagara/Source/Niagara/Public/NiagaraDataChannelPublic.h
@@ -2,615 +2,19 @@
 
 #pragma once
 
-#include "Engine/EngineBaseTypes.h"
-#include "NiagaraCommon.h"
-#include "StructUtils/InstancedStruct.h"
-#include "NiagaraDataChannelPublic.generated.h"
-
-struct FNiagaraDataChannelPublishRequest;
-class FNiagaraWorldManager;
-class UNiagaraDataChannelAsset;
-class UNiagaraDataChannel;
-class UNiagaraDataChannelHandler;
-struct FNiagaraDataChannelGameData;
-struct FNiagaraDataChannelData;
-struct FNiagaraDataChannelLayoutInfo;
-using FNiagaraDataChannelGameDataPtr = TSharedPtr<FNiagaraDataChannelGameData>;
-using FNiagaraDataChannelDataPtr = TSharedPtr<FNiagaraDataChannelData>;
-using FNiagaraDataChannelLayoutInfoPtr = TSharedPtr<FNiagaraDataChannelLayoutInfo>;
-
-USTRUCT(BlueprintType)
-struct FNiagaraDataChannelUpdateContext
-{
-	GENERATED_BODY();
-
-	UPROPERTY(BlueprintReadOnly, Category = "Data Channel")
-	TObjectPtr<class UNiagaraDataChannelReader> Reader;
-
-	UPROPERTY(BlueprintReadOnly, Category = "Data Channel")
-	int32 FirstNewDataIndex = 0;
-
-	UPROPERTY(BlueprintReadOnly, Category = "Data Channel")
-	int32 LastNewDataIndex = 0;
-
-	UPROPERTY(BlueprintReadOnly, Category = "Data Channel")
-	int32 NewElementCount = 0;
-};
-
-DECLARE_DYNAMIC_DELEGATE_OneParam(FOnNewNiagaraDataChannelPublish, const FNiagaraDataChannelUpdateContext&, NewDataContext);
-
-
-USTRUCT(BlueprintType)
-struct FNiagaraDataChannelVariable : public FNiagaraVariableBase
-{
-	GENERATED_USTRUCT_BODY()
-	
-	bool Serialize(FArchive& Ar);
-
-#if WITH_EDITORONLY_DATA
-	
-	/** Can be used to track renamed data channel variables */
-	UPROPERTY(meta=(IgnoreForMemberInitializationTest))
-	FGuid Version = FGuid::NewGuid();
-
-	NIAGARA_API static bool IsAllowedType(const FNiagaraTypeDefinition& Type);
-#endif
-	
-	NIAGARA_API static FNiagaraTypeDefinition ToDataChannelType(const FNiagaraTypeDefinition& Type);
-
-};
-
-template<>
-struct TStructOpsTypeTraits<FNiagaraDataChannelVariable> : public TStructOpsTypeTraitsBase2<FNiagaraDataChannelVariable>
-{
-	enum
-	{
-		WithSerializer = true,
-	};
-};
-
-/**
-Minimal set of types and declares required for external users of Niagara Data Channels.
-*/
-
-/**
-Legacy search parameters to support legacy code.
-*/
-USTRUCT(BlueprintType)
-struct FNiagaraDataChannelSearchParameters
-{
-	GENERATED_BODY()
-
-	FNiagaraDataChannelSearchParameters() = default;
-
-	FNiagaraDataChannelSearchParameters(USceneComponent* Owner)
-		: OwningComponent(Owner)
-		, Location(FVector::ZeroVector)
-		, bOverrideLocation(false)
-	{
-	}
-
-	FNiagaraDataChannelSearchParameters(USceneComponent* Owner, FVector LocationOverride)
-		: OwningComponent(Owner)
-		, Location(LocationOverride)
-		, bOverrideLocation(true)
-	{
-	}
-
-	FNiagaraDataChannelSearchParameters(FVector InLocation)
-		: OwningComponent(nullptr)
-		, Location(InLocation)
-		, bOverrideLocation(true)
-	{
-	}
-
-	[[nodiscard]] NIAGARA_API FVector GetLocation()const;
-	[[nodiscard]] inline USceneComponent* GetOwner()const { return OwningComponent; }
-
-	/** In cases where there is an owning component such as an object spawning from itself etc, then we pass that component in. Some handlers may only use it's location but others may make use of more data. */
-	UPROPERTY(EditAnywhere, BlueprintReadWrite, Category = "Parameters")
-	TObjectPtr<USceneComponent> OwningComponent = nullptr;
-
-	/** In cases where there is no owning component for data being read or written to a data channel, we simply pass in a location. We can also use this when bOverrideLocaiton is set. */
-	UPROPERTY(EditAnywhere, BlueprintReadWrite, Category = "Parameters")
-	FVector Location = FVector::ZeroVector;
-
-	/** If true, even if an owning component is set, the data channel should use the Location value rather than the component location. If this is false, the NDC will get any location needed from the owning component. */
-	UPROPERTY(EditAnywhere, BlueprintReadWrite, Category = "Parameters")
-	uint32 bOverrideLocation : 1 = false;
-};
-
-
-/**
-Base class for context used when retrieving a specific set of Data Channel Data to read or write.
-Many Data Channel types will have multiple internal sets of data and these parameters control which the Channel should return to users for access.
-An example of this would be the Islands Data Channel type which will subdivide the world and have a different set of data for each sub division.
-It will return to users the correct data for their location based on these parameters.
-NDC Types must specify which of these access context types they require.
-The base struct contains only data needed for internal systems to function. 
-Beyond this you are free to add anything else you need.
-Some data can be intended as input to the search/internal routing or data creation. This should be marked with "NDCAccessContextInput".
-Other data can be intended as output to store derived or transient state and pass back to the user. This should be marked with "NDCAccessContextOutput".
-*/
-USTRUCT(BlueprintType, meta=(Hidden))
-struct FNDCAccessContextBase
-{
-	GENERATED_BODY()
-
-public:
-
-	FNDCAccessContextBase() = default;
-	[[nodiscard]] explicit FNDCAccessContextBase(USceneComponent* Owner)
-		: OwningComponent(Owner)
-		, bIsAutoLinkingSystem(false)
-	{
-	}
-
-	[[nodiscard]] inline USceneComponent* GetOwner()const { return OwningComponent; }
-
-	UPROPERTY(BlueprintReadWrite, Transient, Category = "Input", meta=(NDCAccessContextInput))
-	TObjectPtr<USceneComponent> OwningComponent = nullptr;
-	
-	/** 
-	If true, this access is for a Niagara System that wishes to "Auto link" i.e. not do a proper lookup, just use the currently initializing map entry or the spawned component lookup.
-	This is used by NDC Data Interfaces in systems spawned by NDCs to make linking back to the originating NDC entry easier.
-	*/
-	UPROPERTY(transient)
-	uint32 bIsAutoLinkingSystem : 1 = false;
-};
-
-/** Utility class to wrap a weak pointer to a Niagara Component allowing FNDCAccessContext to hold an array of weak spawned systems. */
-USTRUCT(BlueprintType)
-struct FNDCSpawnedSystemRef
-{
-	GENERATED_BODY()
-	
-	FNDCSpawnedSystemRef() = default;
-	[[nodiscard]] FNDCSpawnedSystemRef(UNiagaraComponent* Component);
-
-	UPROPERTY(BlueprintReadWrite, Transient, Category = "Output")
-	TWeakObjectPtr<UNiagaraComponent> SpawnedSystem;
-
-	[[nodiscard]] inline UNiagaraComponent* Get()const{ return SpawnedSystem.Get(); }
-};
-
-/**
-Basic AccessContext type providing a little more utility than the base.
-If you require more data than this you can subclass this or replace entirely and subclass FNDCAccessContextBase.
-*/
-USTRUCT(BlueprintType)
-struct FNDCAccessContext : public FNDCAccessContextBase
-{
-	GENERATED_BODY()
-
-	FNDCAccessContext() = default;
-	[[nodiscard]] explicit FNDCAccessContext(USceneComponent* Owner)
-		: FNDCAccessContextBase(Owner)
-	{
-	}
-
-	[[nodiscard]] explicit FNDCAccessContext(USceneComponent* Owner, FVector LocationOverride)
-		: FNDCAccessContextBase(Owner)
-		, Location(LocationOverride)
-		, bOverrideLocation(true)
-		, bOverrideSystemToSpawn(false)
-	{
-	}
-
-	[[nodiscard]] explicit FNDCAccessContext(FVector InLocation)
-		: FNDCAccessContextBase(nullptr)
-		, Location(InLocation)
-		, bOverrideLocation(true)
-		, bOverrideSystemToSpawn(false)
-	{
-	}
-
-	[[nodiscard]] NIAGARA_API FVector GetLocation()const;
-
-	/** In cases where there is no owning component for data being read or written to a data channel, we simply pass in a location. We can also use this when bOverrideLocaiton is set. */
-	UPROPERTY(BlueprintReadWrite, Transient, Category = "Input", meta=(NDCAccessContextInput))
-	FVector Location = FVector::ZeroVector;
-
-	/** If true, even if an owning component is set, the data channel should use the Location value rather than the component location. If this is false, the NDC will get any location needed from the owning component. */
-	UPROPERTY(BlueprintReadWrite, Transient, Category = "Input", meta=(NDCAccessContextInput))
-	uint32 bOverrideLocation : 1 = false;
-
-	/** If true, we'll override the system to spawn for this NDC write. If false, we'll use the system defined in the NDC asset, if there is one. */
-	UPROPERTY(BlueprintReadWrite, Category = "Input", meta = (InlineEditCondition, NDCAccessContextInput))
-	uint32 bOverrideSystemToSpawn : 1 = false;
-
-	UPROPERTY(EditAnywhere, BlueprintReadWrite, Category = "Input", meta = (EditCondition="bOverrideSystemToSpawn", NDCAccessContextInput, AllowedClasses = "/Script/Niagara.NiagaraSystem,/Script/Niagara.NiagaraSystemCollection"))
-	TObjectPtr<UObject> SystemToSpawn;
-
-	/** The Niagara System spawned in response to this access, if any. Allows callers to set parameters etc on this system. It is not safe to store and use this component reference later. It can be reclaimed by the system at any time. */
-	UPROPERTY(BlueprintReadWrite, Transient, Category = "Output", meta = (NDCAccessContextOutput, ShowOnlyInnerProperties))
-	TArray<FNDCSpawnedSystemRef> SpawnedSystems;
-};
-
-/** Special legacy access context type. Used only to facilitate legacy NDCs to be usable by the new API. */
-USTRUCT(BlueprintType)
-struct FNDCAccessContextLegacy : public FNDCAccessContextBase
-{
-	GENERATED_BODY()
-
-	FNDCAccessContextLegacy() = default;
-	[[nodiscard]] explicit FNDCAccessContextLegacy(USceneComponent* Owner)
-		: FNDCAccessContextBase(Owner)
-	{
-	}
-
-	[[nodiscard]] explicit FNDCAccessContextLegacy(USceneComponent* Owner, FVector LocationOverride)
-		: FNDCAccessContextBase(Owner)
-		, Location(LocationOverride)
-		, bOverrideLocation(true)
-	{
-	}
-
-	[[nodiscard]] explicit FNDCAccessContextLegacy(FVector InLocation)
-		: FNDCAccessContextBase(nullptr)
-		, Location(InLocation)
-		, bOverrideLocation(true)
-	{
-	}
-
-	[[nodiscard]] NIAGARA_API FVector GetLocation()const;
-
-	/** In cases where there is no owning component for data being read or written to a data channel, we simply pass in a location. We can also use this when bOverrideLocaiton is set. */
-	UPROPERTY(BlueprintReadWrite, Transient, Category = "Input", meta = (NDCAccessContextInput))
-	FVector Location = FVector::ZeroVector;
-
-	/** If true, even if an owning component is set, the data channel should use the Location value rather than the component location. If this is false, the NDC will get any location needed from the owning component. */
-	UPROPERTY(BlueprintReadWrite, Transient, Category = "Input", meta = (NDCAccessContextInput))
-	uint32 bOverrideLocation : 1 = false;
-};
-
-//Describes a NDC Access Context Type.
-//Wrapper around a UScriptStruct* that checks for a valid type.
-struct TNDCAccessContextType
-{
-	TNDCAccessContextType() = default;
-	[[nodiscard]] explicit TNDCAccessContextType(UScriptStruct* InStruct)
-		: AccessContextStruct(InStruct->IsChildOf(FNDCAccessContextBase::StaticStruct()) ? InStruct : nullptr)
-	{
-		check(AccessContextStruct);
-	}
-
-	[[nodiscard]] const UScriptStruct* Get()const { return AccessContextStruct; }
-
-	//Returns the underlying struct as non-const. Only use to inter operate with code that stores structs without constness.
-	[[nodiscard]] UScriptStruct* Get_MutableInteropOnly()const { return const_cast<UScriptStruct*>(AccessContextStruct); }
-
-private:
-	const UScriptStruct* AccessContextStruct = nullptr;
-};
-
 /**
-Context object for Niagara and Game code to access a Niagara Data Channel.
-Can be customized per Data Channel to provide specific control for each Data Channel type.
-Can include input data from accessing code such as a location, game play tags etc that can influence internal data routing/partitions.
-Can also include output data allowing the NDC to pass information back to the accessing code.
+Types and headers required for users of Niagara Data Channels.
 */
-USTRUCT(BlueprintType)
-struct FNDCAccessContextInst
-{
-	GENERATED_BODY()
-
-	FNDCAccessContextInst() = default;
-	[[nodiscard]] FNDCAccessContextInst(TNDCAccessContextType AccessContextType)
-	{
-		AccessContext.InitializeAs(AccessContextType.Get());
-	}
-
-	inline bool Init(TNDCAccessContextType ParamsType)
-	{
-		AccessContext.InitializeAs(ParamsType.Get());
-		return AccessContext.IsValid();
-	}
-	
-	[[nodiscard]] bool operator==(const FNDCAccessContextInst&) const = default;
-
-	inline void Reset()
-	{
-		AccessContext.Reset();
-	}
-
-	[[nodiscard]] inline bool IsValid()const { return AccessContext.IsValid(); }
-
-	template<typename T = FNDCAccessContextBase>
-	[[nodiscard]] inline T& GetChecked() { return AccessContext.GetMutable<T>(); }
-
-	template<typename T = FNDCAccessContextBase>
-	[[nodiscard]] inline T* Get() { return AccessContext.GetMutablePtr<T>();	}
-
-	[[nodiscard]] inline const UScriptStruct* GetScriptStruct()const { return AccessContext.GetScriptStruct(); }
-
-	UPROPERTY(EditAnywhere, BlueprintReadWrite, Category = "Access Context", meta = (ShowOnlyInnerProperties, BaseStruct = "/Script/Niagara.NDCAccessContextBase", ExcludeBaseStruct))
-	FInstancedStruct AccessContext;
-};
-
-
-USTRUCT()
-struct FNiagaraDataChannelGameDataLayout
-{
-	GENERATED_BODY();
-
-	/** Map of all variables contained in this DataChannel data and the indices into data arrays for game data. */
-	UPROPERTY()
-	TMap<FNiagaraVariableBase, int32> VariableIndices;
-
-	/** Helpers for converting LWC types into Niagara simulation SWC types. */
-	UPROPERTY()
-	TArray<FNiagaraLwcStructConverter> LwcConverters;
-
-	void Init(TConstArrayView<FNiagaraDataChannelVariable> Variables);
-};
-
-
-#if WITH_NIAGARA_DEBUGGER
-
-/** Hooks into internal NiagaraDataChannels code for debugging and testing purposes. */
-class FNiagaraDataChannelDebugUtilities
-{
-public: 
-
-	static NIAGARA_API void BeginFrame(FNiagaraWorldManager* WorldMan, float DeltaSeconds);
-	static NIAGARA_API void EndFrame(FNiagaraWorldManager* WorldMan, float DeltaSeconds);
-	static NIAGARA_API void Tick(FNiagaraWorldManager* WorldMan, float DeltaSeconds, ETickingGroup TickGroup);
-	
-	static NIAGARA_API UNiagaraDataChannelHandler* FindDataChannelHandler(FNiagaraWorldManager* WorldMan, UNiagaraDataChannel* DataChannel);
-
-	static void LogWrite(const FNiagaraDataChannelPublishRequest& WriteRequest, const UNiagaraDataChannel* DataChannel, const ETickingGroup& TickGroup);
-	static void DumpAllWritesToLog();
-
-	static FNiagaraDataChannelDebugUtilities& Get();
-	static void TearDown();
-
-private:
-	struct FChannelWriteRequest
-	{
-		TWeakObjectPtr<const UNiagaraDataChannel> Channel;
-		TSharedPtr<FNiagaraDataChannelGameData> Data;
-		bool bVisibleToGame = false;
-		bool bVisibleToCPUSims = false;
-		bool bVisibleToGPUSims = false;
-		ETickingGroup TickGroup;
-		FString DebugSource;
-	};
-	struct FFrameDebugData
-	{
-		uint64 FrameNumber;
-		TArray<FChannelWriteRequest> WriteRequests;
-	};
-	
-	static FString ToJson(FNiagaraDataChannelGameData* Data);
-	static FString TickGroupToString(const ETickingGroup& TickGroup);
-	TArray<FFrameDebugData> FrameData;
-};
-
-
-#endif//WITH_NIAGARA_DEBUGGER
-
-
-/** Buffer containing a single FNiagaraVariable's data at the game level. AoS layout. LWC Types. */
-struct FNiagaraDataChannelVariableBuffer
-{
-	TArray<uint8> Data;
-	TArray<uint8> PrevData;
-	int32 Size = 0;
-
-	void Init(const FNiagaraVariableBase& Var)
-	{
-		//Position types are a special case where we have to store an LWC Vector in game level data and convert to a simulation friendly FVector3f as it enters Niagara.
-		if (Var.GetType() == FNiagaraTypeDefinition::GetPositionDef())
-		{
-			Size = FNiagaraTypeHelper::GetVectorDef().GetSize();
-		}
-		else
-		{
-			Size = Var.GetSizeInBytes();
-		}
-	}
-
-	void Empty() 
-	{
-		PrevData.Empty();
-		Data.Empty(); 
-	}
-
-	void Reset()
-	{
-		PrevData.Reset();
-		Data.Reset();
-	}
-
-	void BeginFrame(bool bKeepPrevious)
-	{
-		if (bKeepPrevious)
-		{
-			Swap(Data, PrevData);
-		}
-		Data.Reset();
-	}
-
-	template<typename T>
-	bool Write(int32 Index, const T& InData)
-	{
-		check(sizeof(T) == Size);
-		if (Index >= 0 && Index < Num())
-		{
-			T* Dest = reinterpret_cast<T*>(Data.GetData()) + Index;
-			*Dest = InData;
-
-			return true;
-		}
-		return false;
-	}
-
-	template<typename T>
-	bool Read(int32 Index, T& OutData, bool bPreviousFrameData)const
-	{
-		check(sizeof(T) == Size);
-
-		//Should we fallback to previous if we don't have current?
-		//bPreviousFrameData |= Data.Num() == 0;
-
-		int32 NumElems = bPreviousFrameData ? PrevNum() : Num();
-		if (Index >= 0 && Index < NumElems)
-		{
-			const uint8* DataPtr = bPreviousFrameData ? PrevData.GetData() : Data.GetData();
-			const T* Src = reinterpret_cast<const T*>(DataPtr) + Index;
-			OutData = *Src;
-
-			return true;
-		}
-		OutData = T();
-		return false;
-	}
-
-	void SetNum(int32 Num)
-	{
-		Data.SetNumZeroed(Size * Num);
-	}
-
-	void Reserve(int32 Num)
-	{
-		Data.Reserve(Size * Num);
-	}
-
-	int32 Num()const { return Data.Num() / Size; }
-	int32 PrevNum()const { return PrevData.Num() / Size; }
-
-	int32 GetElementSize()const { return Size; }
-};
-
-
-/** Storage for game level DataChannels generated by BP / C++ */
-struct FNiagaraDataChannelGameData : public TSharedFromThis<FNiagaraDataChannelGameData>
-{
-private:
-
-	/** Per variable storage buffers. */
-	TArray<FNiagaraDataChannelVariableBuffer> VariableData;
-
-	int32 NumElements = 0;
-	int32 MaxElements = 0;
-	int32 PrevNumElements = 0;
-
-	//Keep reference to the layout this data was built with.
-	FNiagaraDataChannelLayoutInfoPtr LayoutInfo;
-
-public:
-	
-	FNiagaraDataChannelGameData(const FNiagaraDataChannelLayoutInfoPtr& InLayoutInfo);
-	NIAGARA_API void Init(const FNiagaraDataChannelLayoutInfoPtr& InLayoutInfo);
-	NIAGARA_API void Empty();
-	NIAGARA_API void Reset();
-	NIAGARA_API void BeginFrame();
-	NIAGARA_API void SetNum(int32 NewNum);
-	NIAGARA_API void Reserve(int32 NewNum);
-	int32 Num() const { return NumElements; }
-	int32 Max() const { return MaxElements; }
-	int32 PrevNum()const { return PrevNumElements; }
-
-	int32 Add(int32 Count=1)
-	{
-		int32 Ret = NumElements;
-		SetNum(Ret + Count);
-		return Ret;
-	}
-
-	NIAGARA_API FNiagaraDataChannelVariableBuffer* FindVariableBuffer(const FNiagaraVariableBase& Var);
-
-	NIAGARA_API void WriteToDataSet(FNiagaraDataBuffer* DestBuffer, int32 DestStartIdx, FVector3f SimulationLwcTile);
-
-	NIAGARA_API void AppendFromGameData(const FNiagaraDataChannelGameData& GameData);
-
-	NIAGARA_API void AppendFromDataSet(const FNiagaraDataBuffer* SrcBuffer, FVector3f SimulationLwcTile);
-
-	const FNiagaraDataChannelLayoutInfoPtr& GetLayoutInfo()const { return LayoutInfo; }
-
-	const TConstArrayView<FNiagaraDataChannelVariableBuffer> GetVariableBuffers()const { return VariableData; }
-	void SetFromSimCache(const FNiagaraVariableBase& SourceVar, TConstArrayView<uint8> Data, int32 Size);
-
-	template<typename T>
-	bool Write(int32 VarIndex, int32 Index, const T& InValue)
-	{
-		if(VariableData.IsValidIndex(VarIndex))
-		{
-			FNiagaraDataChannelVariableBuffer& VarBuffer = VariableData[VarIndex];
-			return VarBuffer.Write<T>(Index, InValue);
-		}
-		return false;
-	}
-
-	template<typename T>
-	bool Read(int32 VarIndex, int32 Index, T& OutValue, bool bPreviousFrame)
-	{
-		if (VariableData.IsValidIndex(VarIndex))
-		{
-			FNiagaraDataChannelVariableBuffer& VarBuffer = VariableData[VarIndex];
-			return VarBuffer.Read<T>(Index, OutValue, bPreviousFrame);
-		}
-		return false;
-	}
-};
-
-/** Niagara Data Channels are a system for communication between Niagara Systems and with game code/Blueprint.
-
-Data channel assets define the payload as well as some transfer settings.
-Niagara Systems can read from and write to data channels via data interfaces.
-Blueprint and C++ code can also read from and write to data channels using its API functions.
-
- */
-UCLASS(BlueprintType, DisplayName = "Niagara Data Channel", MinimalAPI)
-class UNiagaraDataChannelAsset : public UObject
-{
-	GENERATED_BODY()
-
-	UPROPERTY(EditAnywhere, Category = DataChannel, Instanced)
-	TObjectPtr<UNiagaraDataChannel> DataChannel;
-	
-	#if WITH_EDITORONLY_DATA
-	/** When changing data channel types we cache the old channel and attempt to copy over any common properties from one to the other. */
-	UPROPERTY(Transient)
-	TObjectPtr<UNiagaraDataChannel> CachedPreChangeDataChannel;
-	#endif
-
-public:
-
-#if WITH_EDITOR
-	NIAGARA_API virtual void PreEditChange(class FProperty* PropertyAboutToChange)override;
-	NIAGARA_API virtual void PostEditChangeProperty(FPropertyChangedEvent& PropertyChangedEvent)override;
-#endif
-
-	[[nodiscard]] UNiagaraDataChannel* Get() const { return DataChannel; }
-};
-
-/** A reference to a Niagara Data Channel with a pre-typed access context for convenience. */
-USTRUCT(BlueprintType)
-struct FNiagaraDataChannelReference
-{
-	GENERATED_BODY()
-
-	[[nodiscard]] inline UNiagaraDataChannelAsset* GetDataChannel()const { return DataChannel; }
-	/** Returns the internal access context instance. Can be invalid. */
-	[[nodiscard]] NIAGARA_API const FNDCAccessContextInst& GetAccessContext()const;
-	/** Returns a usable access context for accessing the referenced NDC. If the internal context is valid it will include it's input data. */
-	[[nodiscard]] NIAGARA_API FNDCAccessContextInst& GetUsableAccessContext()const;
-	NIAGARA_API void SetDataChannel(UNiagaraDataChannelAsset* InDataChannel);	
-	NIAGARA_API void OnValueChanged();
-
-	UPROPERTY(EditAnywhere, BlueprintReadWrite, Category = DataChannel)
-	TObjectPtr<UNiagaraDataChannelAsset> DataChannel;
-
-	UPROPERTY(EditAnywhere, BlueprintReadWrite, Category = DataChannel, meta = (EditCondition = "bCustomAccessContext" , ShowOnlyInnerProperties, ShowNDCAccessContextInput))
-	mutable FNDCAccessContextInst AccessContext;
-
-	/** If true this reference will provide a custom access context. */
-	UPROPERTY(EditAnywhere, BlueprintReadOnly, Category = DataChannel, meta = (InlineEditCondition))
-	bool bCustomAccessContext = false;
-
-private:
 
-	void InitAccessContext(bool bTransferExisting = true)const;
-};
\ No newline at end of file
+#include "NiagaraDataChannelCommon.h"
+#include "NiagaraDataChannel.h"
+#include "NiagaraDataChannelAsset.h"
+#include "NiagaraDataChannelData.h"
+#include "NiagaraDataChannelAccessContext.h"
+#include "NiagaraDataChannelGameData.h"
+#include "NiagaraDataChannelLayoutInfo.h"
+#include "NiagaraDataSetCompiledData.h"
+#include "NiagaraDataChannelFunctionLibrary.h"
+#include "NiagaraDataChannelDebugUtilities.h"
+#include "NiagaraDataChannelPublishRequest.h"
+#include "NiagaraDataChannelReference.h"
\ No newline at end of file
diff --git a/Engine/Plugins/FX/Niagara/Source/NiagaraBlueprintNodes/Internal/K2Node_DataChannelBase.h b/Engine/Plugins/FX/Niagara/Source/NiagaraBlueprintNodes/Internal/K2Node_DataChannelBase.h
--- a/Engine/Plugins/FX/Niagara/Source/NiagaraBlueprintNodes/Internal/K2Node_DataChannelBase.h
+++ b/Engine/Plugins/FX/Niagara/Source/NiagaraBlueprintNodes/Internal/K2Node_DataChannelBase.h
@@ -4,6 +4,7 @@
 
 #include "CoreMinimal.h"
 #include "NiagaraDataChannel.h"
+#include "NiagaraDataChannelFunctionLibrary.h"
 #include "K2Node_CallFunction.h"
 #include "K2Node_DataChannelBase.generated.h"
 
