a few small fixes.

diff --git a/Engine/Plugins/FX/Niagara/Source/Niagara/Private/NiagaraDataChannelAsset.cpp b/Engine/Plugins/FX/Niagara/Source/Niagara/Private/NiagaraDataChannelAsset.cpp
--- a/Engine/Plugins/FX/Niagara/Source/Niagara/Private/NiagaraDataChannelAsset.cpp
+++ b/Engine/Plugins/FX/Niagara/Source/Niagara/Private/NiagaraDataChannelAsset.cpp
@@ -11,7 +11,7 @@ void UNiagaraDataChannelAsset::PreEditChange(class FProperty* PropertyAboutToCha
 {
 	Super::PreEditChange(PropertyAboutToChange);
 
-	if(PropertyAboutToChange->GetFName() == GET_MEMBER_NAME_CHECKED(UNiagaraDataChannelAsset, DataChannel) && DataChannel)
+	if(PropertyAboutToChange && PropertyAboutToChange->GetFName() == GET_MEMBER_NAME_CHECKED(UNiagaraDataChannelAsset, DataChannel) && DataChannel)
 	{
 		CachedPreChangeDataChannel = DataChannel;
 	}
@@ -21,11 +21,14 @@ void UNiagaraDataChannelAsset::PostEditChangeProperty(FPropertyChangedEvent& Pro
 {
 	Super::PostEditChangeProperty(PropertyChangedEvent);
 
-	if(CachedPreChangeDataChannel && DataChannel)
+	if(PropertyChangedEvent.GetPropertyName() == GET_MEMBER_NAME_CHECKED(UNiagaraDataChannelAsset, DataChannel))
 	{
-		UEngine::FCopyPropertiesForUnrelatedObjectsParams Params;
-		UEngine::CopyPropertiesForUnrelatedObjects(CachedPreChangeDataChannel, DataChannel, Params);
-		CachedPreChangeDataChannel = nullptr;
+		if(CachedPreChangeDataChannel && DataChannel)
+		{
+			UEngine::FCopyPropertiesForUnrelatedObjectsParams Params;
+			UEngine::CopyPropertiesForUnrelatedObjects(CachedPreChangeDataChannel, DataChannel, Params);
+			CachedPreChangeDataChannel = nullptr;
+		}
 	}
 }
 
diff --git a/Engine/Plugins/FX/Niagara/Source/Niagara/Private/NiagaraDataChannelDebugUtilities.cpp b/Engine/Plugins/FX/Niagara/Source/Niagara/Private/NiagaraDataChannelDebugUtilities.cpp
--- a/Engine/Plugins/FX/Niagara/Source/Niagara/Private/NiagaraDataChannelDebugUtilities.cpp
+++ b/Engine/Plugins/FX/Niagara/Source/Niagara/Private/NiagaraDataChannelDebugUtilities.cpp
@@ -167,7 +167,7 @@ void FNiagaraDataChannelDebugUtilities::DumpAllWritesToLog()
 			, Request.bVisibleToCPUSims ? TEXT("X") : TEXT(" ")
 			, Request.bVisibleToGPUSims ? TEXT("X") : TEXT(" ")
 			, *TickGroupToString(Request.TickGroup)
-			, Request.Data->Num()
+			, Request.Data ? Request.Data->Num() : -1
 			, *GetPathNameSafe(Request.Channel.Get())
 			, *DataString);
 		}
