- Use the component transform for local space vs the simulation one
- When disabling requires current frame data (for performance) these could be out of sync, which mismatches other renderers on the RT (i.e. sprites & meshes)

#tests Test level with impact decals on a moving object

diff --git a/Engine/Plugins/FX/Niagara/Source/Niagara/Private/NiagaraRendererDecals.cpp b/Engine/Plugins/FX/Niagara/Source/Niagara/Private/NiagaraRendererDecals.cpp
--- a/Engine/Plugins/FX/Niagara/Source/Niagara/Private/NiagaraRendererDecals.cpp
+++ b/Engine/Plugins/FX/Niagara/Source/Niagara/Private/NiagaraRendererDecals.cpp
@@ -148,9 +148,10 @@ FNiagaraDynamicDataBase* FNiagaraRendererDecals::GenerateDynamicData(const FNiag
 	WeakOwnerComponent = OwnerComponent;
 	const UNiagaraDecalRendererProperties* RendererProperties = CastChecked<const UNiagaraDecalRendererProperties>(InProperties);
 
+	// Note: We use the OwnerComponent for local space simulations as disabling requires current frame data means the one inside the simulation is old and causes a mismatch with sprite / mesh renderers
 	const bool bUseLocalSpace = UseLocalSpace(Proxy);
-	const FTransform SimToWorld = SystemInstance->GetLWCSimToWorld(bUseLocalSpace);
-	const FNiagaraPosition DefaultPos = bUseLocalSpace ? FVector::ZeroVector : SystemInstance->GetWorldTransform().GetLocation();
+	const FTransform SimToWorld = bUseLocalSpace ? OwnerComponent->GetComponentTransform() : SystemInstance->GetLWCSimToWorld(false);
+	const FNiagaraPosition DefaultPos = FVector3f::ZeroVector;
 	const FQuat4f DefaultRot = UNiagaraDecalRendererProperties::GetDefaultOrientation();
 	const FVector3f DefaultSize = UNiagaraDecalRendererProperties::GetDefaultDecalSize();
 	const FNiagaraBool DefaultVisible = UNiagaraDecalRendererProperties::GetDefaultDecalVisible();
diff --git a/Engine/Plugins/FX/Niagara/Source/Niagara/Private/NiagaraRendererLights.cpp b/Engine/Plugins/FX/Niagara/Source/Niagara/Private/NiagaraRendererLights.cpp
--- a/Engine/Plugins/FX/Niagara/Source/Niagara/Private/NiagaraRendererLights.cpp
+++ b/Engine/Plugins/FX/Niagara/Source/Niagara/Private/NiagaraRendererLights.cpp
@@ -153,12 +153,13 @@ FNiagaraDynamicDataBase* FNiagaraRendererLights::GenerateDynamicData(const FNiag
 	FNiagaraSystemInstance* SystemInstance = Emitter->GetParentSystemInstance();
 	FNiagaraDynamicDataLights* DynamicData = new FNiagaraDynamicDataLights(Emitter);
 
+	USceneComponent* OwnerComponent = SystemInstance->GetAttachComponent();
+	check(OwnerComponent);
 
-	// This used to use Proxy->GetLocalToWorld(), but that's a bad thing to do here, because the proxy gets updated on the render thread,
-	// and this function happens during EndOfFrame updates. So instead, use the most up-to-date transform here (fixes local-space frame-behind issues)
+	// Note: We use the OwnerComponent for local space simulations as disabling requires current frame data means the one inside the simulation is old and causes a mismatch with sprite / mesh renderers
 	const bool bUseLocalSpace = UseLocalSpace(Proxy);
-	const FTransform SimToWorld = SystemInstance->GetLWCSimToWorld(bUseLocalSpace);
-	const FVector3f DefaultSimPos = bUseLocalSpace ? FVector3f::ZeroVector : FVector3f(SystemInstance->GetWorldTransform().GetLocation());
+	const FTransform SimToWorld = bUseLocalSpace ? OwnerComponent->GetComponentTransform() : SystemInstance->GetLWCSimToWorld(false);
+	const FVector3f DefaultSimPos = FVector3f::ZeroVector;
 
 	const FNiagaraParameterStore& ParameterStore = Emitter->GetRendererBoundVariables();
 	const FVector3f DefaultPos = ParameterStore.GetParameterValueOrDefault(Properties->PositionBinding.GetParamMapBindableVariable(), DefaultSimPos);
