- Stateless Curve Editor Improvements
- Moved the existing LW curve editor into compact curve editor
- Added a new curve editor for distributions which is closer to the existing curve editor
- Exposed some of the common widgets shared between the editors
- Future change will be to wrap the DI inside the distribution adapter and share the editors, which will be useful for the compact editor also

diff --git a/Engine/Plugins/FX/Niagara/Source/NiagaraEditor/Private/Stateless/NiagaraDistributionPropertyCustomization.cpp b/Engine/Plugins/FX/Niagara/Source/NiagaraEditor/Private/Stateless/NiagaraDistributionPropertyCustomization.cpp
--- a/Engine/Plugins/FX/Niagara/Source/NiagaraEditor/Private/Stateless/NiagaraDistributionPropertyCustomization.cpp
+++ b/Engine/Plugins/FX/Niagara/Source/NiagaraEditor/Private/Stateless/NiagaraDistributionPropertyCustomization.cpp
@@ -84,7 +84,7 @@ public:
 		}
 		else
 		{
-			EAxisList::Type ComponentAxis = FNiagaraEditorUtilities::VectorComponentToAxis(EditorNumChannels, ChannelIndex);
+			const EAxisList::Type ComponentAxis = FNiagaraEditorUtilities::VectorComponentToAxis(EditorNumChannels, ChannelIndex);
 			if (ComponentAxis != EAxisList::None)
 			{
 				return AxisDisplayInfo::GetAxisDisplayNameShort(ComponentAxis);
@@ -105,6 +105,20 @@ public:
 		return FText();
 	}
 
+	virtual FText GetChannelToolTip(int32 ChannelIndex) const override
+	{
+		if (bDisplayAsColor == false)
+		{
+			const EAxisList::Type ComponentAxis = FNiagaraEditorUtilities::VectorComponentToAxis(EditorNumChannels, ChannelIndex);
+			if (ComponentAxis != EAxisList::None)
+			{
+				return AxisDisplayInfo::GetAxisDisplayName(ComponentAxis);
+			}
+		}
+
+		return FText();
+	}
+
 	virtual FSlateColor GetChannelColor(int32 ChannelIndex) const override
 	{
 		if (EditorNumChannels != 1)
@@ -629,15 +643,14 @@ public:
 
 	virtual const FRichCurve* GetCurveValue(int32 ChannelIndex) const override
 	{
-		if (ChannelIndex < EditorNumChannels &&
-			ChannelIndex < SourceDistribution->ChannelCurves.Num())
+		if (ChannelIndex < EditorNumChannels && ChannelIndex < SourceDistribution->ChannelCurves.Num())
 		{
 			return &SourceDistribution->ChannelCurves[ChannelIndex];
 		}
 		return nullptr;
 	}
 
-	virtual void SetCurveValue(int32 ChannelIndex, const FRichCurve& InValue) override
+	virtual void SetCurveValue(int32 ChannelIndex, const FRichCurve& InValue, bool bShouldTransact) override
 	{
 		UObject* OwnerObject = OwnerObjectWeak.Get();
 		const FRichCurve* CurrentValue = GetCurveValue(ChannelIndex);
@@ -653,9 +666,12 @@ public:
 				ContinuousTransactionIndex = GEditor->BeginTransaction(TransactionText);
 			}
 
-			const FScopedTransaction Transaction(TransactionText, ContinuousTransactionIndex.IsSet() == false);
-			OwnerObject->Modify();
-			PropertyHandle->NotifyPreChange();
+			const FScopedTransaction Transaction(TransactionText, ContinuousTransactionIndex.IsSet() == false && bShouldTransact);
+			if (bShouldTransact)
+			{
+				OwnerObject->Modify();
+				PropertyHandle->NotifyPreChange();
+			}
 
 			const int32 RequiredCurveCount = EditorNumChannels;
 			if (SourceDistribution->ChannelCurves.Num() < RequiredCurveCount)
@@ -668,7 +684,10 @@ public:
 			}
 
 			SourceDistribution->ChannelCurves[ChannelIndex] = InValue;
-			PropertyHandle->NotifyPostChange(bContinuousChangeActive ? EPropertyChangeType::Interactive : EPropertyChangeType::ValueSet);
+			if (bShouldTransact)
+			{
+				PropertyHandle->NotifyPostChange(bContinuousChangeActive ? EPropertyChangeType::Interactive : EPropertyChangeType::ValueSet);
+			}
 		}
 	}
 
diff --git a/Engine/Plugins/FX/Niagara/Source/NiagaraEditor/Private/Widgets/INiagaraDistributionAdapter.h b/Engine/Plugins/FX/Niagara/Source/NiagaraEditor/Private/Widgets/INiagaraDistributionAdapter.h
--- a/Engine/Plugins/FX/Niagara/Source/NiagaraEditor/Private/Widgets/INiagaraDistributionAdapter.h
+++ b/Engine/Plugins/FX/Niagara/Source/NiagaraEditor/Private/Widgets/INiagaraDistributionAdapter.h
@@ -66,6 +66,9 @@ public:
 	/** Gets the display name for a specific channel by index. */
 	virtual FText GetChannelDisplayName(int32 ChannelIndex) const = 0;
 
+	/** Gets any tooltip associated with this channel by index. */
+	virtual FText GetChannelToolTip(int32 ChannelIndex) const = 0;
+
 	/** Gets the color used to display a channel by index. */
 	virtual FSlateColor GetChannelColor(int32 ChannelIndex) const = 0;
 
@@ -132,7 +135,7 @@ public:
 	virtual const FRichCurve* GetCurveValue(int32 ChannelIndex) const = 0;
 
 	/** Set a curve value for the specified channel */
-	virtual void SetCurveValue(int32 ChannelIndex, const FRichCurve& InValue) = 0;
+	virtual void SetCurveValue(int32 ChannelIndex, const FRichCurve& InValue, bool bShouldTransact = true) = 0;
 
 	/** Are we allowed to read / write the interpolation mode for the current distribution mode. */
 	virtual bool AllowInterpolationMode() const = 0;
diff --git a/Engine/Plugins/FX/Niagara/Source/NiagaraEditor/Private/Widgets/SNiagaraDistributionCompactCurveEditor.cpp b/Engine/Plugins/FX/Niagara/Source/NiagaraEditor/Private/Widgets/SNiagaraDistributionCompactCurveEditor.cpp
new file mode 100644
--- /dev/null
+++ b/Engine/Plugins/FX/Niagara/Source/NiagaraEditor/Private/Widgets/SNiagaraDistributionCompactCurveEditor.cpp
@@ -0,0 +1,1006 @@
+// Copyright Epic Games, Inc. All Rights Reserved.
+
+#include "Widgets/SNiagaraDistributionCompactCurveEditor.h"
+
+#include "NiagaraEditorStyle.h"
+#include "ScopedTransaction.h"
+#include "SEnumCombo.h"
+#include "Widgets/INiagaraDistributionAdapter.h"
+#include "Widgets/Input/SButton.h"
+#include "Widgets/Input/SNumericEntryBox.h"
+#include "Widgets/Layout/SBox.h"
+#include "Widgets/Layout/SGridPanel.h"
+#include "Widgets/SLeafWidget.h"
+
+#define LOCTEXT_NAMESPACE "NiagaraDistributionCompactCurveEditor"
+
+class SNiagaraCompactCurveView : public SLeafWidget
+{
+	SLATE_BEGIN_ARGS(SNiagaraCompactCurveView)
+		: _Width(20)
+		, _Height(20)
+		, _CurveColor(FSlateColor(EStyleColor::White))
+		, _CurveStateSerialNumber(0)
+		{ }
+		SLATE_ARGUMENT(float, Width)
+		SLATE_ARGUMENT(float, Height)
+		SLATE_ARGUMENT(FSlateColor, CurveColor)
+		SLATE_ATTRIBUTE(FKeyHandle, SelectedKeyHandle)
+		SLATE_ATTRIBUTE(int32, CurveStateSerialNumber)
+	SLATE_END_ARGS()
+
+	void Construct(const FArguments& InArgs, const FRichCurve* InCurveToDisplay)
+	{
+		Width = InArgs._Width;
+		Height = InArgs._Height;
+		CurveColor = InArgs._CurveColor;
+		SelectedKeyHandle = InArgs._SelectedKeyHandle;
+		CurveStateSerialNumber = InArgs._CurveStateSerialNumber;
+		CurveToDisplay = InCurveToDisplay;
+
+		KeyColor = FSlateColor(EStyleColor::AccentGray);
+		SelectedKeyColor = FSlateColor(EStyleColor::Select);
+	}
+
+	virtual int32 OnPaint(const FPaintArgs& Args, const FGeometry& AllottedGeometry, const FSlateRect& MyCullingRect, FSlateWindowElementList& OutDrawElements, int32 LayerId, const FWidgetStyle& InWidgetStyle, bool bParentEnabled) const override
+	{
+		FKeyHandle NewSelectedKeyHandle = SelectedKeyHandle.Get();
+		int32 NewCurveStateSerialNumber = CurveStateSerialNumber.Get(0);
+
+		if (CurveStateSerialNumberCache != NewCurveStateSerialNumber || AllottedGeometryCache != AllottedGeometry)
+		{
+			SelectedKeyHandleCache = NewSelectedKeyHandle;
+			CurveStateSerialNumberCache = NewCurveStateSerialNumber;
+			AllottedGeometryCache = AllottedGeometry;
+			CacheLines(AllottedGeometry.Size.X, AllottedGeometry.Size.Y, true, true, true);
+		}
+		else if (SelectedKeyHandleCache != NewSelectedKeyHandle)
+		{
+			SelectedKeyHandleCache = NewSelectedKeyHandle;
+			CacheLines(AllottedGeometry.Size.X, AllottedGeometry.Size.Y, false, false, true);
+		}
+
+		FSlateDrawElement::MakeLines(OutDrawElements, LayerId, AllottedGeometry.ToPaintGeometry(), CurvePointsCache, ESlateDrawEffect::None, CurveColor.GetColor(InWidgetStyle), true, 2.0f);
+		for (const TArray<FVector2D>& KeyLinePoints : KeyLinesCache)
+		{
+			FSlateDrawElement::MakeLines(OutDrawElements, LayerId, AllottedGeometry.ToPaintGeometry(), KeyLinePoints, ESlateDrawEffect::None, KeyColor.GetColor(InWidgetStyle), true, 2.0f);
+		}
+		FSlateDrawElement::MakeLines(OutDrawElements, LayerId, AllottedGeometry.ToPaintGeometry(), SelectedKeyLineCache, ESlateDrawEffect::None, SelectedKeyColor.GetColor(InWidgetStyle), true, 3.0f);
+		return LayerId;
+	}
+
+	virtual FVector2D ComputeDesiredSize(float) const override
+	{
+		return FVector2D(Width, Height);
+	}
+
+private:
+	void CacheLines(float InWidth, float InHeight, bool bCacheCurve, bool bCacheKeys, bool bCacheSelectedKey) const
+	{
+		if (CurveToDisplay == nullptr)
+		{
+			return;
+		}
+
+		float TimeMin;
+		float TimeMax;
+		float ValueMin;
+		float ValueMax;
+		CurveToDisplay->GetTimeRange(TimeMin, TimeMax);
+		CurveToDisplay->GetValueRange(ValueMin, ValueMax);
+
+		float TimeRange = TimeMax != TimeMin ? TimeMax - TimeMin : 1;
+		float ValueRange = ValueMax != ValueMin ? ValueMax - ValueMin : 1;
+
+		if (bCacheCurve)
+		{
+			CurvePointsCache.Empty();
+			int32 Points = CurveToDisplay->GetNumKeys() * 20;
+			float TimeIncrement = TimeRange / (Points - 1);
+			for (int32 i = 0; i < Points; i++)
+			{
+				float Time = TimeMin + i * TimeIncrement;
+				float Value = CurveToDisplay->Eval(Time);
+
+				float NormalizedX = (Time - TimeMin) / TimeRange;
+				float NormalizedY = (Value - ValueMin) / ValueRange;
+				CurvePointsCache.Add(FVector2D(NormalizedX * InWidth, (1 - NormalizedY) * InHeight));
+			}
+		}
+
+		if (bCacheKeys)
+		{
+			KeyLinesCache.Empty();
+			for (auto KeyIterator(CurveToDisplay->GetKeyIterator()); KeyIterator; ++KeyIterator)
+			{
+				const FRichCurveKey& Key = *KeyIterator;
+				TArray<FVector2D>& KeyLinePoints = KeyLinesCache.AddDefaulted_GetRef();
+
+				float NormalizedX = (Key.Time - TimeMin) / TimeRange;
+				KeyLinePoints.Add(FVector2D(NormalizedX * InWidth, 0));
+				KeyLinePoints.Add(FVector2D(NormalizedX * InWidth, InHeight));
+			}
+		}
+
+		if (bCacheSelectedKey)
+		{
+			SelectedKeyLineCache.Empty();
+			for (auto KeyHandleIterator(CurveToDisplay->GetKeyHandleIterator()); KeyHandleIterator; ++KeyHandleIterator)
+			{
+				if (*KeyHandleIterator == SelectedKeyHandleCache)
+				{
+					const FRichCurveKey& SelectedKey = CurveToDisplay->GetKey(SelectedKeyHandleCache);
+					float NormalizedX = (SelectedKey.Time - TimeMin) / TimeRange;
+					SelectedKeyLineCache.Add(FVector2D(NormalizedX * InWidth, 0));
+					SelectedKeyLineCache.Add(FVector2D(NormalizedX * InWidth, InHeight));
+					break;
+				}
+			}
+		}
+	}
+
+	const FRichCurve* CurveToDisplay = nullptr;
+	mutable TArray<FVector2D> CurvePointsCache;
+	mutable TArray<TArray<FVector2D>> KeyLinesCache; 
+	mutable TArray<FVector2D> SelectedKeyLineCache;
+	mutable FGeometry AllottedGeometryCache;
+	mutable FKeyHandle SelectedKeyHandleCache;
+	mutable int32 CurveStateSerialNumberCache = 0;
+	float Width;
+	float Height;
+	FSlateColor CurveColor;
+	FSlateColor KeyColor;
+	FSlateColor SelectedKeyColor;
+	TAttribute<FKeyHandle> SelectedKeyHandle;
+	TAttribute<int32> CurveStateSerialNumber;
+};
+
+class SNiagaraCompactCurveKeyHandle : public SBorder
+{
+	DECLARE_DELEGATE_OneParam(FOnMoved, float /* LocalDelta */);
+
+	SLATE_BEGIN_ARGS(SNiagaraCompactCurveKeyHandle)
+		: _Size(10, 10)
+	{ }
+		SLATE_ARGUMENT(FVector2f, Size)
+		SLATE_ATTRIBUTE(bool, IsSelected)
+		SLATE_EVENT(FSimpleDelegate, OnSelected)
+		SLATE_EVENT(FOnMoved, OnMoved)
+	SLATE_END_ARGS()
+
+	void Construct(const FArguments& InArgs)
+	{
+		Size = InArgs._Size;
+		IsSelected = InArgs._IsSelected;
+		OnSelectedDelegate = InArgs._OnSelected;
+		OnMovedDelegate = InArgs._OnMoved;
+		SBorder::Construct(SBorder::FArguments()
+			.BorderImage(FNiagaraEditorStyle::Get().GetBrush("NiagaraEditor.DistributionEditor.CurveKeyHandle"))
+			.BorderBackgroundColor(this, &SNiagaraCompactCurveKeyHandle::GetHandleColor)
+			.Padding(0)
+			[
+				SNew(SBox)
+				.WidthOverride(Size.X)
+				.HeightOverride(Size.Y)
+			]);
+	}
+
+	virtual FReply OnMouseButtonDown(const FGeometry& MyGeometry, const FPointerEvent& MouseEvent) override
+	{
+		if (MouseEvent.GetEffectingButton() == EKeys::LeftMouseButton)
+		{
+			FVector2D MouseLocation = MouseEvent.GetScreenSpacePosition();
+			LastMouseLocation = MouseLocation.X;
+			return FReply::Handled().CaptureMouse(SharedThis(this));
+		}
+		return FReply::Unhandled();
+	}
+
+	virtual FReply OnMouseButtonUp(const FGeometry& MyGeometry, const FPointerEvent& MouseEvent) override
+	{
+		if (this->HasMouseCapture())
+		{
+			if (IsSelected.Get() == false)
+			{
+				OnSelectedDelegate.ExecuteIfBound();
+			}
+			return FReply::Handled().ReleaseMouseCapture();
+		}
+		return FReply::Unhandled();
+	}
+
+	virtual FReply OnMouseMove(const FGeometry& MyGeometry, const FPointerEvent& MouseEvent) override
+	{
+		if (this->HasMouseCapture())
+		{
+			FVector2D MouseLocation = MouseEvent.GetScreenSpacePosition();
+			OnMovedDelegate.ExecuteIfBound((MouseLocation.X - LastMouseLocation) / MyGeometry.Scale);
+			LastMouseLocation = MouseLocation.X;
+			return FReply::Handled();
+		}
+		return FReply::Unhandled();
+	}
+
+private:
+	FSlateColor GetHandleColor() const
+	{
+		return IsSelected.Get()
+			? FSlateColor(EStyleColor::Select)
+			: FSlateColor(EStyleColor::White);
+	}
+
+private:
+	FVector2f Size;
+	TAttribute<bool> IsSelected;
+	FSimpleDelegate OnSelectedDelegate;
+	FOnMoved OnMovedDelegate;
+	float LastMouseLocation;
+};
+
+class SNiagaraCompactCurveKeySelector : public SCompoundWidget
+{
+	DECLARE_DELEGATE_OneParam(FOnKeySelected, FKeyHandle /* SelectedKeyHandle */);
+	DECLARE_DELEGATE_TwoParams(FOnKeyMoved, FKeyHandle /* SelectedKeyHandle */, float /* TimeDelta */)
+
+	SLATE_BEGIN_ARGS(SNiagaraCompactCurveKeySelector)
+		: _HandleSize(10, 10)
+		, _CurveStateSerialNumber(0)
+	{ }
+		SLATE_ARGUMENT(FVector2f, HandleSize)
+		SLATE_ATTRIBUTE(FKeyHandle, SelectedKeyHandle)
+		SLATE_ATTRIBUTE(int32, CurveStateSerialNumber)
+		SLATE_EVENT(FOnKeySelected, OnKeySelected)
+		SLATE_EVENT(FOnKeyMoved, OnKeyMoved)
+	SLATE_END_ARGS()
+
+private:
+	struct FKeyWidgetData
+	{
+		FKeyHandle Handle;
+		TSharedPtr<SNiagaraCompactCurveKeyHandle> HandleWidget;
+		FMargin BorderPadding;
+	};
+
+public:
+	void Construct(const FArguments& InArgs, FRichCurve* InEditCurve)
+	{
+		EditCurve = InEditCurve;
+		HandleSize = InArgs._HandleSize;
+		SelectedKeyHandle = InArgs._SelectedKeyHandle;
+		CurveStateSerialNumber = InArgs._CurveStateSerialNumber;
+		OnKeySelectedDelegate = InArgs._OnKeySelected;
+		OnKeyMovedDelegate = InArgs._OnKeyMoved;
+		LocalWidthCache = 100;
+
+		ChildSlot
+		[
+			SAssignNew(BorderOverlay, SOverlay)
+		];
+
+		UpdateKeyWidgets();
+	}
+
+	virtual void Tick(const FGeometry& AllottedGeometry, const double InCurrentTime, const float InDeltaTime) override
+	{
+		SCompoundWidget::Tick(AllottedGeometry, InCurrentTime, InDeltaTime);
+		int32 CurrentCurveStateSerialNumber = CurveStateSerialNumber.Get();
+		LocalWidthCache = AllottedGeometry.GetLocalSize().X - HandleSize.X;
+		if (CachedCurveStateSerialNumber != CurrentCurveStateSerialNumber || AllottedGeometry != AllottedGeometryCache)
+		{
+			CachedCurveStateSerialNumber = CurrentCurveStateSerialNumber;
+			UpdateKeyWidgets();
+		}
+		AllottedGeometryCache = AllottedGeometry;
+	}
+
+	void RequestWidgetUpdate() { bNeedsWidgetUpdate = true; }
+
+private:
+	void UpdateKeyWidgets()
+	{
+		if (KeyWidgetData.Num() != EditCurve->GetNumKeys() || bNeedsWidgetUpdate)
+		{
+			BorderOverlay->ClearChildren();
+			KeyWidgetData.Empty();
+			for (auto KeyHandleIterator = EditCurve->GetKeyHandleIterator(); KeyHandleIterator; KeyHandleIterator++)
+			{
+				int32 DataIndex = KeyWidgetData.Num();
+				FKeyWidgetData& KeyWidgetDataItem = KeyWidgetData.AddDefaulted_GetRef();
+				KeyWidgetDataItem.Handle = *KeyHandleIterator;
+				TAttribute<FMargin> BorderPadding = TAttribute<FMargin>::CreateSP(this, &SNiagaraCompactCurveKeySelector::GetKeyBorderPadding, DataIndex);
+				BorderOverlay->AddSlot()
+				.Padding(BorderPadding)
+				.HAlign(HAlign_Left)
+				[
+					SAssignNew(KeyWidgetDataItem.HandleWidget, SNiagaraCompactCurveKeyHandle)
+					.Size(HandleSize)
+					.IsSelected(this, &SNiagaraCompactCurveKeySelector::GetKeyIsSelected, KeyWidgetDataItem.Handle)
+					.OnSelected(this, &SNiagaraCompactCurveKeySelector::KeyHandleSelected, KeyWidgetDataItem.Handle)
+					.OnMoved(this, &SNiagaraCompactCurveKeySelector::KeyHandleMoved, KeyWidgetDataItem.Handle)
+				];
+			}
+			bNeedsWidgetUpdate = false;
+		}
+
+		float TimeMin;
+		float TimeMax;
+		float ValueMin;
+		float ValueMax;
+		EditCurve->GetTimeRange(TimeMin, TimeMax);
+		EditCurve->GetValueRange(ValueMin, ValueMax);
+		float TimeRange = TimeMax != TimeMin ? TimeMax - TimeMin : 1;
+		float ValueRange = ValueMax != ValueMin ? ValueMax - ValueMin : 1;
+		for (FKeyWidgetData& KeyWidgetDataItem : KeyWidgetData)
+		{
+			const FRichCurveKey& Key = EditCurve->GetKey(KeyWidgetDataItem.Handle);
+			float NormalizedX = (Key.Time - TimeMin) / TimeRange;
+			KeyWidgetDataItem.BorderPadding = FMargin(LocalWidthCache * NormalizedX, 0, 0, 0);
+		}
+	}
+
+	FMargin GetKeyBorderPadding(int32 KeyWidgetDataIndex) const
+	{
+		return KeyWidgetDataIndex < KeyWidgetData.Num() ? KeyWidgetData[KeyWidgetDataIndex].BorderPadding : FMargin(0);
+	}
+
+	bool GetKeyIsSelected(FKeyHandle KeyHandle) const
+	{
+		return KeyHandle == SelectedKeyHandle.Get();
+	}
+
+	void KeyHandleSelected(FKeyHandle KeyHandle)
+	{
+		OnKeySelectedDelegate.ExecuteIfBound(KeyHandle);
+	}
+
+	void KeyHandleMoved(float LocalDelta, FKeyHandle KeyHandle)
+	{
+		float TimeMin;
+		float TimeMax;
+		EditCurve->GetTimeRange(TimeMin, TimeMax);
+		float TimeRange = TimeMax - TimeMin;
+		float NormalizedDelta = FMath::IsNearlyZero(LocalWidthCache) ? 0.0f : LocalDelta / LocalWidthCache;
+		if (SelectedKeyHandle.Get() != KeyHandle)
+		{
+			OnKeySelectedDelegate.ExecuteIfBound(KeyHandle);
+		}
+		OnKeyMovedDelegate.ExecuteIfBound(KeyHandle, NormalizedDelta * TimeRange);
+	}
+
+private:
+	FRichCurve* EditCurve = nullptr;
+	FVector2f HandleSize;
+	TAttribute<FKeyHandle> SelectedKeyHandle;
+	TAttribute<int32> CurveStateSerialNumber;
+	FOnKeySelected OnKeySelectedDelegate;
+	FOnKeyMoved OnKeyMovedDelegate;
+
+	FGeometry AllottedGeometryCache;
+	float LocalWidthCache;
+
+	TSharedPtr<SOverlay> BorderOverlay;
+	int32 CachedCurveStateSerialNumber = INDEX_NONE;
+	TArray<FKeyWidgetData> KeyWidgetData;
+
+	bool bNeedsWidgetUpdate = false;
+};
+
+void SNiagaraDistributionCompactCurveEditor::Construct(const FArguments& InArgs, TSharedRef<INiagaraDistributionAdapter> InDistributionAdapter, int32 InChannelIndex)
+{
+	DistributionAdapter = InDistributionAdapter;
+	ChannelIndex = InChannelIndex;
+	const FRichCurve* EditCurvePtr = DistributionAdapter->GetCurveValue(ChannelIndex);
+	if (EditCurvePtr != nullptr)
+	{
+		EditCurve = *EditCurvePtr;
+	}
+
+	if (EditCurve.GetNumKeys() == 0)
+	{
+		EditCurve.AddKey(0, 0);
+	}
+
+	UpdateSelectedKeyIndex(0);
+
+	float KeyHandleWidth = 16;
+	float KeyHandleHeight = 16;
+	ChildSlot
+	[
+		SNew(SVerticalBox)
+		+ SVerticalBox::Slot()
+		.AutoHeight()
+		.Padding(0, 0, 0, 3)
+		[
+			SNew(SHorizontalBox)
+			+ SHorizontalBox::Slot()
+			.AutoWidth()
+			.Padding(0, 0, 3, 0)
+			[
+				SNew(SVerticalBox)
+				+ SVerticalBox::Slot()
+				.AutoHeight()
+				.Padding(0, 0, 0, 2)
+				[
+					SNew(SNumericEntryBox<int32>)
+					.Font(FAppStyle::Get().GetFontStyle("PropertyWindow.NormalFont"))
+					.Value(this, &SNiagaraDistributionCompactCurveEditor::GetSelectedKeyIndex)
+					.OnValueChanged(this, &SNiagaraDistributionCompactCurveEditor::SelectedKeyIndexChanged)
+					.OnValueCommitted(this, &SNiagaraDistributionCompactCurveEditor::SelectedKeyIndexComitted)
+					.AllowSpin(true)
+					.MinValue(0)
+					.MaxValue(this, &SNiagaraDistributionCompactCurveEditor::GetMaxKeyIndex)
+					.MinSliderValue(0)
+					.MaxSliderValue(this, &SNiagaraDistributionCompactCurveEditor::GetMaxKeyIndex)
+					.Delta(1)
+					.MinDesiredValueWidth(8)
+					.LabelVAlign(VAlign_Center)
+					.Label()
+					[
+						SNew(STextBlock)
+						.Font(FAppStyle::Get().GetFontStyle("PropertyWindow.NormalFont"))
+						.Text(LOCTEXT("KeyLabel", "Key"))
+					]
+				]
+				+ SVerticalBox::Slot()
+				.AutoHeight()
+				[
+					SNew(SHorizontalBox)
+					+ SHorizontalBox::Slot()
+					.AutoWidth()
+					.Padding(0, 0, 2, 0)
+					[
+						SNew(SButton)
+						.ButtonStyle( FAppStyle::Get(), "SimpleButton" )
+						.OnClicked(this, &SNiagaraDistributionCompactCurveEditor::AddKeyButtonClicked)
+						.ToolTipText(LOCTEXT("AddKeyToolTip", "Add a key"))
+						.ContentPadding(0.0f)
+						[ 
+							SNew( SImage )
+							.Image( FAppStyle::GetBrush("Icons.PlusCircle") )
+							.ColorAndOpacity( FSlateColor::UseForeground() )
+						]
+					]
+					+ SHorizontalBox::Slot()
+					.AutoWidth()
+					.Padding(0, 0, 2, 0)
+					[
+						SNew(SButton)
+						.ButtonStyle(FAppStyle::Get(), "SimpleButton")
+						.OnClicked(this, &SNiagaraDistributionCompactCurveEditor::DeleteKeyButtonClicked)
+						.ToolTipText(LOCTEXT("DeleteKeyToolTip", "Delete this key"))
+						.ContentPadding(0.0f)
+						[
+							SNew(SImage)
+							.Image(FAppStyle::GetBrush("Icons.Delete"))
+							.ColorAndOpacity(FSlateColor::UseForeground())
+						]
+					]
+				]
+				+ SVerticalBox::Slot()
+				.AutoHeight()
+				[
+					SNew(SHorizontalBox)
+						+ SHorizontalBox::Slot()
+						.AutoWidth()
+						.Padding(0, 0, 2, 0)
+						[
+							SNew(SButton)
+								.ButtonStyle(FAppStyle::Get(), "SimpleButton")
+								.OnClicked(this, &SNiagaraDistributionCompactCurveEditor::FlipCurve, ECurveFlipDirection::Horizontal)
+								.ToolTipText(LOCTEXT("FlipHorizontalToolTip", "Flip curve horizontally"))
+								.ContentPadding(0.0f)
+								[
+									SNew(SImage)
+										.Image(FAppStyle::GetBrush("Icons.FlipHorizontal"))
+										.ColorAndOpacity(FSlateColor::UseForeground())
+								]
+						]
+						+ SHorizontalBox::Slot()
+						.AutoWidth()
+						.Padding(0, 0, 2, 0)
+						[
+							SNew(SButton)
+								.ButtonStyle(FAppStyle::Get(), "SimpleButton")
+								.OnClicked(this, &SNiagaraDistributionCompactCurveEditor::FlipCurve, ECurveFlipDirection::Vertical)
+								.ToolTipText(LOCTEXT("FlipVerticalToolTip", "Flip curve vertically"))
+								.ContentPadding(0.0f)
+								[
+									SNew(SImage)
+										.Image(FAppStyle::GetBrush("Icons.FlipVertical"))
+										.ColorAndOpacity(FSlateColor::UseForeground())
+								]
+						]
+				]
+			]
+			// Curve view and key selector.
+			+ SHorizontalBox::Slot()
+			.Padding(2)
+			[
+				SNew(SGridPanel)
+				.FillRow(0, 1)
+				.FillColumn(0, 1)
+				// Value Labels
+				+ SGridPanel::Slot(1, 0)
+				[
+					SNew(SVerticalBox)
+					+ SVerticalBox::Slot()
+					.VAlign(VAlign_Top)
+					[
+						SNew(STextBlock)
+						.Font(FAppStyle::Get().GetFontStyle("SmallFont"))
+						.Text(this, &SNiagaraDistributionCompactCurveEditor::GetCurveValueMaxText)
+					]
+					+ SVerticalBox::Slot()
+					.VAlign(VAlign_Bottom)
+					[
+						SNew(STextBlock)
+						.Font(FAppStyle::Get().GetFontStyle("SmallFont"))
+						.Text(this, &SNiagaraDistributionCompactCurveEditor::GetCurveValueMinText)
+					]
+				]
+				// Curve view
+				+ SGridPanel::Slot(0, 0)
+				.Padding(FMargin(KeyHandleWidth / 2, 3, KeyHandleWidth / 2, 3))
+				[
+					SNew(SNiagaraCompactCurveView, &EditCurve)
+					.Width(1000)
+					.Height(50)
+					.CurveColor(InArgs._CurveColor)
+					.SelectedKeyHandle(this, &SNiagaraDistributionCompactCurveEditor::GetSelectedKeyHandle)
+					.CurveStateSerialNumber(this, &SNiagaraDistributionCompactCurveEditor::GetCurveStateSerialNumber)
+				]
+				// Time Labels
+				+ SGridPanel::Slot(0, 2)
+				[
+					SNew(SHorizontalBox)
+					+ SHorizontalBox::Slot()
+					.HAlign(HAlign_Left)
+					[
+						SNew(STextBlock)
+						.Font(FAppStyle::Get().GetFontStyle("SmallFont"))
+						.Text(this, &SNiagaraDistributionCompactCurveEditor::GetCurveTimeMinText)
+					]
+					+ SHorizontalBox::Slot()
+					.HAlign(HAlign_Right)
+					[
+						SNew(STextBlock)
+						.Font(FAppStyle::Get().GetFontStyle("SmallFont"))
+						.Text(this, &SNiagaraDistributionCompactCurveEditor::GetCurveTimeMaxText)
+					]
+				]
+				// key selector
+				+ SGridPanel::Slot(0, 1)
+				.Padding(0, 0, 0, 3)
+				[
+					SAssignNew(KeySelectorWidget, SNiagaraCompactCurveKeySelector, &EditCurve)
+					.HandleSize(FVector2f(KeyHandleWidth, KeyHandleHeight))
+					.SelectedKeyHandle(this, &SNiagaraDistributionCompactCurveEditor::GetSelectedKeyHandle)
+					.CurveStateSerialNumber(this, &SNiagaraDistributionCompactCurveEditor::GetCurveStateSerialNumber)
+					.OnKeySelected(this, &SNiagaraDistributionCompactCurveEditor::OnKeySelectorKeySelected)
+					.OnKeyMoved(this, &SNiagaraDistributionCompactCurveEditor::OnKeySelectorKeyMoved)
+				]
+			]
+		]
+		+ SVerticalBox::Slot()
+		.AutoHeight()
+		[
+			SNew(SHorizontalBox)
+			+ SHorizontalBox::Slot()
+			.Padding(0, 0, 5, 0)
+			[
+				SNew(SNumericEntryBox<float>)
+				.Font(FAppStyle::Get().GetFontStyle("PropertyWindow.NormalFont"))
+				.Value(this, &SNiagaraDistributionCompactCurveEditor::GetSelectedKeyTime)
+				.OnValueChanged(this, &SNiagaraDistributionCompactCurveEditor::SelectedKeyTimeChanged)
+				.OnValueCommitted(this, &SNiagaraDistributionCompactCurveEditor::SelectedKeyTimeComitted)
+				.OnBeginSliderMovement(this, &SNiagaraDistributionCompactCurveEditor::BeginSliderMovement)
+				.OnEndSliderMovement(this, &SNiagaraDistributionCompactCurveEditor::EndSliderMovement)
+				.AllowSpin(true)
+				.MinValue(TOptional<float>())
+				.MaxValue(TOptional<float>())
+				.MinSliderValue(TOptional<float>())
+				.MaxSliderValue(TOptional<float>())
+				.LabelVAlign(VAlign_Center)
+				.Label()
+				[
+					SNew(STextBlock)
+					.Font(FAppStyle::Get().GetFontStyle("PropertyWindow.NormalFont"))
+					.Text(LOCTEXT("TimeLabel", "Time"))
+				]
+			]
+			+ SHorizontalBox::Slot()
+			.Padding(0, 0, 5, 0)
+			[
+				SNew(SNumericEntryBox<float>)
+				.Font(FAppStyle::Get().GetFontStyle("PropertyWindow.NormalFont"))
+				.Value(this, &SNiagaraDistributionCompactCurveEditor::GetSelectedKeyValue)
+				.OnValueChanged(this, &SNiagaraDistributionCompactCurveEditor::SelectedKeyValueChanged)
+				.OnValueCommitted(this, &SNiagaraDistributionCompactCurveEditor::SelectedKeyValueComitted)
+				.OnBeginSliderMovement(this, &SNiagaraDistributionCompactCurveEditor::BeginSliderMovement)
+				.OnEndSliderMovement(this, &SNiagaraDistributionCompactCurveEditor::EndSliderMovement)
+				.AllowSpin(true)
+				.MinValue(TOptional<float>())
+				.MaxValue(TOptional<float>())
+				.MinSliderValue(TOptional<float>())
+				.MaxSliderValue(TOptional<float>())
+				.LabelVAlign(VAlign_Center)
+				.Label()
+				[
+					SNew(STextBlock)
+					.Font(FAppStyle::Get().GetFontStyle("PropertyWindow.NormalFont"))
+					.Text(LOCTEXT("ValueLabel", "Value"))
+				]
+			]
+			+ SHorizontalBox::Slot()
+			[
+				SNew(SHorizontalBox)
+				+ SHorizontalBox::Slot()
+				.VAlign(VAlign_Center)
+				.AutoWidth()
+				.Padding(0, 0, 3, 0)
+				[
+					SNew(STextBlock)
+					.Font(FAppStyle::Get().GetFontStyle("PropertyWindow.NormalFont"))
+					.Text(LOCTEXT("InterpLabel", "Interp"))
+				]
+				+ SHorizontalBox::Slot()
+				[
+					SNew(SEnumComboBox, StaticEnum<ERichCurveInterpMode>())
+					.Font(FAppStyle::Get().GetFontStyle("PropertyWindow.NormalFont"))
+					.ContentPadding(FMargin(2, 0))
+					.CurrentValue(this, &SNiagaraDistributionCompactCurveEditor::GetSelectedKeyInterpMode)
+					.OnEnumSelectionChanged(this, &SNiagaraDistributionCompactCurveEditor::SelectedKeyInterpModeChanged)
+				]
+			]
+		]
+	];
+}
+
+TOptional<int32> SNiagaraDistributionCompactCurveEditor::GetSelectedKeyIndex() const
+{
+	return SelectedKeyIndex;
+}
+
+TOptional<int32> SNiagaraDistributionCompactCurveEditor::GetMaxKeyIndex() const
+{
+	return EditCurve.GetNumKeys() - 1;
+}
+
+void SNiagaraDistributionCompactCurveEditor::SelectedKeyIndexChanged(int32 InValue)
+{
+	UpdateSelectedKeyIndex(InValue);
+}
+
+void SNiagaraDistributionCompactCurveEditor::SelectedKeyIndexComitted(int32 InValue, ETextCommit::Type CommitInfo)
+{
+	if (CommitInfo == ETextCommit::OnEnter || CommitInfo == ETextCommit::OnUserMovedFocus)
+	{
+		UpdateSelectedKeyIndex(InValue);
+	}
+}
+
+void SNiagaraDistributionCompactCurveEditor::UpdateSelectedKeyIndex(int32 NewKeyIndex)
+{
+	SelectedKeyIndex = NewKeyIndex;
+	int32 KeyIndex = 0;
+	for (auto KeyHandleIterator = EditCurve.GetKeyHandleIterator(); KeyHandleIterator; KeyHandleIterator++)
+	{
+		if (KeyIndex == NewKeyIndex)
+		{
+			SelectedKeyHandle = *KeyHandleIterator;
+			break;
+		}
+		KeyIndex++;
+	}
+	CurveStateSerialNumber++;
+}
+
+TOptional<float> SNiagaraDistributionCompactCurveEditor::GetSelectedKeyTime() const
+{
+	UpdateCachedValues();
+	return KeyTimeCache;
+}
+
+void SNiagaraDistributionCompactCurveEditor::SelectedKeyTimeChanged(float InValue)
+{
+	UpdateSelectedKeyTime(InValue);
+}
+
+void SNiagaraDistributionCompactCurveEditor::SelectedKeyTimeComitted(float InValue, ETextCommit::Type CommitInfo)
+{
+	if (CommitInfo == ETextCommit::OnEnter || CommitInfo == ETextCommit::OnUserMovedFocus)
+	{
+		UpdateSelectedKeyTime(InValue);
+	}
+}
+
+void SNiagaraDistributionCompactCurveEditor::UpdateSelectedKeyTime(float NewKeyTime)
+{
+	if (GetSelectedKeyTime() != NewKeyTime)
+	{
+		EditCurve.SetKeyTime(SelectedKeyHandle, NewKeyTime);
+		DistributionAdapter->SetCurveValue(ChannelIndex, EditCurve);
+
+		int32 NewKeyIndex = INDEX_NONE;
+		int32 KeyIndex = 0;
+		for (auto KeyHandleIterator = EditCurve.GetKeyHandleIterator(); KeyHandleIterator; KeyHandleIterator++)
+		{
+			if (SelectedKeyHandle == *KeyHandleIterator)
+			{
+				NewKeyIndex = KeyIndex;
+				break;
+			}
+			KeyIndex++;
+		}
+
+		SelectedKeyIndex = NewKeyIndex;
+		KeyTimeCache = NewKeyTime;
+		CurveStateSerialNumber++;
+	}
+}
+
+TOptional<float> SNiagaraDistributionCompactCurveEditor::GetSelectedKeyValue() const
+{
+	UpdateCachedValues();
+	return KeyValueCache;
+}
+
+void SNiagaraDistributionCompactCurveEditor::SelectedKeyValueChanged(float InValue)
+{
+	UpdateSelectedKeyValue(InValue);
+}
+
+void SNiagaraDistributionCompactCurveEditor::SelectedKeyValueComitted(float InValue, ETextCommit::Type CommitInfo)
+{
+	if (CommitInfo == ETextCommit::OnEnter || CommitInfo == ETextCommit::OnUserMovedFocus)
+	{
+		UpdateSelectedKeyValue(InValue);
+	}
+}
+
+void SNiagaraDistributionCompactCurveEditor::UpdateSelectedKeyValue(float NewKeyValue)
+{
+	if (GetSelectedKeyValue() != NewKeyValue)
+	{
+		EditCurve.SetKeyValue(SelectedKeyHandle, NewKeyValue);
+		DistributionAdapter->SetCurveValue(ChannelIndex, EditCurve);
+		KeyValueCache = NewKeyValue;
+		CurveStateSerialNumber++;
+	}
+}
+
+int32 SNiagaraDistributionCompactCurveEditor::GetSelectedKeyInterpMode() const
+{
+	UpdateCachedValues();
+	return KeyInterpModeCache;
+}
+
+void SNiagaraDistributionCompactCurveEditor::SelectedKeyInterpModeChanged(int32 InNewValue, ESelectInfo::Type InInfo)
+{
+	if (GetSelectedKeyInterpMode() != InNewValue)
+	{
+		EditCurve.SetKeyInterpMode(SelectedKeyHandle, (ERichCurveInterpMode)InNewValue);
+		DistributionAdapter->SetCurveValue(ChannelIndex, EditCurve);
+		KeyInterpModeCache = InNewValue;
+		CurveStateSerialNumber++;
+	}
+}
+
+void SNiagaraDistributionCompactCurveEditor::BeginSliderMovement()
+{
+	DistributionAdapter->BeginContinuousChange();
+}
+
+void SNiagaraDistributionCompactCurveEditor::EndSliderMovement(float Value)
+{
+	DistributionAdapter->EndContinuousChange();
+}
+
+FReply SNiagaraDistributionCompactCurveEditor::AddKeyButtonClicked()
+{
+	int32 NumKeys = EditCurve.GetNumKeys();
+	int32 NewSelectedKeyIndex;
+	if (NumKeys == 0)
+	{
+		EditCurve.AddKey(0, 0);
+		NewSelectedKeyIndex = 0;
+	}
+	else if (NumKeys == 1)
+	{
+		FRichCurveKey FirstKey = EditCurve.GetFirstKey();
+		EditCurve.AddKey(FirstKey.Time + 1, FirstKey.Value);
+		NewSelectedKeyIndex = 1;
+	}
+	else
+	{
+		if (SelectedKeyIndex < NumKeys - 1)
+		{
+			FKeyHandle NextKeyHandle = EditCurve.GetNextKey(SelectedKeyHandle);
+			FRichCurveKey SelectedKey = EditCurve.GetKey(SelectedKeyHandle);
+			FRichCurveKey NextKey = EditCurve.GetKey(NextKeyHandle);
+			EditCurve.AddKey((SelectedKey.Time + NextKey.Time) / 2, (SelectedKey.Value + NextKey.Value) / 2);
+			NewSelectedKeyIndex = SelectedKeyIndex + 1;
+		}
+		else
+		{
+			FKeyHandle PreviousKeyHandle = EditCurve.GetPreviousKey(SelectedKeyHandle);
+			FRichCurveKey SelectedKey = EditCurve.GetKey(SelectedKeyHandle);
+			FRichCurveKey PreviousKey = EditCurve.GetKey(PreviousKeyHandle);
+			EditCurve.AddKey(SelectedKey.Time + (SelectedKey.Time - PreviousKey.Time), SelectedKey.Value + (SelectedKey.Value - PreviousKey.Value));
+			NewSelectedKeyIndex = SelectedKeyIndex;
+		}
+	}
+
+	DistributionAdapter->SetCurveValue(ChannelIndex, EditCurve);
+	UpdateSelectedKeyIndex(NewSelectedKeyIndex);
+	CurveStateSerialNumber++;
+	return FReply::Handled();
+}
+
+FReply SNiagaraDistributionCompactCurveEditor::DeleteKeyButtonClicked()
+{
+	if (EditCurve.GetNumKeys() > 1)
+	{
+		EditCurve.DeleteKey(GetSelectedKeyHandle());
+		DistributionAdapter->SetCurveValue(ChannelIndex, EditCurve);
+		int32 NewKeyIndex = FMath::Clamp(SelectedKeyIndex - 1, 0, EditCurve.GetNumKeys() - 1);
+		UpdateSelectedKeyIndex(NewKeyIndex);
+		CurveStateSerialNumber++;
+		return FReply::Handled();
+	}
+	return FReply::Unhandled();
+}
+
+FReply SNiagaraDistributionCompactCurveEditor::FlipCurve(ECurveFlipDirection Direction)
+{
+	if (EditCurve.GetNumKeys() <= 0)
+	{
+		return FReply::Unhandled();
+	}
+
+	int32 NumKeys = EditCurve.GetNumKeys();
+	int32 NewSelectedKeyIndex = SelectedKeyIndex;
+
+	// Intialize the range
+	float RangeMin = FLT_MAX;
+	float RangeMax = -FLT_MAX;
+
+	// If flipping horizontally
+	if (Direction == ECurveFlipDirection::Horizontal)
+	{
+		// Get the range
+		float MinTime, MaxTime;
+		EditCurve.GetTimeRange(MinTime, MaxTime);
+		RangeMin = FMath::Min(MinTime, RangeMin);
+		RangeMax = FMath::Max(MaxTime, RangeMax);
+
+		TArray<TPair<float, float>> TimeValuePairs;
+		TArray<ERichCurveInterpMode> InterpModes;
+		TimeValuePairs.SetNum(NumKeys);
+		InterpModes.SetNum(NumKeys);
+
+		// Get the (flipped x, original y) value of all points
+		FKeyHandle KeyHandle = EditCurve.GetFirstKeyHandle();
+		for (int32 Index = 0; Index <= NumKeys - 1; ++Index)
+		{
+			FRichCurveKey Key = EditCurve.GetKey(KeyHandle);
+			TimeValuePairs[Index].Key = RangeMax - Key.Time + RangeMin;
+			TimeValuePairs[Index].Value = Key.Value;
+			InterpModes[Index] = EditCurve.GetKeyInterpMode(KeyHandle);
+			KeyHandle = EditCurve.GetNextKey(KeyHandle);
+		}
+
+		// Delete all existing keys
+		EditCurve.Reset();
+
+		// Add the flipped points back in
+		FKeyHandle NewKeyHandle;
+		for (int32 Index = 0; Index <= NumKeys - 1; ++Index)
+		{
+			NewKeyHandle = EditCurve.AddKey(TimeValuePairs[Index].Key, TimeValuePairs[Index].Value);
+			EditCurve.SetKeyInterpMode(NewKeyHandle, InterpModes[Index]);
+			//NOTE: key tangents are not handled right now since distribution curve editor currently don't support tangent editing from user
+		}
+
+		NewSelectedKeyIndex = NumKeys - 1 - SelectedKeyIndex;
+		if (KeySelectorWidget) 
+		{
+			KeySelectorWidget->RequestWidgetUpdate();
+		}
+	}
+
+	// If flipping vertically
+	if (Direction == ECurveFlipDirection::Vertical)
+	{
+		// Get the range
+		float MinVal, MaxVal;
+		EditCurve.GetValueRange(MinVal, MaxVal);
+		RangeMin = FMath::Min(MinVal, RangeMin);
+		RangeMax = FMath::Max(MaxVal, RangeMax);
+
+		// Flip the y value of all keys
+		FKeyHandle KeyHandle = EditCurve.GetFirstKeyHandle();
+		for (int32 Index = 0; Index <= NumKeys - 1; ++Index)
+		{
+			EditCurve.SetKeyValue(KeyHandle, RangeMax - EditCurve.GetKey(KeyHandle).Value + RangeMin);
+			KeyHandle = EditCurve.GetNextKey(KeyHandle);
+		}
+	}
+
+	DistributionAdapter->SetCurveValue(ChannelIndex, EditCurve);
+	UpdateSelectedKeyIndex(NewSelectedKeyIndex);
+	CurveStateSerialNumber++;
+	return FReply::Handled();
+}
+
+void SNiagaraDistributionCompactCurveEditor::OnKeySelectorKeySelected(FKeyHandle InSelectedKeyHandle)
+{
+	int32 KeyIndex = 0;
+	for (auto KeyHandleIterator = EditCurve.GetKeyHandleIterator(); KeyHandleIterator; KeyHandleIterator++)
+	{
+		if (*KeyHandleIterator == InSelectedKeyHandle)
+		{
+			SelectedKeyHandle = *KeyHandleIterator;
+			SelectedKeyIndex = KeyIndex;
+			CurveStateSerialNumber++;
+			break;
+		}
+		KeyIndex++;
+	}
+}
+
+void SNiagaraDistributionCompactCurveEditor::OnKeySelectorKeyMoved(FKeyHandle InSelectedKeyHandle, float TimeDelta)
+{
+	if (SelectedKeyHandle == InSelectedKeyHandle)
+	{
+		UpdateSelectedKeyTime(GetSelectedKeyTime().GetValue() + TimeDelta);
+	}
+}
+
+FText SNiagaraDistributionCompactCurveEditor::GetCurveTimeMinText() const
+{
+	UpdateCachedValues();
+	return CurveTimeMinTextCache;
+}
+
+FText SNiagaraDistributionCompactCurveEditor::GetCurveTimeMaxText() const
+{
+	UpdateCachedValues();
+	return CurveTimeMaxTextCache;
+}
+
+FText SNiagaraDistributionCompactCurveEditor::GetCurveValueMinText() const
+{
+	UpdateCachedValues();
+	return CurveValueMinTextCache;
+}
+
+FText SNiagaraDistributionCompactCurveEditor::GetCurveValueMaxText() const
+{
+	UpdateCachedValues();
+	return CurveValueMaxTextCache;
+}
+
+void SNiagaraDistributionCompactCurveEditor::UpdateCachedValues() const
+{
+	if (CacheCurveStateSerialNumber != CurveStateSerialNumber)
+	{
+		CacheCurveStateSerialNumber = CurveStateSerialNumber;
+		KeyTimeCache = EditCurve.GetKeyTime(SelectedKeyHandle);
+		KeyValueCache = EditCurve.GetKeyValue(SelectedKeyHandle);
+		KeyInterpModeCache = (int32)EditCurve.GetKeyInterpMode(SelectedKeyHandle);
+
+		float MinTime;
+		float MaxTime;
+		float MinValue;
+		float MaxValue;
+		EditCurve.GetTimeRange(MinTime, MaxTime);
+		EditCurve.GetValueRange(MinValue, MaxValue);
+		CurveTimeMinTextCache = FText::AsNumber(MinTime);
+		CurveTimeMaxTextCache = FText::AsNumber(MaxTime);
+		CurveValueMinTextCache = FText::AsNumber(MinValue);
+		CurveValueMaxTextCache = FText::AsNumber(MaxValue);
+	}
+}
+
+#undef LOCTEXT_NAMESPACE
\ No newline at end of file
diff --git a/Engine/Plugins/FX/Niagara/Source/NiagaraEditor/Private/Widgets/SNiagaraDistributionCompactCurveEditor.h b/Engine/Plugins/FX/Niagara/Source/NiagaraEditor/Private/Widgets/SNiagaraDistributionCompactCurveEditor.h
new file mode 100644
--- /dev/null
+++ b/Engine/Plugins/FX/Niagara/Source/NiagaraEditor/Private/Widgets/SNiagaraDistributionCompactCurveEditor.h
@@ -0,0 +1,90 @@
+// Copyright Epic Games, Inc. All Rights Reserved.
+
+#pragma once
+
+#include "Curves/RichCurve.h"
+#include "Styling/StyleColors.h"
+#include "Widgets/SCompoundWidget.h"
+
+class INiagaraDistributionAdapter;
+class SNiagaraCompactCurveKeySelector;
+
+class SNiagaraDistributionCompactCurveEditor : public SCompoundWidget
+{
+public:
+	SLATE_BEGIN_ARGS(SNiagaraDistributionCompactCurveEditor)
+		: _CurveColor(FSlateColor(EStyleColor::White))
+		{ }
+		SLATE_ARGUMENT(FSlateColor, CurveColor)
+	SLATE_END_ARGS();
+
+	void Construct(const FArguments& InArgs, TSharedRef<INiagaraDistributionAdapter> InDistributionAdapter, int32 InChannelIndex);
+
+private:
+	/** Curve flip direction types **/
+	enum class ECurveFlipDirection : uint8
+	{
+		Horizontal,
+		Vertical
+	};
+
+private:
+	TOptional<int32> GetSelectedKeyIndex() const;
+	TOptional<int32> GetMaxKeyIndex() const;
+	void SelectedKeyIndexChanged(int32 InValue);
+	void SelectedKeyIndexComitted(int32 InValue, ETextCommit::Type CommitInfo);
+	void UpdateSelectedKeyIndex(int32 NewKeyIndex);
+
+	TOptional<float> GetSelectedKeyTime() const;
+	void SelectedKeyTimeChanged(float InValue);
+	void SelectedKeyTimeComitted(float InValue, ETextCommit::Type CommitInfo);
+	void UpdateSelectedKeyTime(float NewKeyTime);
+
+	TOptional<float> GetSelectedKeyValue() const;
+	void SelectedKeyValueChanged(float InValue);
+	void SelectedKeyValueComitted(float InValue, ETextCommit::Type CommitInfo);
+	void UpdateSelectedKeyValue(float NewKeyValue);
+
+	int32 GetSelectedKeyInterpMode() const;
+	void SelectedKeyInterpModeChanged(int32 InNewValue, ESelectInfo::Type InInfo);
+
+	void BeginSliderMovement();
+	void EndSliderMovement(float Value);
+
+	FReply AddKeyButtonClicked();
+	FReply DeleteKeyButtonClicked();
+	FReply FlipCurve(ECurveFlipDirection Direction);
+
+	FKeyHandle GetSelectedKeyHandle() const { return SelectedKeyHandle; }
+	int32 GetCurveStateSerialNumber() const { return CurveStateSerialNumber; }
+
+	void OnKeySelectorKeySelected(FKeyHandle InSelectedKeyHandle);
+	void OnKeySelectorKeyMoved(FKeyHandle InSelectedKeyHandle, float TimeDelta);
+
+	FText GetCurveTimeMinText() const;
+	FText GetCurveTimeMaxText() const;
+	FText GetCurveValueMinText() const;
+	FText GetCurveValueMaxText() const;
+
+	void UpdateCachedValues() const;
+
+private:
+	TSharedPtr<SNiagaraCompactCurveKeySelector> KeySelectorWidget;
+	TSharedPtr<INiagaraDistributionAdapter> DistributionAdapter;
+	int32 ChannelIndex;
+
+	FRichCurve EditCurve;
+	FKeyHandle SelectedKeyHandle;
+
+	int32 SelectedKeyIndex = INDEX_NONE;
+	int32 CurveStateSerialNumber = 0;
+
+	mutable int32 CacheCurveStateSerialNumber = INDEX_NONE;
+	mutable TOptional<float> KeyTimeCache;
+	mutable TOptional<float> KeyValueCache;
+	mutable int32 KeyInterpModeCache;
+	mutable FText CurveTimeMinTextCache;
+	mutable FText CurveTimeMaxTextCache;
+	mutable FText CurveValueMinTextCache;
+	mutable FText CurveValueMaxTextCache;
+};
\ No newline at end of file
diff --git a/Engine/Plugins/FX/Niagara/Source/NiagaraEditor/Private/Widgets/SNiagaraDistributionCurveEditor.cpp b/Engine/Plugins/FX/Niagara/Source/NiagaraEditor/Private/Widgets/SNiagaraDistributionCurveEditor.cpp
--- a/Engine/Plugins/FX/Niagara/Source/NiagaraEditor/Private/Widgets/SNiagaraDistributionCurveEditor.cpp
+++ b/Engine/Plugins/FX/Niagara/Source/NiagaraEditor/Private/Widgets/SNiagaraDistributionCurveEditor.cpp
@@ -2,406 +2,798 @@
 
 #include "Widgets/SNiagaraDistributionCurveEditor.h"
 
-#include "NiagaraEditorStyle.h"
-#include "ScopedTransaction.h"
-#include "SEnumCombo.h"
-#include "Widgets/INiagaraDistributionAdapter.h"
-#include "Widgets/Input/SButton.h"
+#include "INiagaraDistributionAdapter.h"
+#include "NiagaraEditorSettings.h"
+#include "NiagaraEditorUtilities.h"
+
+#include "Engine/Engine.h"
+#include "Curves/CurveFloat.h"
+#include "Curves/RichCurve.h"
+#include "Framework/MultiBox/MultiBoxBuilder.h"
+#include "Misc/AxisDisplayInfo.h"
+#include "Tree/SCurveEditorTree.h"
+#include "Tree/ICurveEditorTreeItem.h"
+#include "Tree/SCurveEditorTreePin.h"
+#include "Widgets/Colors/SColorBlock.h"
+#include "Widgets/Colors/SColorPicker.h"
 #include "Widgets/Input/SNumericEntryBox.h"
 #include "Widgets/Layout/SBox.h"
 #include "Widgets/Layout/SGridPanel.h"
-#include "Widgets/SLeafWidget.h"
-
-#define LOCTEXT_NAMESPACE "NiagaraDistributionCurveEditor"
-
-class SNiagaraCompactCurveView : public SLeafWidget
+#include "Widgets/SNiagaraExpandedToggle.h"
+#include "CurveEditor.h"
+#include "CurveEditorCommands.h"
+#include "ICurveEditorBounds.h"
+#include "SCurveEditorPanel.h"
+#include "SCurveKeyDetailPanel.h"
+#include "EditorFontGlyphs.h"
+#include "PropertyHandle.h"
+#include "RichCurveEditorModel.h"
+#include "SResizeBox.h"
+
+#define LOCTEXT_NAMESPACE "NiagaraCurveEditor"
+
+namespace
 {
-	SLATE_BEGIN_ARGS(SNiagaraCompactCurveView)
-		: _Width(20)
-		, _Height(20)
-		, _CurveColor(FSlateColor(EStyleColor::White))
-		, _CurveStateSerialNumber(0)
-		{ }
-		SLATE_ARGUMENT(float, Width)
-		SLATE_ARGUMENT(float, Height)
-		SLATE_ARGUMENT(FSlateColor, CurveColor)
-		SLATE_ATTRIBUTE(FKeyHandle, SelectedKeyHandle)
-		SLATE_ATTRIBUTE(int32, CurveStateSerialNumber)
-	SLATE_END_ARGS()
-
-	void Construct(const FArguments& InArgs, const FRichCurve* InCurveToDisplay)
+	struct FNiagaraDistributionCurveEditorBounds : public ICurveEditorBounds
 	{
-		Width = InArgs._Width;
-		Height = InArgs._Height;
-		CurveColor = InArgs._CurveColor;
-		SelectedKeyHandle = InArgs._SelectedKeyHandle;
-		CurveStateSerialNumber = InArgs._CurveStateSerialNumber;
-		CurveToDisplay = InCurveToDisplay;
-
-		KeyColor = FSlateColor(EStyleColor::AccentGray);
-		SelectedKeyColor = FSlateColor(EStyleColor::Select);
-	}
+		FNiagaraDistributionCurveEditorBounds(TSharedRef<FNiagaraDistributionCurveEditorOptions> InCurveEditorOptions)
+			: CurveEditorOptions(InCurveEditorOptions)
+		{
+		}
+	
+		virtual void GetInputBounds(double& OutMin, double& OutMax) const
+		{
+			OutMin = CurveEditorOptions->GetViewMinInput();
+			OutMax = CurveEditorOptions->GetViewMaxInput();
+		}
+	
+		virtual void SetInputBounds(double InMin, double InMax)
+		{
+			CurveEditorOptions->SetInputViewRange((float)InMin, (float)InMax);
+		}
+	
+	private:
+		TSharedRef<FNiagaraDistributionCurveEditorOptions> CurveEditorOptions;
+	};
 
-	virtual int32 OnPaint(const FPaintArgs& Args, const FGeometry& AllottedGeometry, const FSlateRect& MyCullingRect, FSlateWindowElementList& OutDrawElements, int32 LayerId, const FWidgetStyle& InWidgetStyle, bool bParentEnabled) const override
+	struct FNiagaraDistributionCurveEditorDetailsTreeItem : public ICurveEditorTreeItem, TSharedFromThis<FNiagaraDistributionCurveEditorDetailsTreeItem>
 	{
-		FKeyHandle NewSelectedKeyHandle = SelectedKeyHandle.Get();
-		int32 NewCurveStateSerialNumber = CurveStateSerialNumber.Get(0);
+		FNiagaraDistributionCurveEditorDetailsTreeItem(TSharedRef<INiagaraDistributionAdapter> InDistributionAdapter, int InChannelIndex)
+			: DistributionAdapter(InDistributionAdapter)
+			, ChannelIndex(InChannelIndex)
+		{
+			Curve			= *DistributionAdapter->GetCurveValue(ChannelIndex);
+			CurveName		= DistributionAdapter->GetChannelDisplayName(ChannelIndex);
+			CurveToolTip	= DistributionAdapter->GetChannelToolTip(ChannelIndex);
+			CurveColor		= DistributionAdapter->GetChannelColor(ChannelIndex).GetSpecifiedColor();
+
+			if ( TSharedPtr<IPropertyHandle> PropertyHandle = DistributionAdapter->GetPropertyHandle() )
+			{
+				TArray<UObject*> OwnerObjects;
+				PropertyHandle->GetOuterObjects(OwnerObjects);
+				WeakOwnerObject = OwnerObjects.Num() > 0 ? OwnerObjects[0] : nullptr;
+			}
+		}
 
-		if (CurveStateSerialNumberCache != NewCurveStateSerialNumber || AllottedGeometryCache != AllottedGeometry)
+		virtual TSharedPtr<SWidget> GenerateCurveEditorTreeWidget(const FName& InColumnName, TWeakPtr<FCurveEditor> InCurveEditor, FCurveEditorTreeItemID InTreeItemID, const TSharedRef<ITableRow>& TableRow) override
 		{
-			SelectedKeyHandleCache = NewSelectedKeyHandle;
-			CurveStateSerialNumberCache = NewCurveStateSerialNumber;
-			AllottedGeometryCache = AllottedGeometry;
-			CacheLines(AllottedGeometry.Size.X, AllottedGeometry.Size.Y, true, true, true);
+			if (InColumnName == ColumnNames.PinHeader)
+			{
+				return SNew(SCurveEditorTreePin, InCurveEditor, InTreeItemID, TableRow);
+			}
+
+			if (InColumnName == ColumnNames.Label)
+			{
+				TSharedRef<SHorizontalBox> LabelBox =
+					SNew(SHorizontalBox)
+					+ SHorizontalBox::Slot()
+					.AutoWidth()
+					.VAlign(VAlign_Center)
+					.Padding(0, 0, 5, 0)
+					[
+						SNew(STextBlock)
+						.Font(FAppStyle::Get().GetFontStyle("FontAwesome.8"))
+						.Text(FEditorFontGlyphs::Circle)
+						.ColorAndOpacity(FSlateColor(CurveColor))
+					]
+					+ SHorizontalBox::Slot()
+					.VAlign(VAlign_Center)
+					[
+						SNew(SBox)
+						.VAlign(VAlign_Center)
+						.MinDesiredHeight(22)
+						[
+							SNew(STextBlock)
+							.Text(CurveName)
+						]
+					];
+
+				if (!CurveToolTip.IsEmpty())
+				{
+					LabelBox->SetToolTipText(CurveToolTip);
+				}
+
+				return LabelBox;
+			}
+
+			return nullptr;
 		}
-		else if (SelectedKeyHandleCache != NewSelectedKeyHandle)
+
+		virtual void CreateCurveModels(TArray<TUniquePtr<FCurveModel>>& OutCurveModels) override
 		{
-			SelectedKeyHandleCache = NewSelectedKeyHandle;
-			CacheLines(AllottedGeometry.Size.X, AllottedGeometry.Size.Y, false, false, true);
+			UObject* OwnerObject = WeakOwnerObject.Get();
+			if (OwnerObject == nullptr)
+			{
+				return;
+			}
+
+			TUniquePtr<FRichCurveEditorModelRaw> NewCurve = MakeUnique<FRichCurveEditorModelRaw>(&Curve, OwnerObject);
+			NewCurve->SetShortDisplayName(CurveName);
+			NewCurve->SetColor(CurveColor);
+			NewCurve->OnCurveModified().AddSP(this, &FNiagaraDistributionCurveEditorDetailsTreeItem::OnCurveChanged);
+			OutCurveModels.Add(MoveTemp(NewCurve));
 		}
 
-		FSlateDrawElement::MakeLines(OutDrawElements, LayerId, AllottedGeometry.ToPaintGeometry(), CurvePointsCache, ESlateDrawEffect::None, CurveColor.GetColor(InWidgetStyle), true, 2.0f);
-		for (const TArray<FVector2D>& KeyLinePoints : KeyLinesCache)
+		virtual bool PassesFilter(const FCurveEditorTreeFilter* InFilter) const override { return true; }
+
+	private:
+		void OnCurveChanged()
 		{
-			FSlateDrawElement::MakeLines(OutDrawElements, LayerId, AllottedGeometry.ToPaintGeometry(), KeyLinePoints, ESlateDrawEffect::None, KeyColor.GetColor(InWidgetStyle), true, 2.0f);
+			DistributionAdapter->SetCurveValue(ChannelIndex, Curve, false);
+
+			if (TSharedPtr<IPropertyHandle> PropertyHandle = DistributionAdapter->GetPropertyHandle())
+			{
+				PropertyHandle->NotifyPostChange(EPropertyChangeType::ValueSet);
+			}
 		}
-		FSlateDrawElement::MakeLines(OutDrawElements, LayerId, AllottedGeometry.ToPaintGeometry(), SelectedKeyLineCache, ESlateDrawEffect::None, SelectedKeyColor.GetColor(InWidgetStyle), true, 3.0f);
-		return LayerId;
-	}
 
-	virtual FVector2D ComputeDesiredSize(float) const override
-	{
-		return FVector2D(Width, Height);
-	}
+	private:
+		TSharedRef<INiagaraDistributionAdapter>	DistributionAdapter;
+		TWeakObjectPtr<UObject>					WeakOwnerObject;
+		int										ChannelIndex = 0;
+		FRichCurve								Curve;
+		FText									CurveName;
+		FText									CurveToolTip;
+		FLinearColor							CurveColor;
+	};
 
-private:
-	void CacheLines(float InWidth, float InHeight, bool bCacheCurve, bool bCacheKeys, bool bCacheSelectedKey) const
+	class SNiagaraDistributionCurveKeySelector : public SCompoundWidget
 	{
-		if (CurveToDisplay == nullptr)
+	public:
+		SLATE_BEGIN_ARGS(SNiagaraDistributionCurveKeySelector) {}
+		SLATE_END_ARGS()
+
+		void Construct(const FArguments& InArgs, TSharedPtr<FCurveEditor> InCurveEditor, const TArray<FCurveEditorTreeItemID>& InCurveTreeItemIds, TSharedPtr<SCurveEditorTree> InCurveEditorTree)
 		{
-			return;
-		}
+			CurveEditor = InCurveEditor;
+			OrderedCurveTreeItemIds = InCurveTreeItemIds;
+			CurveEditorTree = InCurveEditorTree;
+			ChildSlot
+			[
+				SNew(SHorizontalBox)
 
-		float TimeMin;
-		float TimeMax;
-		float ValueMin;
-		float ValueMax;
-		CurveToDisplay->GetTimeRange(TimeMin, TimeMax);
-		CurveToDisplay->GetValueRange(ValueMin, ValueMax);
+				// Zoom to fit button
+				+ SHorizontalBox::Slot()
+				.AutoWidth()
+				.Padding(0, 0, 4, 0)
+				[
+					SNew(SButton)
+					.ButtonStyle(FAppStyle::Get(), "SimpleButton")
+					.OnClicked(this, &SNiagaraDistributionCurveKeySelector::ZoomToFitClicked)
+					.ToolTipText(LOCTEXT("ZoomToFitToolTip", "Zoom to fit all keys"))
+					.Content()
+					[
+						SNew(STextBlock)
+						.Font(FAppStyle::Get().GetFontStyle("FontAwesome.10"))
+						.ColorAndOpacity(FSlateColor::UseForeground())
+						.Text(FEditorFontGlyphs::Expand)
+					]
+				]
+
+				// Previous curve button
+				+ SHorizontalBox::Slot()
+				.AutoWidth()
+				.Padding(0, 0, 1, 0)
+				[
+					SNew(SButton)
+					.ButtonStyle(FAppStyle::Get(), "SimpleButton")
+					.OnClicked(this, &SNiagaraDistributionCurveKeySelector::PreviousCurveClicked)
+					.ToolTipText(LOCTEXT("PreviousCurveToolTip", "Select the previous curve"))
+					.Visibility(this, &SNiagaraDistributionCurveKeySelector::GetNextPreviousCurveButtonVisibility)
+					.Content()
+					[
+						SNew(SImage)
+						.ColorAndOpacity(FSlateColor::UseForeground())
+						.Image(FAppStyle::Get().GetBrush("Icons.ArrowUp"))
+					]
+				]
+
+				// Next curve button
+				+ SHorizontalBox::Slot()
+				.AutoWidth()
+				.Padding(0, 0, 4, 0)
+				[
+					SNew(SButton)
+					.ButtonStyle(FAppStyle::Get(), "SimpleButton")
+					.OnClicked(this, &SNiagaraDistributionCurveKeySelector::NextCurveClicked)
+					.ToolTipText(LOCTEXT("NextCurveToolTip", "Select the next curve"))
+					.Visibility(this, &SNiagaraDistributionCurveKeySelector::GetNextPreviousCurveButtonVisibility)
+					.Content()
+					[
+						SNew(SImage)
+						.Image(FAppStyle::Get().GetBrush("Icons.ArrowDown"))
+						.ColorAndOpacity(FSlateColor::UseForeground())
+					]
+				]
+
+				// Previous key button
+				+ SHorizontalBox::Slot()
+				.AutoWidth()
+				.Padding(0, 0, 1, 0)
+				[
+					SNew(SButton)
+					.ButtonStyle(FAppStyle::Get(), "SimpleButton")
+					.OnClicked(this, &SNiagaraDistributionCurveKeySelector::PreviousKeyClicked)
+					.ToolTipText(LOCTEXT("PreviousKeyToolTip", "Select the previous key for the selected curve."))
+					.Content()
+					[
+						SNew(SImage)
+						.Image(FAppStyle::Get().GetBrush("Icons.ArrowLeft"))
+						.ColorAndOpacity(FSlateColor::UseForeground())
+					]
+				]
 
-		float TimeRange = TimeMax != TimeMin ? TimeMax - TimeMin : 1;
-		float ValueRange = ValueMax != ValueMin ? ValueMax - ValueMin : 1;
+				// Next key button
+				+ SHorizontalBox::Slot()
+				.AutoWidth()
+				.Padding(0, 0, 4, 0)
+				[
+					SNew(SButton)
+					.ButtonStyle(FAppStyle::Get(), "SimpleButton")
+					.OnClicked(this, &SNiagaraDistributionCurveKeySelector::NextKeyClicked)
+					.ToolTipText(LOCTEXT("NextKeyToolTip", "Select the next key for the selected curve."))
+					.Content()
+					[
+						SNew(SImage)
+						.Image(FAppStyle::Get().GetBrush("Icons.ArrowRight"))
+						.ColorAndOpacity(FSlateColor::UseForeground())
+					]
+				]
 
-		if (bCacheCurve)
+				// Add key button
+				+ SHorizontalBox::Slot()
+				.AutoWidth()
+				.Padding(0, 0, 1, 0)
+				[
+					SNew(SButton)
+					.ButtonStyle(FAppStyle::Get(), "SimpleButton")
+					.OnClicked(this, &SNiagaraDistributionCurveKeySelector::AddKeyClicked)
+					.ToolTipText(LOCTEXT("AddKeyToolTip", "Add a key to the selected curve."))
+					.Content()
+					[
+						SNew(SImage)
+						.Image(FAppStyle::Get().GetBrush("Icons.PlusCircle"))
+						.ColorAndOpacity(FSlateColor::UseForeground())
+					]
+				]
+
+				// Delete key button
+				+ SHorizontalBox::Slot()
+				.AutoWidth()
+				[
+					SNew(SButton)
+					.ButtonStyle(FAppStyle::Get(), "SimpleButton")
+					.OnClicked(this, &SNiagaraDistributionCurveKeySelector::DeleteKeyClicked)
+					.ToolTipText(LOCTEXT("DeleteKeyToolTip", "Delete the currently selected keys."))
+					.Content()
+					[
+						SNew(SImage)
+						.Image(FAppStyle::Get().GetBrush("Icons.Delete"))
+						.ColorAndOpacity(FSlateColor::UseForeground())
+					]
+				]
+			];
+		}
+
+	private:
+		void GetActiveCurveModelAndSelectedKeys(TOptional<FCurveModelID>& OutActiveCurveModelId, TArray<FKeyHandle>& OutSelectedKeyHandles)
 		{
-			CurvePointsCache.Empty();
-			int32 Points = CurveToDisplay->GetNumKeys() * 20;
-			float TimeIncrement = TimeRange / (Points - 1);
-			for (int32 i = 0; i < Points; i++)
+			const TMap<FCurveEditorTreeItemID, ECurveEditorTreeSelectionState>& CurveTreeSelection = CurveEditor->GetTree()->GetSelection();
+			if (CurveTreeSelection.Num() != 0)
 			{
-				float Time = TimeMin + i * TimeIncrement;
-				float Value = CurveToDisplay->Eval(Time);
+				// If there are curves selected in the tree use those first.
+				FCurveEditorTreeItemID FirstSelectedCurveTreeItemId;
+				for (int32 i = OrderedCurveTreeItemIds.Num() - 1; i >= 0; i--)
+				{
+					const FCurveEditorTreeItemID& CurveTreeItemId = OrderedCurveTreeItemIds[i];
+					if (CurveTreeSelection.Contains(CurveTreeItemId))
+					{
+						FirstSelectedCurveTreeItemId = CurveTreeItemId;
+						break;
+					}
+				}
+
+				TArrayView<const FCurveModelID> CurveModelIds = CurveEditor->GetTreeItem(FirstSelectedCurveTreeItemId).GetOrCreateCurves(CurveEditor.Get());
+				if (CurveModelIds.Num() == 1)
+				{
+					OutActiveCurveModelId = CurveModelIds[0];
 
-				float NormalizedX = (Time - TimeMin) / TimeRange;
-				float NormalizedY = (Value - ValueMin) / ValueRange;
-				CurvePointsCache.Add(FVector2D(NormalizedX * InWidth, (1 - NormalizedY) * InHeight));
+					const FKeyHandleSet* SelectedKeyHandleSet = CurveEditor->GetSelection().GetAll().Find(OutActiveCurveModelId.GetValue());
+					if (SelectedKeyHandleSet != nullptr)
+					{
+						OutSelectedKeyHandles = SelectedKeyHandleSet->AsArray();
+					}
+				}
+			}
+			else
+			{
+				// Otherwise check if there are keys selected and if so use first curve with selected keys.
+				FCurveEditorSelection& CurveEditorSelection = CurveEditor->GetSelection();
+				if (CurveEditorSelection.IsEmpty() == false)
+				{
+					for (int32 i = OrderedCurveTreeItemIds.Num() - 1; i >= 0; i--)
+					{
+						const FCurveEditorTreeItemID& CurveTreeItemId = OrderedCurveTreeItemIds[i];
+						TArrayView<const FCurveModelID> CurveModelIds = CurveEditor->GetTreeItem(CurveTreeItemId).GetCurves();
+						if (CurveModelIds.Num() == 1)
+						{
+							const FKeyHandleSet* SelectedKeyHandleSet = CurveEditorSelection.GetAll().Find(CurveModelIds[0]);
+							if (SelectedKeyHandleSet != nullptr)
+							{
+								OutActiveCurveModelId = CurveModelIds[0];
+								OutSelectedKeyHandles = SelectedKeyHandleSet->AsArray();
+								break;
+							}
+						}
+					}
+				}
+				else
+				{
+					// Otherwise just use the first pinned curve.
+					const TSet<FCurveModelID>& PinnedCurveIds = CurveEditor->GetPinnedCurves();
+					if (PinnedCurveIds.Num() > 0)
+					{
+						for (FCurveEditorTreeItemID& CurveTreeItemId : OrderedCurveTreeItemIds)
+						{
+							TArrayView<const FCurveModelID> CurveModelIds = CurveEditor->GetTreeItem(CurveTreeItemId).GetOrCreateCurves(CurveEditor.Get());
+							if (CurveModelIds.Num() == 1 && PinnedCurveIds.Contains(CurveModelIds[0]))
+							{
+								OutActiveCurveModelId = CurveModelIds[0];
+								break;
+							}
+						}
+					}
+				}
 			}
 		}
 
-		if (bCacheKeys)
+		struct FKeyHandlePositionPair
 		{
-			KeyLinesCache.Empty();
-			for (auto KeyIterator(CurveToDisplay->GetKeyIterator()); KeyIterator; ++KeyIterator)
-			{
-				const FRichCurveKey& Key = *KeyIterator;
-				TArray<FVector2D>& KeyLinePoints = KeyLinesCache.AddDefaulted_GetRef();
+			FKeyHandle Handle;
+			FKeyPosition Position;
+		};
+
+		void GetSortedKeyHandlessAndPositionsForModel(FCurveModel& InCurveModel, TArray<FKeyHandlePositionPair>& OutSortedKeyHandlesAndPositions)
+		{
+			const TArray<FKeyHandle> KeyHandles = InCurveModel.GetAllKeys();
 
-				float NormalizedX = (Key.Time - TimeMin) / TimeRange;
-				KeyLinePoints.Add(FVector2D(NormalizedX * InWidth, 0));
-				KeyLinePoints.Add(FVector2D(NormalizedX * InWidth, InHeight));
+			TArray<FKeyPosition> KeyPositions;
+			KeyPositions.AddDefaulted(KeyHandles.Num());
+			InCurveModel.GetKeyPositions(KeyHandles, KeyPositions);
+
+			for (int32 i = 0; i < KeyHandles.Num(); i++)
+			{
+				FKeyHandlePositionPair& KeyHandlePositionPair = OutSortedKeyHandlesAndPositions.AddDefaulted_GetRef();
+				KeyHandlePositionPair.Handle = KeyHandles[i];
+				KeyHandlePositionPair.Position = KeyPositions[i];
 			}
+
+			OutSortedKeyHandlesAndPositions.Sort([](const FKeyHandlePositionPair& A, const FKeyHandlePositionPair& B) { return A.Position.InputValue < B.Position.InputValue; });
 		}
 
-		if (bCacheSelectedKey)
+		void GetOrderedActiveCurveModelIds(TArray<FCurveModelID>& OutOrderedActiveCurveModelIds)
 		{
-			SelectedKeyLineCache.Empty();
-			for (auto KeyHandleIterator(CurveToDisplay->GetKeyHandleIterator()); KeyHandleIterator; ++KeyHandleIterator)
+			if (CurveEditor->GetTreeSelection().Num() > 0)
+			{
+				// If there are curves selected in the tree then those are the only active ones.
+				for (const FCurveEditorTreeItemID& OrderedCurveTreeItemId : OrderedCurveTreeItemIds)
+				{
+					if(CurveEditor->GetTreeSelectionState(OrderedCurveTreeItemId) != ECurveEditorTreeSelectionState::None)
+					{
+						TArrayView<const FCurveModelID> CurveModelIds = CurveEditor->GetTreeItem(OrderedCurveTreeItemId).GetOrCreateCurves(CurveEditor.Get());
+						if (CurveModelIds.Num() == 1)
+						{
+							OutOrderedActiveCurveModelIds.Add(CurveModelIds[0]);
+						}
+					}
+				}
+			}
+			else
 			{
-				if (*KeyHandleIterator == SelectedKeyHandleCache)
+				// Otherwise the active curves are the pinned curves.
+				for (const FCurveEditorTreeItemID& OrderedCurveTreeItemId : OrderedCurveTreeItemIds)
 				{
-					const FRichCurveKey& SelectedKey = CurveToDisplay->GetKey(SelectedKeyHandleCache);
-					float NormalizedX = (SelectedKey.Time - TimeMin) / TimeRange;
-					SelectedKeyLineCache.Add(FVector2D(NormalizedX * InWidth, 0));
-					SelectedKeyLineCache.Add(FVector2D(NormalizedX * InWidth, InHeight));
-					break;
+					TArrayView<const FCurveModelID> CurveModelIds = CurveEditor->GetTreeItem(OrderedCurveTreeItemId).GetCurves();
+					if (CurveModelIds.Num() == 1 && CurveEditor->IsCurvePinned(CurveModelIds[0]))
+					{
+						OutOrderedActiveCurveModelIds.Add(CurveModelIds[0]);
+					}
 				}
 			}
 		}
-	}
 
-	const FRichCurve* CurveToDisplay = nullptr;
-	mutable TArray<FVector2D> CurvePointsCache;
-	mutable TArray<TArray<FVector2D>> KeyLinesCache; 
-	mutable TArray<FVector2D> SelectedKeyLineCache;
-	mutable FGeometry AllottedGeometryCache;
-	mutable FKeyHandle SelectedKeyHandleCache;
-	mutable int32 CurveStateSerialNumberCache = 0;
-	float Width;
-	float Height;
-	FSlateColor CurveColor;
-	FSlateColor KeyColor;
-	FSlateColor SelectedKeyColor;
-	TAttribute<FKeyHandle> SelectedKeyHandle;
-	TAttribute<int32> CurveStateSerialNumber;
-};
-
-class SNiagaraCompactCurveKeyHandle : public SBorder
-{
-	DECLARE_DELEGATE_OneParam(FOnMoved, float /* LocalDelta */);
-
-	SLATE_BEGIN_ARGS(SNiagaraCompactCurveKeyHandle)
-		: _Size(10, 10)
-	{ }
-		SLATE_ARGUMENT(FVector2f, Size)
-		SLATE_ATTRIBUTE(bool, IsSelected)
-		SLATE_EVENT(FSimpleDelegate, OnSelected)
-		SLATE_EVENT(FOnMoved, OnMoved)
-	SLATE_END_ARGS()
-
-	void Construct(const FArguments& InArgs)
-	{
-		Size = InArgs._Size;
-		IsSelected = InArgs._IsSelected;
-		OnSelectedDelegate = InArgs._OnSelected;
-		OnMovedDelegate = InArgs._OnMoved;
-		SBorder::Construct(SBorder::FArguments()
-			.BorderImage(FNiagaraEditorStyle::Get().GetBrush("NiagaraEditor.DistributionEditor.CurveKeyHandle"))
-			.BorderBackgroundColor(this, &SNiagaraCompactCurveKeyHandle::GetHandleColor)
-			.Padding(0)
-			[
-				SNew(SBox)
-				.WidthOverride(Size.X)
-				.HeightOverride(Size.Y)
-			]);
-	}
+		enum class ENavigateDirection
+		{
+			Previous,
+			Next
+		};
 
-	virtual FReply OnMouseButtonDown(const FGeometry& MyGeometry, const FPointerEvent& MouseEvent) override
-	{
-		if (MouseEvent.GetEffectingButton() == EKeys::LeftMouseButton)
+		void NavigateToAdjacentCurve(ENavigateDirection Direction)
 		{
-			FVector2D MouseLocation = MouseEvent.GetScreenSpacePosition();
-			LastMouseLocation = MouseLocation.X;
-			return FReply::Handled().CaptureMouse(SharedThis(this));
+			if (CurveEditorTree.IsValid())
+			{
+				if (CurveEditor->GetTreeSelection().Num() == 0)
+				{
+					CurveEditorTree->SetItemSelection(OrderedCurveTreeItemIds[0], true);
+				}
+				else
+				{
+					int32 TargetSelectedTreeItemIndex = INDEX_NONE;
+					int32 StartIndex = Direction == ENavigateDirection::Previous ? 0 : OrderedCurveTreeItemIds.Num() - 1;
+					int32 EndIndex = Direction == ENavigateDirection::Previous ? OrderedCurveTreeItemIds.Num() : -1;
+					int32 IndexOffset = Direction == ENavigateDirection::Previous ? 1 : -1;
+					for (int32 i = StartIndex; i != EndIndex; i += IndexOffset)
+					{
+						if (CurveEditor->GetTreeSelectionState(OrderedCurveTreeItemIds[i]) != ECurveEditorTreeSelectionState::None)
+						{
+							TargetSelectedTreeItemIndex = i;
+							break;
+						}
+					}
+
+					if (TargetSelectedTreeItemIndex != INDEX_NONE)
+					{
+						int32 IndexToSelect = Direction == ENavigateDirection::Previous 
+							? TargetSelectedTreeItemIndex - 1
+							: TargetSelectedTreeItemIndex + 1;
+
+						if (IndexToSelect > OrderedCurveTreeItemIds.Num() - 1)
+						{
+							IndexToSelect = 0;
+						}
+						else if(IndexToSelect < 0)
+						{
+							IndexToSelect = OrderedCurveTreeItemIds.Num() - 1;
+						}
+
+						CurveEditorTree->ClearSelection();
+						CurveEditorTree->SetItemSelection(OrderedCurveTreeItemIds[IndexToSelect], true);
+					}
+				}
+			}
 		}
-		return FReply::Unhandled();
-	}
 
-	virtual FReply OnMouseButtonUp(const FGeometry& MyGeometry, const FPointerEvent& MouseEvent) override
-	{
-		if (this->HasMouseCapture())
+		void NavigateToAdjacentKey(ENavigateDirection Direction)
 		{
-			if (IsSelected.Get() == false)
+			TOptional<FCurveModelID> ActiveCurveModelId;
+			TArray<FKeyHandle> SelectedKeyHandles;
+			GetActiveCurveModelAndSelectedKeys(ActiveCurveModelId, SelectedKeyHandles);
+
+			TOptional<FCurveModelID> CurveModelIdToSelect;
+			TOptional<FKeyHandle> KeyHandleToSelect;
+			if (ActiveCurveModelId.IsSet())
 			{
-				OnSelectedDelegate.ExecuteIfBound();
+				FCurveModel* ActiveCurveModel = CurveEditor->GetCurves()[ActiveCurveModelId.GetValue()].Get();
+				TArray<FKeyHandlePositionPair> ActiveSortedKeyHandlePositionPairs;
+				GetSortedKeyHandlessAndPositionsForModel(*ActiveCurveModel, ActiveSortedKeyHandlePositionPairs);
+
+				if (SelectedKeyHandles.Num() != 0)
+				{
+					// If there's currently a selected key on the active curve then we want to use that as the target for navigating.
+					int32 TargetSelectedKeyIndex = INDEX_NONE;
+					int32 StartIndex = Direction == ENavigateDirection::Previous ? 0 : ActiveSortedKeyHandlePositionPairs.Num() - 1;
+					int32 EndIndex = Direction == ENavigateDirection::Previous ? ActiveSortedKeyHandlePositionPairs.Num() : -1;
+					int32 IndexOffset = Direction == ENavigateDirection::Previous ? 1 : -1;
+					for (int32 i = StartIndex; i != EndIndex; i += IndexOffset)
+					{
+						if (SelectedKeyHandles.Contains(ActiveSortedKeyHandlePositionPairs[i].Handle))
+						{
+							TargetSelectedKeyIndex = i;
+							break;
+						}
+					}
+
+					if (TargetSelectedKeyIndex != INDEX_NONE)
+					{
+						if(Direction == ENavigateDirection::Previous && TargetSelectedKeyIndex > 0)
+						{
+							// If we're navigating previous and we're not at the first key we can just select the previous key on this curve.
+							CurveModelIdToSelect = ActiveCurveModelId;
+							KeyHandleToSelect = ActiveSortedKeyHandlePositionPairs[TargetSelectedKeyIndex - 1].Handle;
+						}
+						else if (Direction == ENavigateDirection::Next && TargetSelectedKeyIndex < ActiveSortedKeyHandlePositionPairs.Num() - 1)
+						{
+							// If we're navigating next and we're not at the last key we can just select the next key on this curve.
+							CurveModelIdToSelect = ActiveCurveModelId;
+							KeyHandleToSelect = ActiveSortedKeyHandlePositionPairs[TargetSelectedKeyIndex + 1].Handle;
+						}
+						else
+						{
+							// Otherwise we're going to need to navigate to another curve since we're at the start or end of the current curve.
+							TArray<FCurveModelID> OrderedActiveCurveModelIds;
+							GetOrderedActiveCurveModelIds(OrderedActiveCurveModelIds);
+
+							// Find the adjacent curve with keys so we can select a key there.
+							FCurveModel* CurveModelToSelect = nullptr;
+							int32 CurrentIndex = OrderedActiveCurveModelIds.IndexOfByKey(ActiveCurveModelId);
+							int32 CurrentIndexOffset = Direction == ENavigateDirection::Previous ? -1 : 1;
+							while (CurveModelIdToSelect.IsSet() == false)
+							{
+								CurrentIndex += CurrentIndexOffset;
+								if (CurrentIndex < 0)
+								{
+									CurrentIndex = OrderedActiveCurveModelIds.Num() - 1;
+								}
+								else if (CurrentIndex > OrderedActiveCurveModelIds.Num() - 1)
+								{
+									CurrentIndex = 0;
+								}
+
+								FCurveModel* CurrentCurveModel = CurveEditor->GetCurves()[OrderedActiveCurveModelIds[CurrentIndex]].Get();
+								if (CurrentCurveModel->GetNumKeys() > 0)
+								{
+									CurveModelIdToSelect = OrderedActiveCurveModelIds[CurrentIndex];
+									CurveModelToSelect = CurrentCurveModel;
+								}
+							}
+
+							if (CurveModelIdToSelect == ActiveCurveModelId)
+							{
+								// There were no other active curves with keys to select from so just wrap around the currently active curve.
+								if (Direction == ENavigateDirection::Previous)
+								{
+									KeyHandleToSelect = ActiveSortedKeyHandlePositionPairs.Last().Handle;
+								}
+								else
+								{
+									KeyHandleToSelect = ActiveSortedKeyHandlePositionPairs[0].Handle;
+								}
+							}
+							else if(CurveModelToSelect != nullptr)
+							{
+								// We're selecting a key on a different curve so we need to sort the positions and select the first or last based
+								// on the navigation direction.
+								TArray<FKeyHandlePositionPair> SortedKeyHandlePositionPairs;
+								GetSortedKeyHandlessAndPositionsForModel(*CurveModelToSelect, SortedKeyHandlePositionPairs);
+								if (Direction == ENavigateDirection::Previous)
+								{
+									KeyHandleToSelect = SortedKeyHandlePositionPairs.Last().Handle;
+								}
+								else
+								{
+									KeyHandleToSelect = SortedKeyHandlePositionPairs[0].Handle;
+								}
+							}
+						}
+					}
+				}
+				else if(ActiveSortedKeyHandlePositionPairs.Num() > 0)
+				{
+					// There weren't any keys already selected so just select the first key on the active curve.
+					CurveModelIdToSelect = ActiveCurveModelId;
+					KeyHandleToSelect = ActiveSortedKeyHandlePositionPairs[0].Handle;
+				}
+			}
+
+			if (CurveModelIdToSelect.IsSet() && KeyHandleToSelect.IsSet())
+			{
+				FCurveEditorSelection& CurveEditorSelection = CurveEditor->GetSelection();
+				CurveEditorSelection.Clear();
+				CurveEditorSelection.Add(CurveModelIdToSelect.GetValue(), ECurvePointType::Key, KeyHandleToSelect.GetValue());
+
+				TArray<FCurveModelID> CurvesToFit;
+				CurvesToFit.Add(CurveModelIdToSelect.GetValue());
+				CurveEditor->ZoomToFitCurves(CurvesToFit);
 			}
-			return FReply::Handled().ReleaseMouseCapture();
 		}
-		return FReply::Unhandled();
-	}
 
-	virtual FReply OnMouseMove(const FGeometry& MyGeometry, const FPointerEvent& MouseEvent) override
-	{
-		if (this->HasMouseCapture())
+		FReply ZoomToFitClicked()
 		{
-			FVector2D MouseLocation = MouseEvent.GetScreenSpacePosition();
-			OnMovedDelegate.ExecuteIfBound((MouseLocation.X - LastMouseLocation) / MyGeometry.Scale);
-			LastMouseLocation = MouseLocation.X;
+			TArray<FCurveModelID> ActiveCurveModelIds;
+			GetOrderedActiveCurveModelIds(ActiveCurveModelIds);
+			CurveEditor->ZoomToFitCurves(ActiveCurveModelIds);
 			return FReply::Handled();
 		}
-		return FReply::Unhandled();
-	}
-
-private:
-	FSlateColor GetHandleColor() const
-	{
-		return IsSelected.Get()
-			? FSlateColor(EStyleColor::Select)
-			: FSlateColor(EStyleColor::White);
-	}
 
-private:
-	FVector2f Size;
-	TAttribute<bool> IsSelected;
-	FSimpleDelegate OnSelectedDelegate;
-	FOnMoved OnMovedDelegate;
-	float LastMouseLocation;
-};
-
-class SNiagaraCompactCurveKeySelector : public SCompoundWidget
-{
-	DECLARE_DELEGATE_OneParam(FOnKeySelected, FKeyHandle /* SelectedKeyHandle */);
-	DECLARE_DELEGATE_TwoParams(FOnKeyMoved, FKeyHandle /* SelectedKeyHandle */, float /* TimeDelta */)
-
-	SLATE_BEGIN_ARGS(SNiagaraCompactCurveKeySelector)
-		: _HandleSize(10, 10)
-		, _CurveStateSerialNumber(0)
-	{ }
-		SLATE_ARGUMENT(FVector2f, HandleSize)
-		SLATE_ATTRIBUTE(FKeyHandle, SelectedKeyHandle)
-		SLATE_ATTRIBUTE(int32, CurveStateSerialNumber)
-		SLATE_EVENT(FOnKeySelected, OnKeySelected)
-		SLATE_EVENT(FOnKeyMoved, OnKeyMoved)
-	SLATE_END_ARGS()
-
-private:
-	struct FKeyWidgetData
-	{
-		FKeyHandle Handle;
-		TSharedPtr<SNiagaraCompactCurveKeyHandle> HandleWidget;
-		FMargin BorderPadding;
-	};
+		FReply PreviousCurveClicked()
+		{
+			NavigateToAdjacentCurve(ENavigateDirection::Previous);
+			return FReply::Handled();
+		}
 
-public:
-	void Construct(const FArguments& InArgs, FRichCurve* InEditCurve)
-	{
-		EditCurve = InEditCurve;
-		HandleSize = InArgs._HandleSize;
-		SelectedKeyHandle = InArgs._SelectedKeyHandle;
-		CurveStateSerialNumber = InArgs._CurveStateSerialNumber;
-		OnKeySelectedDelegate = InArgs._OnKeySelected;
-		OnKeyMovedDelegate = InArgs._OnKeyMoved;
-		LocalWidthCache = 100;
-
-		ChildSlot
-		[
-			SAssignNew(BorderOverlay, SOverlay)
-		];
+		FReply NextCurveClicked()
+		{
+			NavigateToAdjacentCurve(ENavigateDirection::Next);
+			return FReply::Handled();
+		}
 
-		UpdateKeyWidgets();
-	}
+		FReply PreviousKeyClicked()
+		{
+			NavigateToAdjacentKey(ENavigateDirection::Previous);
+			return FReply::Handled();
+		}
 
-	virtual void Tick(const FGeometry& AllottedGeometry, const double InCurrentTime, const float InDeltaTime) override
-	{
-		SCompoundWidget::Tick(AllottedGeometry, InCurrentTime, InDeltaTime);
-		int32 CurrentCurveStateSerialNumber = CurveStateSerialNumber.Get();
-		LocalWidthCache = AllottedGeometry.GetLocalSize().X - HandleSize.X;
-		if (CachedCurveStateSerialNumber != CurrentCurveStateSerialNumber || AllottedGeometry != AllottedGeometryCache)
+		FReply NextKeyClicked()
 		{
-			CachedCurveStateSerialNumber = CurrentCurveStateSerialNumber;
-			UpdateKeyWidgets();
+			NavigateToAdjacentKey(ENavigateDirection::Next);
+			return FReply::Handled();
 		}
-		AllottedGeometryCache = AllottedGeometry;
-	}
 
-	void RequestWidgetUpdate() { bNeedsWidgetUpdate = true; }
+		EVisibility GetNextPreviousCurveButtonVisibility() const
+		{
+			return OrderedCurveTreeItemIds.Num() == 1 ? EVisibility::Collapsed : EVisibility::Visible;
+		}
 
-private:
-	void UpdateKeyWidgets()
-	{
-		if (KeyWidgetData.Num() != EditCurve->GetNumKeys() || bNeedsWidgetUpdate)
+		FReply AddKeyClicked()
 		{
-			BorderOverlay->ClearChildren();
-			KeyWidgetData.Empty();
-			for (auto KeyHandleIterator = EditCurve->GetKeyHandleIterator(); KeyHandleIterator; KeyHandleIterator++)
+			TOptional<FCurveModelID> CurveModelIdForAdd;
+			TArray<FKeyHandle> SelectedKeyHandles;
+			GetActiveCurveModelAndSelectedKeys(CurveModelIdForAdd, SelectedKeyHandles);
+
+			if(CurveModelIdForAdd.IsSet())
 			{
-				int32 DataIndex = KeyWidgetData.Num();
-				FKeyWidgetData& KeyWidgetDataItem = KeyWidgetData.AddDefaulted_GetRef();
-				KeyWidgetDataItem.Handle = *KeyHandleIterator;
-				TAttribute<FMargin> BorderPadding = TAttribute<FMargin>::CreateSP(this, &SNiagaraCompactCurveKeySelector::GetKeyBorderPadding, DataIndex);
-				BorderOverlay->AddSlot()
-				.Padding(BorderPadding)
-				.HAlign(HAlign_Left)
-				[
-					SAssignNew(KeyWidgetDataItem.HandleWidget, SNiagaraCompactCurveKeyHandle)
-					.Size(HandleSize)
-					.IsSelected(this, &SNiagaraCompactCurveKeySelector::GetKeyIsSelected, KeyWidgetDataItem.Handle)
-					.OnSelected(this, &SNiagaraCompactCurveKeySelector::KeyHandleSelected, KeyWidgetDataItem.Handle)
-					.OnMoved(this, &SNiagaraCompactCurveKeySelector::KeyHandleMoved, KeyWidgetDataItem.Handle)
-				];
+				FCurveModel* CurveModelForAdd = CurveEditor->GetCurves()[CurveModelIdForAdd.GetValue()].Get();
+
+				FKeyPosition NewKeyPosition;
+				FKeyAttributes NewKeyAttributes = CurveEditor->GetDefaultKeyAttributes().Get();
+				NewKeyAttributes.SetInterpMode(RCIM_Cubic);
+				NewKeyAttributes.SetTangentMode(RCTM_Auto);
+				if (CurveModelForAdd->GetNumKeys() == 0)
+				{
+					// If there are no keys, add one at 0, 0.
+					NewKeyPosition.InputValue = 0.0f;
+					NewKeyPosition.OutputValue = 0.0f;
+				}
+				else if (CurveModelForAdd->GetNumKeys() == 1)
+				{
+					// If there's a single key, add the new key at the same value, but time + 1.
+					const TArray<FKeyHandle> KeyHandles = CurveModelForAdd->GetAllKeys();
+				
+					TArray<FKeyPosition> KeyPositions;
+					KeyPositions.AddDefaulted();
+					CurveModelForAdd->GetKeyPositions(KeyHandles, KeyPositions);
+						
+					NewKeyPosition.InputValue = KeyPositions[0].InputValue + 1;
+					NewKeyPosition.OutputValue = KeyPositions[0].OutputValue;
+				}
+				else
+				{
+					TArray<FKeyHandlePositionPair> SortedKeyHandlePositionPairs;
+					GetSortedKeyHandlessAndPositionsForModel(*CurveModelForAdd, SortedKeyHandlePositionPairs);
+
+					int32 IndexToAddAfter = INDEX_NONE;
+					if (SelectedKeyHandles.Num() > 0)
+					{
+						for (int32 i = SortedKeyHandlePositionPairs.Num() - 1; i >= 0; i--)
+						{
+							if (SelectedKeyHandles.Contains(SortedKeyHandlePositionPairs[i].Handle))
+							{
+								IndexToAddAfter = i;
+								break;
+							}
+						}
+					}
+
+					if(IndexToAddAfter == INDEX_NONE)
+					{
+						IndexToAddAfter = SortedKeyHandlePositionPairs.Num() - 1;
+					}
+
+					if (IndexToAddAfter == SortedKeyHandlePositionPairs.Num() - 1)
+					{
+						const FKeyPosition& TargetKeyPosition = SortedKeyHandlePositionPairs[IndexToAddAfter].Position;
+						const FKeyPosition& PreviousKeyPosition = SortedKeyHandlePositionPairs[IndexToAddAfter - 1].Position;
+						NewKeyPosition.InputValue = TargetKeyPosition.InputValue + (TargetKeyPosition.InputValue - PreviousKeyPosition.InputValue);
+						NewKeyPosition.OutputValue = TargetKeyPosition.OutputValue + (TargetKeyPosition.OutputValue - PreviousKeyPosition.OutputValue);
+					}
+					else
+					{
+						const FKeyPosition& TargetKeyPosition = SortedKeyHandlePositionPairs[IndexToAddAfter].Position;
+						const FKeyPosition& NextKeyPosition = SortedKeyHandlePositionPairs[IndexToAddAfter + 1].Position;
+						NewKeyPosition.InputValue = TargetKeyPosition.InputValue + ((NextKeyPosition.InputValue - TargetKeyPosition.InputValue) / 2);
+						NewKeyPosition.OutputValue = TargetKeyPosition.OutputValue + ((NextKeyPosition.OutputValue - TargetKeyPosition.OutputValue) / 2);
+					}
+				}
+
+				FScopedTransaction Transaction(LOCTEXT("AddKey", "Add Key"));
+				TOptional<FKeyHandle> NewKeyHandle = CurveModelForAdd->AddKey(NewKeyPosition, NewKeyAttributes);
+				if(NewKeyHandle.IsSet())
+				{
+					FCurveEditorSelection& CurveEditorSelection = CurveEditor->GetSelection();
+					CurveEditorSelection.Clear();
+					CurveEditorSelection.Add(CurveModelIdForAdd.GetValue(), ECurvePointType::Key, NewKeyHandle.GetValue());
+
+					TArray<FCurveModelID> CurvesToFit;
+					CurvesToFit.Add(CurveModelIdForAdd.GetValue());
+					CurveEditor->ZoomToFitCurves(CurvesToFit);
+				}
 			}
-			bNeedsWidgetUpdate = false;
+			return FReply::Handled();
 		}
 
-		float TimeMin;
-		float TimeMax;
-		float ValueMin;
-		float ValueMax;
-		EditCurve->GetTimeRange(TimeMin, TimeMax);
-		EditCurve->GetValueRange(ValueMin, ValueMax);
-		float TimeRange = TimeMax != TimeMin ? TimeMax - TimeMin : 1;
-		float ValueRange = ValueMax != ValueMin ? ValueMax - ValueMin : 1;
-		for (FKeyWidgetData& KeyWidgetDataItem : KeyWidgetData)
-		{
-			const FRichCurveKey& Key = EditCurve->GetKey(KeyWidgetDataItem.Handle);
-			float NormalizedX = (Key.Time - TimeMin) / TimeRange;
-			KeyWidgetDataItem.BorderPadding = FMargin(LocalWidthCache * NormalizedX, 0, 0, 0);
+		FReply DeleteKeyClicked()
+		{
+			CurveEditor->DeleteSelection();
+			return FReply::Handled();
 		}
-	}
 
-	FMargin GetKeyBorderPadding(int32 KeyWidgetDataIndex) const
-	{
-		return KeyWidgetDataIndex < KeyWidgetData.Num() ? KeyWidgetData[KeyWidgetDataIndex].BorderPadding : FMargin(0);
-	}
+	private:
+		TSharedPtr<FCurveEditor> CurveEditor;
+		TArray<FCurveEditorTreeItemID> OrderedCurveTreeItemIds;
+		TSharedPtr<SCurveEditorTree> CurveEditorTree;
+	};
+}
 
-	bool GetKeyIsSelected(FKeyHandle KeyHandle) const
-	{
-		return KeyHandle == SelectedKeyHandle.Get();
-	}
+void SNiagaraDistributionCurveEditor::Construct(const FArguments& InArgs, TSharedRef<INiagaraDistributionAdapter> InDistributionAdapter)
+{
+	DistributionAdapter = InDistributionAdapter;
+	CurveEditorOptions	= InArgs._CurveEditorOptions;
 
-	void KeyHandleSelected(FKeyHandle KeyHandle)
+	if (CurveEditorOptions == nullptr)
 	{
-		OnKeySelectedDelegate.ExecuteIfBound(KeyHandle);
+		CurveEditorOptions = MakeShared<FNiagaraDistributionCurveEditorOptions>();
 	}
 
-	void KeyHandleMoved(float LocalDelta, FKeyHandle KeyHandle)
+	if (CurveEditorOptions->GetNeedsInitializeView())
 	{
-		float TimeMin;
-		float TimeMax;
-		EditCurve->GetTimeRange(TimeMin, TimeMax);
-		float TimeRange = TimeMax - TimeMin;
-		float NormalizedDelta = LocalDelta / LocalWidthCache;
-		if (SelectedKeyHandle.Get() != KeyHandle)
-		{
-			OnKeySelectedDelegate.ExecuteIfBound(KeyHandle);
-		}
-		OnKeyMovedDelegate.ExecuteIfBound(KeyHandle, NormalizedDelta * TimeRange);
+		InitializeView();
 	}
 
-private:
-	FRichCurve* EditCurve = nullptr;
-	FVector2f HandleSize;
-	TAttribute<FKeyHandle> SelectedKeyHandle;
-	TAttribute<int32> CurveStateSerialNumber;
-	FOnKeySelected OnKeySelectedDelegate;
-	FOnKeyMoved OnKeyMovedDelegate;
+	UNiagaraEditorSettings* EditorSettings = GetMutableDefault<UNiagaraEditorSettings>();
 
-	FGeometry AllottedGeometryCache;
-	float LocalWidthCache;
+	CurveEditor = MakeShared<FCurveEditor>();
 
-	TSharedPtr<SOverlay> BorderOverlay;
-	int32 CachedCurveStateSerialNumber = INDEX_NONE;
-	TArray<FKeyWidgetData> KeyWidgetData;
+	FCurveEditorInitParams InitParams;
+	CurveEditor->InitCurveEditor(InitParams);
+	CurveEditor->InputSnapEnabledAttribute = TAttribute<bool>::Create(TAttribute<bool>::FGetter::CreateUObject(EditorSettings, &UNiagaraEditorSettings::IsCurveInputSnapEnabled));
+	CurveEditor->OutputSnapEnabledAttribute = TAttribute<bool>::Create(TAttribute<bool>::FGetter::CreateUObject(EditorSettings, &UNiagaraEditorSettings::IsCurveOutputSnapEnabled));
+	CurveEditor->OnInputSnapEnabledChanged.BindUObject(EditorSettings, &UNiagaraEditorSettings::SetCurveInputSnapEnabled);
+	CurveEditor->OnOutputSnapEnabledChanged.BindUObject(EditorSettings, &UNiagaraEditorSettings::SetCurveOutputSnapEnabled);
+	CurveEditor->GridLineLabelFormatXAttribute = LOCTEXT("GridXLabelFormat", "{0}");
 
-	bool bNeedsWidgetUpdate = false;
-};
+	// Initialize our bounds at slightly larger than default to avoid clipping the tabs on the color widget.
+	TUniquePtr<ICurveEditorBounds> EditorBounds = MakeUnique<FNiagaraDistributionCurveEditorBounds>(CurveEditorOptions.ToSharedRef());
+	CurveEditor->SetBounds(MoveTemp(EditorBounds));
 
-void SNiagaraDistributionCurveEditor::Construct(const FArguments& InArgs, TSharedRef<INiagaraDistributionAdapter> InDistributionAdapter, int32 InChannelIndex)
-{
-	DistributionAdapter = InDistributionAdapter;
-	ChannelIndex = InChannelIndex;
-	const FRichCurve* EditCurvePtr = DistributionAdapter->GetCurveValue(ChannelIndex);
-	if (EditCurvePtr != nullptr)
-	{
-		EditCurve = *EditCurvePtr;
-	}
-	else
+	const int NumChannels = DistributionAdapter->GetNumChannels();
+	const bool bIsSingleCurve = NumChannels == 1;
+
+	CurveEditorPanel =
+		SNew(SCurveEditorPanel, CurveEditor.ToSharedRef())
+		.MinimumViewPanelHeight(50.0f)
+		.TreeContent()
+		[
+			bIsSingleCurve ? SNullWidget::NullWidget : SAssignNew(CurveEditorTree, SCurveEditorTree, CurveEditor).SelectColumnWidth(0)
+		];
+
+	TArray<FCurveEditorTreeItemID> TreeItemIds;
+	for (int iChannel=0; iChannel < NumChannels; ++iChannel)
 	{
-		EditCurve.AddKey(0, 0);
+		TSharedRef<FNiagaraDistributionCurveEditorDetailsTreeItem> TreeItem = MakeShared<FNiagaraDistributionCurveEditorDetailsTreeItem>(DistributionAdapter.ToSharedRef(), iChannel);
+		FCurveEditorTreeItem* NewItem = CurveEditor->AddTreeItem(FCurveEditorTreeItemID::Invalid());
+		NewItem->SetStrongItem(TreeItem);
+		TreeItemIds.Add(NewItem->GetID());
+		for (const FCurveModelID& CurveModel : NewItem->GetOrCreateCurves(CurveEditor.Get()))
+		{
+			CurveEditor->PinCurve(CurveModel);
+		}
 	}
-	UpdateSelectedKeyIndex(0);
 
-	float KeyHandleWidth = 16;
-	float KeyHandleHeight = 16;
+	const float KeySelectorLeftPadding = bIsSingleCurve ? 0 : 7;
+
+	FNiagaraEditorModule& NiagaraEditorModule = FModuleManager::LoadModuleChecked<FNiagaraEditorModule>("NiagaraEditor");
+	TSharedRef<INiagaraEditorWidgetProvider> WidgetProvider = NiagaraEditorModule.GetWidgetProvider();
+
 	ChildSlot
 	[
 		SNew(SVerticalBox)
@@ -409,34 +801,26 @@ void SNiagaraDistributionCurveEditor::Construct(const FArguments& InArgs, TShare
 		.AutoHeight()
 		.Padding(0, 0, 0, 3)
 		[
-			SNew(SHorizontalBox)
-			+ SHorizontalBox::Slot()
-			.AutoWidth()
-			.Padding(0, 0, 3, 0)
+			CreateToolbarWidget()
+		]
+		+ SVerticalBox::Slot()
+		.AutoHeight()
+		[
+			SNew(SVerticalResizeBox)
+			.ContentHeight(CurveEditorOptions.ToSharedRef(), &FNiagaraDistributionCurveEditorOptions::GetHeight)
+			.ContentHeightChanged(CurveEditorOptions.ToSharedRef(), &FNiagaraDistributionCurveEditorOptions::SetHeight)
+			.Content()
 			[
 				SNew(SVerticalBox)
 				+ SVerticalBox::Slot()
-				.AutoHeight()
-				.Padding(0, 0, 0, 2)
+				.Padding(0, 0, 0, 5)
 				[
-					SNew(SNumericEntryBox<int32>)
-					.Font(FAppStyle::Get().GetFontStyle("PropertyWindow.NormalFont"))
-					.Value(this, &SNiagaraDistributionCurveEditor::GetSelectedKeyIndex)
-					.OnValueChanged(this, &SNiagaraDistributionCurveEditor::SelectedKeyIndexChanged)
-					.OnValueCommitted(this, &SNiagaraDistributionCurveEditor::SelectedKeyIndexComitted)
-					.AllowSpin(true)
-					.MinValue(0)
-					.MaxValue(this, &SNiagaraDistributionCurveEditor::GetMaxKeyIndex)
-					.MinSliderValue(0)
-					.MaxSliderValue(this, &SNiagaraDistributionCurveEditor::GetMaxKeyIndex)
-					.Delta(1)
-					.MinDesiredValueWidth(8)
-					.LabelVAlign(VAlign_Center)
-					.Label()
+					SNew(SHorizontalBox)
+					+ SHorizontalBox::Slot()
+					.HAlign(HAlign_Fill)
+					.FillWidth(1.0f)
 					[
-						SNew(STextBlock)
-						.Font(FAppStyle::Get().GetFontStyle("PropertyWindow.NormalFont"))
-						.Text(LOCTEXT("KeyLabel", "Key"))
+						CurveEditorPanel.ToSharedRef()
 					]
 				]
 				+ SVerticalBox::Slot()
@@ -444,563 +828,221 @@ void SNiagaraDistributionCurveEditor::Construct(const FArguments& InArgs, TShare
 				[
 					SNew(SHorizontalBox)
 					+ SHorizontalBox::Slot()
+					.VAlign(VAlign_Center)
 					.AutoWidth()
-					.Padding(0, 0, 2, 0)
+					.Padding(KeySelectorLeftPadding, 0, 8, 0)
 					[
-						SNew(SButton)
-						.ButtonStyle( FAppStyle::Get(), "SimpleButton" )
-						.OnClicked(this, &SNiagaraDistributionCurveEditor::AddKeyButtonClicked)
-						.ToolTipText(LOCTEXT("AddKeyToolTip", "Add a key"))
-						.ContentPadding(0.0f)
-						[ 
-							SNew( SImage )
-							.Image( FAppStyle::GetBrush("Icons.PlusCircle") )
-							.ColorAndOpacity( FSlateColor::UseForeground() )
-						]
+						WidgetProvider->CreateCurveKeySelector(CurveEditor, TreeItemIds, CurveEditorTree)
 					]
 					+ SHorizontalBox::Slot()
+					.VAlign(VAlign_Center)
 					.AutoWidth()
-					.Padding(0, 0, 2, 0)
-					[
-						SNew(SButton)
-						.ButtonStyle(FAppStyle::Get(), "SimpleButton")
-						.OnClicked(this, &SNiagaraDistributionCurveEditor::DeleteKeyButtonClicked)
-						.ToolTipText(LOCTEXT("DeleteKeyToolTip", "Delete this key"))
-						.ContentPadding(0.0f)
-						[
-							SNew(SImage)
-							.Image(FAppStyle::GetBrush("Icons.Delete"))
-							.ColorAndOpacity(FSlateColor::UseForeground())
-						]
-					]
-				]
-				+ SVerticalBox::Slot()
-				.AutoHeight()
-				[
-					SNew(SHorizontalBox)
-						+ SHorizontalBox::Slot()
-						.AutoWidth()
-						.Padding(0, 0, 2, 0)
-						[
-							SNew(SButton)
-								.ButtonStyle(FAppStyle::Get(), "SimpleButton")
-								.OnClicked(this, &SNiagaraDistributionCurveEditor::FlipCurve, ECurveFlipDirection::Horizontal)
-								.ToolTipText(LOCTEXT("FlipHorizontalToolTip", "Flip curve horizontally"))
-								.ContentPadding(0.0f)
-								[
-									SNew(SImage)
-										.Image(FAppStyle::GetBrush("Icons.FlipHorizontal"))
-										.ColorAndOpacity(FSlateColor::UseForeground())
-								]
-						]
-						+ SHorizontalBox::Slot()
-						.AutoWidth()
-						.Padding(0, 0, 2, 0)
-						[
-							SNew(SButton)
-								.ButtonStyle(FAppStyle::Get(), "SimpleButton")
-								.OnClicked(this, &SNiagaraDistributionCurveEditor::FlipCurve, ECurveFlipDirection::Vertical)
-								.ToolTipText(LOCTEXT("FlipVerticalToolTip", "Flip curve vertically"))
-								.ContentPadding(0.0f)
-								[
-									SNew(SImage)
-										.Image(FAppStyle::GetBrush("Icons.FlipVertical"))
-										.ColorAndOpacity(FSlateColor::UseForeground())
-								]
-						]
-				]
-			]
-			// Curve view and key selector.
-			+ SHorizontalBox::Slot()
-			.Padding(2)
-			[
-				SNew(SGridPanel)
-				.FillRow(0, 1)
-				.FillColumn(0, 1)
-				// Value Labels
-				+ SGridPanel::Slot(1, 0)
-				[
-					SNew(SVerticalBox)
-					+ SVerticalBox::Slot()
-					.VAlign(VAlign_Top)
+					.Padding(0, 0, 3, 0)
 					[
 						SNew(STextBlock)
-						.Font(FAppStyle::Get().GetFontStyle("SmallFont"))
-						.Text(this, &SNiagaraDistributionCurveEditor::GetCurveValueMaxText)
+						.Text(LOCTEXT("KeyLabel", "Key Data"))
 					]
-					+ SVerticalBox::Slot()
-					.VAlign(VAlign_Bottom)
-					[
-						SNew(STextBlock)
-						.Font(FAppStyle::Get().GetFontStyle("SmallFont"))
-						.Text(this, &SNiagaraDistributionCurveEditor::GetCurveValueMinText)
-					]
-				]
-				// Curve view
-				+ SGridPanel::Slot(0, 0)
-				.Padding(FMargin(KeyHandleWidth / 2, 3, KeyHandleWidth / 2, 3))
-				[
-					SNew(SNiagaraCompactCurveView, &EditCurve)
-					.Width(1000)
-					.Height(50)
-					.CurveColor(InArgs._CurveColor)
-					.SelectedKeyHandle(this, &SNiagaraDistributionCurveEditor::GetSelectedKeyHandle)
-					.CurveStateSerialNumber(this, &SNiagaraDistributionCurveEditor::GetCurveStateSerialNumber)
-				]
-				// Time Labels
-				+ SGridPanel::Slot(0, 2)
-				[
-					SNew(SHorizontalBox)
 					+ SHorizontalBox::Slot()
-					.HAlign(HAlign_Left)
+					.VAlign(VAlign_Center)
+					.FillWidth(1.0f)
 					[
-						SNew(STextBlock)
-						.Font(FAppStyle::Get().GetFontStyle("SmallFont"))
-						.Text(this, &SNiagaraDistributionCurveEditor::GetCurveTimeMinText)
-					]
-					+ SHorizontalBox::Slot()
-					.HAlign(HAlign_Right)
-					[
-						SNew(STextBlock)
-						.Font(FAppStyle::Get().GetFontStyle("SmallFont"))
-						.Text(this, &SNiagaraDistributionCurveEditor::GetCurveTimeMaxText)
+						CurveEditorPanel->GetKeyDetailsView().ToSharedRef()
 					]
 				]
-				// key selector
-				+ SGridPanel::Slot(0, 1)
-				.Padding(0, 0, 0, 3)
-				[
-					SAssignNew(KeySelectorWidget, SNiagaraCompactCurveKeySelector, &EditCurve)
-					.HandleSize(FVector2f(KeyHandleWidth, KeyHandleHeight))
-					.SelectedKeyHandle(this, &SNiagaraDistributionCurveEditor::GetSelectedKeyHandle)
-					.CurveStateSerialNumber(this, &SNiagaraDistributionCurveEditor::GetCurveStateSerialNumber)
-					.OnKeySelected(this, &SNiagaraDistributionCurveEditor::OnKeySelectorKeySelected)
-					.OnKeyMoved(this, &SNiagaraDistributionCurveEditor::OnKeySelectorKeyMoved)
-				]
-			]
-		]
-		+ SVerticalBox::Slot()
-		.AutoHeight()
-		[
-			SNew(SHorizontalBox)
-			+ SHorizontalBox::Slot()
-			.Padding(0, 0, 5, 0)
-			[
-				SNew(SNumericEntryBox<float>)
-				.Font(FAppStyle::Get().GetFontStyle("PropertyWindow.NormalFont"))
-				.Value(this, &SNiagaraDistributionCurveEditor::GetSelectedKeyTime)
-				.OnValueChanged(this, &SNiagaraDistributionCurveEditor::SelectedKeyTimeChanged)
-				.OnValueCommitted(this, &SNiagaraDistributionCurveEditor::SelectedKeyTimeComitted)
-				.OnBeginSliderMovement(this, &SNiagaraDistributionCurveEditor::BeginSliderMovement)
-				.OnEndSliderMovement(this, &SNiagaraDistributionCurveEditor::EndSliderMovement)
-				.AllowSpin(true)
-				.MinValue(TOptional<float>())
-				.MaxValue(TOptional<float>())
-				.MinSliderValue(TOptional<float>())
-				.MaxSliderValue(TOptional<float>())
-				.LabelVAlign(VAlign_Center)
-				.Label()
-				[
-					SNew(STextBlock)
-					.Font(FAppStyle::Get().GetFontStyle("PropertyWindow.NormalFont"))
-					.Text(LOCTEXT("TimeLabel", "Time"))
-				]
-			]
-			+ SHorizontalBox::Slot()
-			.Padding(0, 0, 5, 0)
-			[
-				SNew(SNumericEntryBox<float>)
-				.Font(FAppStyle::Get().GetFontStyle("PropertyWindow.NormalFont"))
-				.Value(this, &SNiagaraDistributionCurveEditor::GetSelectedKeyValue)
-				.OnValueChanged(this, &SNiagaraDistributionCurveEditor::SelectedKeyValueChanged)
-				.OnValueCommitted(this, &SNiagaraDistributionCurveEditor::SelectedKeyValueComitted)
-				.OnBeginSliderMovement(this, &SNiagaraDistributionCurveEditor::BeginSliderMovement)
-				.OnEndSliderMovement(this, &SNiagaraDistributionCurveEditor::EndSliderMovement)
-				.AllowSpin(true)
-				.MinValue(TOptional<float>())
-				.MaxValue(TOptional<float>())
-				.MinSliderValue(TOptional<float>())
-				.MaxSliderValue(TOptional<float>())
-				.LabelVAlign(VAlign_Center)
-				.Label()
-				[
-					SNew(STextBlock)
-					.Font(FAppStyle::Get().GetFontStyle("PropertyWindow.NormalFont"))
-					.Text(LOCTEXT("ValueLabel", "Value"))
-				]
-			]
-			+ SHorizontalBox::Slot()
-			[
-				SNew(SHorizontalBox)
-				+ SHorizontalBox::Slot()
-				.VAlign(VAlign_Center)
-				.AutoWidth()
-				.Padding(0, 0, 3, 0)
-				[
-					SNew(STextBlock)
-					.Font(FAppStyle::Get().GetFontStyle("PropertyWindow.NormalFont"))
-					.Text(LOCTEXT("InterpLabel", "Interp"))
-				]
-				+ SHorizontalBox::Slot()
-				[
-					SNew(SEnumComboBox, StaticEnum<ERichCurveInterpMode>())
-					.Font(FAppStyle::Get().GetFontStyle("PropertyWindow.NormalFont"))
-					.ContentPadding(FMargin(2, 0))
-					.CurrentValue(this, &SNiagaraDistributionCurveEditor::GetSelectedKeyInterpMode)
-					.OnEnumSelectionChanged(this, &SNiagaraDistributionCurveEditor::SelectedKeyInterpModeChanged)
-				]
 			]
 		]
 	];
 }
 
-TOptional<int32> SNiagaraDistributionCurveEditor::GetSelectedKeyIndex() const
-{
-	return SelectedKeyIndex;
-}
-
-TOptional<int32> SNiagaraDistributionCurveEditor::GetMaxKeyIndex() const
+void SNiagaraDistributionCurveEditor::InitializeView()
 {
-	return EditCurve.GetNumKeys() - 1;
-}
+	bool bHasKeys = false;
+	float ViewMinInput = TNumericLimits<float>::Max();
+	float ViewMaxInput = TNumericLimits<float>::Lowest();
+	float ViewMinOutput = TNumericLimits<float>::Max();
+	float ViewMaxOutput = TNumericLimits<float>::Lowest();
 
-void SNiagaraDistributionCurveEditor::SelectedKeyIndexChanged(int32 InValue)
-{
-	UpdateSelectedKeyIndex(InValue);
-}
-
-void SNiagaraDistributionCurveEditor::SelectedKeyIndexComitted(int32 InValue, ETextCommit::Type CommitInfo)
-{
-	if (CommitInfo == ETextCommit::OnEnter || CommitInfo == ETextCommit::OnUserMovedFocus)
+	for (int i=0; i < DistributionAdapter->GetNumChannels(); ++i)
 	{
-		UpdateSelectedKeyIndex(InValue);
-	}
-}
-
-void SNiagaraDistributionCurveEditor::UpdateSelectedKeyIndex(int32 NewKeyIndex)
-{
-	SelectedKeyIndex = NewKeyIndex;
-	int32 KeyIndex = 0;
-	for (auto KeyHandleIterator = EditCurve.GetKeyHandleIterator(); KeyHandleIterator; KeyHandleIterator++)
-	{
-		if (KeyIndex == NewKeyIndex)
+		const FRealCurve* Curve = DistributionAdapter->GetCurveValue(i);
+		for (auto KeyIterator = Curve->GetKeyHandleIterator(); KeyIterator; ++KeyIterator)
 		{
-			SelectedKeyHandle = *KeyHandleIterator;
-			break;
+			float KeyTime = Curve->GetKeyTime(*KeyIterator);
+			float KeyValue = Curve->GetKeyValue(*KeyIterator);
+			ViewMinInput = FMath::Min(ViewMinInput, KeyTime);
+			ViewMaxInput = FMath::Max(ViewMaxInput, KeyTime);
+			ViewMinOutput = FMath::Min(ViewMinOutput, KeyValue);
+			ViewMaxOutput = FMath::Max(ViewMaxOutput, KeyValue);
+			bHasKeys = true;
 		}
-		KeyIndex++;
 	}
-	CurveStateSerialNumber++;
-}
 
-TOptional<float> SNiagaraDistributionCurveEditor::GetSelectedKeyTime() const
-{
-	UpdateCachedValues();
-	return KeyTimeCache;
-}
-
-void SNiagaraDistributionCurveEditor::SelectedKeyTimeChanged(float InValue)
-{
-	UpdateSelectedKeyTime(InValue);
-}
-
-void SNiagaraDistributionCurveEditor::SelectedKeyTimeComitted(float InValue, ETextCommit::Type CommitInfo)
-{
-	if (CommitInfo == ETextCommit::OnEnter || CommitInfo == ETextCommit::OnUserMovedFocus)
+	if (bHasKeys == false)
 	{
-		UpdateSelectedKeyTime(InValue);
+		ViewMinInput = 0;
+		ViewMaxInput = 1;
+		ViewMinOutput = 0;
+		ViewMaxOutput = 1;
 	}
-}
 
-void SNiagaraDistributionCurveEditor::UpdateSelectedKeyTime(float NewKeyTime)
-{
-	if (GetSelectedKeyTime() != NewKeyTime)
+	if (FMath::IsNearlyEqual(ViewMinInput, ViewMaxInput))
 	{
-		EditCurve.SetKeyTime(SelectedKeyHandle, NewKeyTime);
-		DistributionAdapter->SetCurveValue(ChannelIndex, EditCurve);
-
-		int32 NewKeyIndex = INDEX_NONE;
-		int32 KeyIndex = 0;
-		for (auto KeyHandleIterator = EditCurve.GetKeyHandleIterator(); KeyHandleIterator; KeyHandleIterator++)
+		if (FMath::IsWithinInclusive(ViewMinInput, 0.0f, 1.0f))
 		{
-			if (SelectedKeyHandle == *KeyHandleIterator)
-			{
-				NewKeyIndex = KeyIndex;
-				break;
-			}
-			KeyIndex++;
+			ViewMinInput = 0;
+			ViewMaxInput = 1;
+		}
+		else
+		{
+			ViewMinInput -= 0.5f;
+			ViewMaxInput += 0.5f;
 		}
-
-		SelectedKeyIndex = NewKeyIndex;
-		KeyTimeCache = NewKeyTime;
-		CurveStateSerialNumber++;
-	}
-}
-
-TOptional<float> SNiagaraDistributionCurveEditor::GetSelectedKeyValue() const
-{
-	UpdateCachedValues();
-	return KeyValueCache;
-}
-
-void SNiagaraDistributionCurveEditor::SelectedKeyValueChanged(float InValue)
-{
-	UpdateSelectedKeyValue(InValue);
-}
-
-void SNiagaraDistributionCurveEditor::SelectedKeyValueComitted(float InValue, ETextCommit::Type CommitInfo)
-{
-	if (CommitInfo == ETextCommit::OnEnter || CommitInfo == ETextCommit::OnUserMovedFocus)
-	{
-		UpdateSelectedKeyValue(InValue);
-	}
-}
-
-void SNiagaraDistributionCurveEditor::UpdateSelectedKeyValue(float NewKeyValue)
-{
-	if (GetSelectedKeyValue() != NewKeyValue)
-	{
-		EditCurve.SetKeyValue(SelectedKeyHandle, NewKeyValue);
-		DistributionAdapter->SetCurveValue(ChannelIndex, EditCurve);
-		KeyValueCache = NewKeyValue;
-		CurveStateSerialNumber++;
-	}
-}
-
-int32 SNiagaraDistributionCurveEditor::GetSelectedKeyInterpMode() const
-{
-	UpdateCachedValues();
-	return KeyInterpModeCache;
-}
-
-void SNiagaraDistributionCurveEditor::SelectedKeyInterpModeChanged(int32 InNewValue, ESelectInfo::Type InInfo)
-{
-	if (GetSelectedKeyInterpMode() != InNewValue)
-	{
-		EditCurve.SetKeyInterpMode(SelectedKeyHandle, (ERichCurveInterpMode)InNewValue);
-		DistributionAdapter->SetCurveValue(ChannelIndex, EditCurve);
-		KeyInterpModeCache = InNewValue;
-		CurveStateSerialNumber++;
 	}
-}
 
-void SNiagaraDistributionCurveEditor::BeginSliderMovement()
-{
-	DistributionAdapter->BeginContinuousChange();
-}
-
-void SNiagaraDistributionCurveEditor::EndSliderMovement(float Value)
-{
-	DistributionAdapter->EndContinuousChange();
-}
-
-FReply SNiagaraDistributionCurveEditor::AddKeyButtonClicked()
-{
-	int32 NumKeys = EditCurve.GetNumKeys();
-	int32 NewSelectedKeyIndex;
-	if (NumKeys == 0)
-	{
-		EditCurve.AddKey(0, 0);
-		NewSelectedKeyIndex = 0;
-	}
-	else if (NumKeys == 1)
-	{
-		FRichCurveKey FirstKey = EditCurve.GetFirstKey();
-		EditCurve.AddKey(FirstKey.Time + 1, FirstKey.Value);
-		NewSelectedKeyIndex = 1;
-	}
-	else
+	if (FMath::IsNearlyEqual(ViewMinOutput, ViewMaxOutput))
 	{
-		if (SelectedKeyIndex < NumKeys - 1)
+		if (FMath::IsWithinInclusive(ViewMinOutput, 0.0f, 1.0f))
 		{
-			FKeyHandle NextKeyHandle = EditCurve.GetNextKey(SelectedKeyHandle);
-			FRichCurveKey SelectedKey = EditCurve.GetKey(SelectedKeyHandle);
-			FRichCurveKey NextKey = EditCurve.GetKey(NextKeyHandle);
-			EditCurve.AddKey((SelectedKey.Time + NextKey.Time) / 2, (SelectedKey.Value + NextKey.Value) / 2);
-			NewSelectedKeyIndex = SelectedKeyIndex + 1;
+			ViewMinOutput = 0;
+			ViewMaxOutput = 1;
 		}
 		else
 		{
-			FKeyHandle PreviousKeyHandle = EditCurve.GetPreviousKey(SelectedKeyHandle);
-			FRichCurveKey SelectedKey = EditCurve.GetKey(SelectedKeyHandle);
-			FRichCurveKey PreviousKey = EditCurve.GetKey(PreviousKeyHandle);
-			EditCurve.AddKey(SelectedKey.Time + (SelectedKey.Time - PreviousKey.Time), SelectedKey.Value + (SelectedKey.Value - PreviousKey.Value));
-			NewSelectedKeyIndex = SelectedKeyIndex;
+			ViewMinOutput -= 0.5f;
+			ViewMaxOutput += 0.5f;
 		}
 	}
 
-	DistributionAdapter->SetCurveValue(ChannelIndex, EditCurve);
-	UpdateSelectedKeyIndex(NewSelectedKeyIndex);
-	CurveStateSerialNumber++;
-	return FReply::Handled();
-}
+	float ViewInputRange = ViewMaxInput - ViewMinInput;
+	float ViewOutputRange = ViewMaxOutput - ViewMinOutput;
+	float ViewInputPadding = ViewInputRange * .05f;
+	float ViewOutputPadding = ViewOutputRange * .05f;
 
-FReply SNiagaraDistributionCurveEditor::DeleteKeyButtonClicked()
-{
-	if (EditCurve.GetNumKeys() > 1)
-	{
-		EditCurve.DeleteKey(GetSelectedKeyHandle());
-		DistributionAdapter->SetCurveValue(ChannelIndex, EditCurve);
-		int32 NewKeyIndex = FMath::Clamp(SelectedKeyIndex - 1, 0, EditCurve.GetNumKeys() - 1);
-		UpdateSelectedKeyIndex(NewKeyIndex);
-		CurveStateSerialNumber++;
-		return FReply::Handled();
-	}
-	return FReply::Unhandled();
+	CurveEditorOptions->InitializeView(
+		ViewMinInput - ViewInputPadding,
+		ViewMaxInput + ViewInputPadding,
+		ViewMinOutput - ViewOutputPadding,
+		ViewMaxOutput + ViewOutputPadding);
 }
 
-FReply SNiagaraDistributionCurveEditor::FlipCurve(ECurveFlipDirection Direction)
+TSharedRef<SWidget> SNiagaraDistributionCurveEditor::CreateToolbarWidget()
 {
-	if (EditCurve.GetNumKeys() <= 0)
-	{
-		return FReply::Unhandled();
-	}
-
-	FScopedTransaction Transaction(LOCTEXT("FlipCurve", "Flip Curve"));
-
-	int32 NumKeys = EditCurve.GetNumKeys();
-	int32 NewSelectedKeyIndex = SelectedKeyIndex;
-
-	// Intialize the range
-	float RangeMin = FLT_MAX;
-	float RangeMax = -FLT_MAX;
-
-	// If flipping horizontally
-	if (Direction == ECurveFlipDirection::Horizontal)
+	FToolBarBuilder ToolBarBuilder(CurveEditor->GetCommands(), FMultiBoxCustomization::None, nullptr, true);
+	ToolBarBuilder.SetStyle(&FAppStyle::Get(), "CurveEditorToolBar");
+	ToolBarBuilder.AddToolBarButton(FCurveEditorCommands::Get().ToggleInputSnapping);
+	ToolBarBuilder.AddToolBarButton(FCurveEditorCommands::Get().ToggleOutputSnapping);
+	ToolBarBuilder.AddSeparator();
+	ToolBarBuilder.AddToolBarButton(FCurveEditorCommands::Get().FlipCurveHorizontal);
+	ToolBarBuilder.AddToolBarButton(FCurveEditorCommands::Get().FlipCurveVertical);
+	ToolBarBuilder.AddSeparator();
+
+	const UNiagaraEditorSettings* EditorSettings = GetDefault<UNiagaraEditorSettings>();
+	if (EditorSettings->GetCurveTemplates().Num() > 0)
 	{
-		// Get the range
-		float MinTime, MaxTime;
-		EditCurve.GetTimeRange(MinTime, MaxTime);
-		RangeMin = FMath::Min(MinTime, RangeMin);
-		RangeMax = FMath::Max(MaxTime, RangeMax);
-
-		TArray<TPair<float, float>> TimeValuePairs;
-		TArray<ERichCurveInterpMode> InterpModes;
-		TimeValuePairs.SetNum(NumKeys);
-		InterpModes.SetNum(NumKeys);
-
-		// Get the (flipped x, original y) value of all points
-		FKeyHandle KeyHandle = EditCurve.GetFirstKeyHandle();
-		for (int32 Index = 0; Index <= NumKeys - 1; ++Index)
-		{
-			FRichCurveKey Key = EditCurve.GetKey(KeyHandle);
-			TimeValuePairs[Index].Key = RangeMax - Key.Time + RangeMin;
-			TimeValuePairs[Index].Value = Key.Value;
-			InterpModes[Index] = EditCurve.GetKeyInterpMode(KeyHandle);
-			KeyHandle = EditCurve.GetNextKey(KeyHandle);
-		}
+		FNiagaraEditorModule& NiagaraEditorModule = FModuleManager::LoadModuleChecked<FNiagaraEditorModule>("NiagaraEditor");
+		TSharedRef<INiagaraEditorWidgetProvider> WidgetProvider = NiagaraEditorModule.GetWidgetProvider();
 
-		// Delete all existing keys
-		EditCurve.Reset();
-
-		// Add the flipped points back in
-		FKeyHandle NewKeyHandle;
-		for (int32 Index = 0; Index <= NumKeys - 1; ++Index)
-		{
-			NewKeyHandle = EditCurve.AddKey(TimeValuePairs[Index].Key, TimeValuePairs[Index].Value);
-			EditCurve.SetKeyInterpMode(NewKeyHandle, InterpModes[Index]);
-			//NOTE: key tangents are not handled right now since distribution curve editor currently don't support tangent editing from user
-		}
+		ToolBarBuilder.AddWidget(
+			SNew(SBox)
+			.VAlign(VAlign_Center)
+			.Padding(FMargin(3, 0, 5, 0))
+			[
+				SNew(STextBlock)
+				.Text(LOCTEXT("CurveTemplateLabel", "Templates"))
+			]
+		);
 
-		NewSelectedKeyIndex = NumKeys - 1 - SelectedKeyIndex;
-		if (KeySelectorWidget) 
+		for (const FNiagaraCurveTemplate& CurveTemplate : EditorSettings->GetCurveTemplates())
 		{
-			KeySelectorWidget->RequestWidgetUpdate();
-		}
-	}
+			UCurveFloat* FloatCurveAsset = Cast<UCurveFloat>(CurveTemplate.CurveAsset.TryLoad());
+			if (FloatCurveAsset == nullptr)
+			{
+				continue;
+			}
 
-	// If flipping vertically
-	if (Direction == ECurveFlipDirection::Vertical)
-	{
-		// Get the range
-		float MinVal, MaxVal;
-		EditCurve.GetValueRange(MinVal, MaxVal);
-		RangeMin = FMath::Min(MinVal, RangeMin);
-		RangeMax = FMath::Max(MaxVal, RangeMax);
-
-		// Flip the y value of all keys
-		FKeyHandle KeyHandle = EditCurve.GetFirstKeyHandle();
-		for (int32 Index = 0; Index <= NumKeys - 1; ++Index)
-		{
-			EditCurve.SetKeyValue(KeyHandle, RangeMax - EditCurve.GetKey(KeyHandle).Value + RangeMin);
-			KeyHandle = EditCurve.GetNextKey(KeyHandle);
+			FText CurveDisplayName =
+				CurveTemplate.DisplayNameOverride.IsEmpty()
+				? FText::FromName(FloatCurveAsset->GetFName()) //FText::FromString(FName::NameToDisplayString(FloatCurveAsset->GetName(), false))
+				: FText::FromString(CurveTemplate.DisplayNameOverride);
+
+			TWeakObjectPtr<UCurveFloat> WeakCurveAsset(FloatCurveAsset);
+			ToolBarBuilder.AddWidget(
+				SNew(SButton)
+				.ButtonStyle(FAppStyle::Get(), "SimpleButton")
+				.OnClicked(this, &SNiagaraDistributionCurveEditor::ApplyCurveTemplate, WeakCurveAsset)
+				.ToolTipText(FText::Format(LOCTEXT("ApplyCurveTemplateFormat", "{0}\nClick to apply this template to the selected curves."), CurveDisplayName))
+				.ContentPadding(FMargin(3, 10, 3, 0))
+				.Content()
+				[
+					WidgetProvider->CreateCurveThumbnail(FloatCurveAsset->FloatCurve)
+				]
+			);
 		}
 	}
 
-	DistributionAdapter->SetCurveValue(ChannelIndex, EditCurve);
-	UpdateSelectedKeyIndex(NewSelectedKeyIndex);
-	CurveStateSerialNumber++;
-	return FReply::Handled();
+	return ToolBarBuilder.MakeWidget();
 }
 
-void SNiagaraDistributionCurveEditor::OnKeySelectorKeySelected(FKeyHandle InSelectedKeyHandle)
+FReply SNiagaraDistributionCurveEditor::ApplyCurveTemplate(TWeakObjectPtr<UCurveFloat> WeakCurveAsset)
 {
-	int32 KeyIndex = 0;
-	for (auto KeyHandleIterator = EditCurve.GetKeyHandleIterator(); KeyHandleIterator; KeyHandleIterator++)
+	UCurveFloat* FloatCurveAsset = WeakCurveAsset.Get();
+	if (FloatCurveAsset != nullptr)
 	{
-		if (*KeyHandleIterator == InSelectedKeyHandle)
+		TArray<FCurveModelID> CurveModelIdsToSet;
+		if (CurveEditor->GetRootTreeItems().Num() == 1)
 		{
-			SelectedKeyHandle = *KeyHandleIterator;
-			SelectedKeyIndex = KeyIndex;
-			CurveStateSerialNumber++;
-			break;
+			const FCurveEditorTreeItem& TreeItem = CurveEditor->GetTreeItem(CurveEditor->GetRootTreeItems()[0]);
+			for (const FCurveModelID& CurveModelId : TreeItem.GetCurves())
+			{
+				CurveModelIdsToSet.Add(CurveModelId);
+			}
+		}
+		else
+		{
+			for (const TPair<FCurveEditorTreeItemID, ECurveEditorTreeSelectionState>& TreeItemSelectionState : CurveEditor->GetTreeSelection())
+			{
+				if (TreeItemSelectionState.Value != ECurveEditorTreeSelectionState::None)
+				{
+					const FCurveEditorTreeItem& TreeItem = CurveEditor->GetTreeItem(TreeItemSelectionState.Key);
+					for (const FCurveModelID& CurveModelId : TreeItem.GetCurves())
+					{
+						CurveModelIdsToSet.Add(CurveModelId);
+					}
+				}
+			}
 		}
-		KeyIndex++;
-	}
-}
-
-void SNiagaraDistributionCurveEditor::OnKeySelectorKeyMoved(FKeyHandle InSelectedKeyHandle, float TimeDelta)
-{
-	if (SelectedKeyHandle == InSelectedKeyHandle)
-	{
-		UpdateSelectedKeyTime(GetSelectedKeyTime().GetValue() + TimeDelta);
-	}
-}
-
-FText SNiagaraDistributionCurveEditor::GetCurveTimeMinText() const
-{
-	UpdateCachedValues();
-	return CurveTimeMinTextCache;
-}
-
-FText SNiagaraDistributionCurveEditor::GetCurveTimeMaxText() const
-{
-	UpdateCachedValues();
-	return CurveTimeMaxTextCache;
-}
-
-FText SNiagaraDistributionCurveEditor::GetCurveValueMinText() const
-{
-	UpdateCachedValues();
-	return CurveValueMinTextCache;
-}
-
-FText SNiagaraDistributionCurveEditor::GetCurveValueMaxText() const
-{
-	UpdateCachedValues();
-	return CurveValueMaxTextCache;
-}
 
-void SNiagaraDistributionCurveEditor::UpdateCachedValues() const
-{
-	if (CacheCurveStateSerialNumber != CurveStateSerialNumber)
-	{
-		CacheCurveStateSerialNumber = CurveStateSerialNumber;
-		KeyTimeCache = EditCurve.GetKeyTime(SelectedKeyHandle);
-		KeyValueCache = EditCurve.GetKeyValue(SelectedKeyHandle);
-		KeyInterpModeCache = (int32)EditCurve.GetKeyInterpMode(SelectedKeyHandle);
-
-		float MinTime;
-		float MaxTime;
-		float MinValue;
-		float MaxValue;
-		EditCurve.GetTimeRange(MinTime, MaxTime);
-		EditCurve.GetValueRange(MinValue, MaxValue);
-		CurveTimeMinTextCache = FText::AsNumber(MinTime);
-		CurveTimeMaxTextCache = FText::AsNumber(MaxTime);
-		CurveValueMinTextCache = FText::AsNumber(MinValue);
-		CurveValueMaxTextCache = FText::AsNumber(MaxValue);
+		if (CurveModelIdsToSet.Num() > 0)
+		{
+			FScopedTransaction ApplyTemplateTransaction(LOCTEXT("ApplyCurveTemplateTransaction", "Apply curve template"));
+			for (const FCurveModelID& CurveModelId : CurveModelIdsToSet)
+			{
+				FCurveModel* CurveModel = CurveEditor->GetCurves()[CurveModelId].Get();
+				if (CurveModel != nullptr)
+				{
+					const TArray<FKeyHandle> KeyHandles = CurveModel->GetAllKeys();
+					CurveModel->RemoveKeys(KeyHandles, 0.0);
+
+					const FRichCurve& FloatCurve = FloatCurveAsset->FloatCurve;
+					for (auto KeyIterator = FloatCurve.GetKeyHandleIterator(); KeyIterator; ++KeyIterator)
+					{
+						const FRichCurveKey& Key = FloatCurve.GetKey(*KeyIterator);
+						FKeyPosition KeyPosition;
+						KeyPosition.InputValue = Key.Time;
+						KeyPosition.OutputValue = Key.Value;
+						FKeyAttributes KeyAttributes;
+						KeyAttributes.SetInterpMode(Key.InterpMode);
+						KeyAttributes.SetTangentMode(Key.TangentMode);
+						KeyAttributes.SetArriveTangent(Key.ArriveTangent);
+						KeyAttributes.SetLeaveTangent(Key.LeaveTangent);
+						CurveModel->AddKey(KeyPosition, KeyAttributes);
+					}
+				}
+			}
+			CurveEditor->ZoomToFit();
+		}
 	}
+	return FReply::Handled();
 }
 
-#undef LOCTEXT_NAMESPACE
\ No newline at end of file
+#undef LOCTEXT_NAMESPACE
diff --git a/Engine/Plugins/FX/Niagara/Source/NiagaraEditor/Private/Widgets/SNiagaraDistributionCurveEditor.h b/Engine/Plugins/FX/Niagara/Source/NiagaraEditor/Private/Widgets/SNiagaraDistributionCurveEditor.h
--- a/Engine/Plugins/FX/Niagara/Source/NiagaraEditor/Private/Widgets/SNiagaraDistributionCurveEditor.h
+++ b/Engine/Plugins/FX/Niagara/Source/NiagaraEditor/Private/Widgets/SNiagaraDistributionCurveEditor.h
@@ -2,90 +2,76 @@
 
 #pragma once
 
-#include "Curves/RichCurve.h"
-#include "Styling/StyleColors.h"
+#include "Types/SlateStructs.h"
 #include "Widgets/SCompoundWidget.h"
 
+class FCurveEditor;
+class SCurveEditorPanel;
+class SCurveEditorTree;
+class UCurveFloat;
 class INiagaraDistributionAdapter;
-class SNiagaraCompactCurveKeySelector;
 
-class SNiagaraDistributionCurveEditor : public SCompoundWidget
+class FNiagaraDistributionCurveEditorOptions
 {
 public:
-	SLATE_BEGIN_ARGS(SNiagaraDistributionCurveEditor)
-		: _CurveColor(FSlateColor(EStyleColor::White))
-		{ }
-		SLATE_EVENT(FSimpleDelegate, OnCurveChanged)
-		SLATE_ARGUMENT(FSlateColor, CurveColor)
-	SLATE_END_ARGS();
-
-	void Construct(const FArguments& InArgs, TSharedRef<INiagaraDistributionAdapter> InDistributionAdapter, int32 InChannelIndex);
-
-private:
-	/** Curve flip direction types **/
-	enum class ECurveFlipDirection : uint8
+	void InitializeView(float InViewMinInput, float InViewMaxInput, float InViewMinOutput, float InViewMaxOutput)
 	{
-		Horizontal,
-		Vertical
-	};
+		ViewMinInput = InViewMinInput;
+		ViewMaxInput = InViewMaxInput;
+		ViewMinOutput = InViewMinOutput;
+		ViewMaxOutput = InViewMaxOutput;
+		bNeedsInitializeView = false;
+	}
 
-private:
-	TOptional<int32> GetSelectedKeyIndex() const;
-	TOptional<int32> GetMaxKeyIndex() const;
-	void SelectedKeyIndexChanged(int32 InValue);
-	void SelectedKeyIndexComitted(int32 InValue, ETextCommit::Type CommitInfo);
-	void UpdateSelectedKeyIndex(int32 NewKeyIndex);
+	bool GetNeedsInitializeView() const { return bNeedsInitializeView; }
+
+	float GetViewMinInput() const { return ViewMinInput; }
+	float GetViewMaxInput() const { return ViewMaxInput; }
+	void SetInputViewRange(float InViewMinInput, float InViewMaxInput) { ViewMinInput = InViewMinInput; ViewMaxInput = InViewMaxInput; }
 
-	TOptional<float> GetSelectedKeyTime() const;
-	void SelectedKeyTimeChanged(float InValue);
-	void SelectedKeyTimeComitted(float InValue, ETextCommit::Type CommitInfo);
-	void UpdateSelectedKeyTime(float NewKeyTime);
+	float GetViewMinOutput() const { return ViewMinOutput; }
+	float GetViewMaxOutput() const { return ViewMaxOutput; }
+	void SetOutputViewRange(float InViewMinOutput, float InViewMaxOutput) { ViewMinOutput = InViewMinOutput; ViewMaxOutput = InViewMaxOutput; }
 
-	TOptional<float> GetSelectedKeyValue() const;
-	void SelectedKeyValueChanged(float InValue);
-	void SelectedKeyValueComitted(float InValue, ETextCommit::Type CommitInfo);
-	void UpdateSelectedKeyValue(float NewKeyValue);
+	bool GetIsGradientVisible() const { return bIsGradientVisible; }
+	void SetIsGradientVisible(bool bInIsGradientVisible) { bIsGradientVisible = bInIsGradientVisible; }
 
-	int32 GetSelectedKeyInterpMode() const;
-	void SelectedKeyInterpModeChanged(int32 InNewValue, ESelectInfo::Type InInfo);
+	float GetTimelineLength() const { return ViewMaxInput - ViewMinInput; }
 
-	void BeginSliderMovement();
-	void EndSliderMovement(float Value);
+	float GetHeight() const { return Height; }
+	void SetHeight(float InHeight) { Height = InHeight; }
 
-	FReply AddKeyButtonClicked();
-	FReply DeleteKeyButtonClicked();
-	FReply FlipCurve(ECurveFlipDirection Direction);
+private:
+	float ViewMinInput = 0.0f;
+	float ViewMaxInput = 1.0f;
+	float ViewMinOutput = 0.0f;
+	float ViewMaxOutput = 1.0f;
+	bool bIsGradientVisible = true;
+	bool bNeedsInitializeView = true;
+	float Height = 180.0f;
+};
 
-	FKeyHandle GetSelectedKeyHandle() const { return SelectedKeyHandle; }
-	int32 GetCurveStateSerialNumber() const { return CurveStateSerialNumber; }
+class SNiagaraDistributionCurveEditor : public SCompoundWidget
+{
+public:
 
-	void OnKeySelectorKeySelected(FKeyHandle InSelectedKeyHandle);
-	void OnKeySelectorKeyMoved(FKeyHandle InSelectedKeyHandle, float TimeDelta);
+public:
+	SLATE_BEGIN_ARGS(SNiagaraDistributionCurveEditor)
+		{ }
+		SLATE_ARGUMENT(TSharedPtr<FNiagaraDistributionCurveEditorOptions>, CurveEditorOptions)
+	SLATE_END_ARGS();
 
-	FText GetCurveTimeMinText() const;
-	FText GetCurveTimeMaxText() const;
-	FText GetCurveValueMinText() const;
-	FText GetCurveValueMaxText() const;
+	void Construct(const FArguments& InArgs, TSharedRef<INiagaraDistributionAdapter> InCurveEditorAdapter);
 
-	void UpdateCachedValues() const;
+private:
+	void InitializeView();
+	TSharedRef<SWidget> CreateToolbarWidget();
+	FReply ApplyCurveTemplate(TWeakObjectPtr<UCurveFloat> WeakCurveAsset);
 
 private:
-	TSharedPtr<SNiagaraCompactCurveKeySelector> KeySelectorWidget;
-	TSharedPtr<INiagaraDistributionAdapter> DistributionAdapter;
-	int32 ChannelIndex;
-
-	FRichCurve EditCurve;
-	FKeyHandle SelectedKeyHandle;
-
-	int32 SelectedKeyIndex = INDEX_NONE;
-	int32 CurveStateSerialNumber = 0;
-
-	mutable int32 CacheCurveStateSerialNumber = INDEX_NONE;
-	mutable TOptional<float> KeyTimeCache;
-	mutable TOptional<float> KeyValueCache;
-	mutable int32 KeyInterpModeCache;
-	mutable FText CurveTimeMinTextCache;
-	mutable FText CurveTimeMaxTextCache;
-	mutable FText CurveValueMinTextCache;
-	mutable FText CurveValueMaxTextCache;
-};
\ No newline at end of file
+	TSharedPtr<INiagaraDistributionAdapter>				DistributionAdapter;
+	TSharedPtr<FNiagaraDistributionCurveEditorOptions>	CurveEditorOptions;
+	TSharedPtr<FCurveEditor>							CurveEditor;
+	TSharedPtr<SCurveEditorPanel>						CurveEditorPanel;
+	TSharedPtr<SCurveEditorTree>						CurveEditorTree;
+};
diff --git a/Engine/Plugins/FX/Niagara/Source/NiagaraEditor/Private/Widgets/SNiagaraDistributionEditor.cpp b/Engine/Plugins/FX/Niagara/Source/NiagaraEditor/Private/Widgets/SNiagaraDistributionEditor.cpp
--- a/Engine/Plugins/FX/Niagara/Source/NiagaraEditor/Private/Widgets/SNiagaraDistributionEditor.cpp
+++ b/Engine/Plugins/FX/Niagara/Source/NiagaraEditor/Private/Widgets/SNiagaraDistributionEditor.cpp
@@ -18,6 +18,7 @@
 #include "Widgets/SNiagaraColorEditor.h"
 #include "Widgets/SNiagaraDistributionArrayEditor.h"
 #include "Widgets/SNiagaraDistributionCurveEditor.h"
+#include "Widgets/SNiagaraDistributionCompactCurveEditor.h"
 #include "Widgets/SNiagaraDistributionLookupValueWidget.h"
 #include "Widgets/SNiagaraParameterName.h"
 
@@ -25,6 +26,16 @@
 
 #define LOCTEXT_NAMESPACE "NiagaraDistributionEditor"
 
+namespace SNiagaraDistributionEditorPrivate
+{
+	static bool GUseCompactEditor = false;
+	static FAutoConsoleVariableRef CVarUseCompactEditor(
+		TEXT("fx.Niagara.UI.Distribution.UseCompactEditor"),
+		GUseCompactEditor,
+		TEXT("When enabled we use the compact editor vs the curve editor.")
+	);
+}
+
 class SNiagaraDistributionModeSelector : public SComboButton
 {
 public:
@@ -750,15 +761,28 @@ TSharedRef<SWidget> SNiagaraDistributionEditor::ConstructContentForMode()
 	else if (FNiagaraDistributionEditorUtilities::IsCurve(Mode))
 	{
 		TSharedRef<SVerticalBox> CurveBox = SNew(SVerticalBox);
-		int32 NumberOfChannels = FNiagaraDistributionEditorUtilities::IsUniform(Mode) ? 1 : DistributionAdapter->GetNumChannels();
-		for (int32 ChannelIndex = 0; ChannelIndex < NumberOfChannels; ChannelIndex++)
+
+		if (SNiagaraDistributionEditorPrivate::GUseCompactEditor)
+		{
+			int32 NumberOfChannels = FNiagaraDistributionEditorUtilities::IsUniform(Mode) ? 1 : DistributionAdapter->GetNumChannels();
+			for (int32 ChannelIndex = 0; ChannelIndex < NumberOfChannels; ChannelIndex++)
+			{
+				CurveBox->AddSlot()
+				.AutoHeight()
+				.Padding(0, 0, 0, 3)
+				[
+					SNew(SNiagaraDistributionCompactCurveEditor, DistributionAdapter.ToSharedRef(), ChannelIndex)
+					.CurveColor(DistributionAdapter->GetChannelColor(ChannelIndex))
+				];
+			}
+		}
+		else
 		{
 			CurveBox->AddSlot()
 			.AutoHeight()
 			.Padding(0, 0, 0, 3)
 			[
-				SNew(SNiagaraDistributionCurveEditor, DistributionAdapter.ToSharedRef(), ChannelIndex)
-				.CurveColor(DistributionAdapter->GetChannelColor(ChannelIndex))
+				SNew(SNiagaraDistributionCurveEditor, DistributionAdapter.ToSharedRef())
 			];
 		}
 
@@ -780,4 +804,4 @@ TSharedRef<SWidget> SNiagaraDistributionEditor::ConstructContentForMode()
 	}
 }
 
-#undef LOCTEXT_NAMESPACE
\ No newline at end of file
+#undef LOCTEXT_NAMESPACE
diff --git a/Engine/Plugins/FX/Niagara/Source/NiagaraEditor/Public/NiagaraEditorModule.h b/Engine/Plugins/FX/Niagara/Source/NiagaraEditor/Public/NiagaraEditorModule.h
--- a/Engine/Plugins/FX/Niagara/Source/NiagaraEditor/Public/NiagaraEditorModule.h
+++ b/Engine/Plugins/FX/Niagara/Source/NiagaraEditor/Public/NiagaraEditorModule.h
@@ -27,6 +27,10 @@ namespace UE::Niagara::Wizard
 	class FModuleWizardGenerator;
 }
 
+class FCurveEditor;
+struct FCurveEditorTreeItemID;
+class SCurveEditorTree;
+struct FRichCurve;
 class FNiagaraRecentAndFavoritesManager;
 class IAssetTools;
 class IAssetTypeActions;
@@ -88,6 +92,9 @@ public:
 
 	virtual TSharedRef<SWidget> CreateScriptScratchPadManager(UNiagaraScratchPadViewModel& ScriptScratchPadViewModel) const = 0;
 	virtual TSharedRef<SWidget> CreateCurveOverview(TSharedRef<FNiagaraSystemViewModel> SystemViewModel) const = 0;
+	virtual TSharedRef<SWidget> CreateCurveKeySelector(TSharedPtr<FCurveEditor> InCurveEditor, const TArray<FCurveEditorTreeItemID>& InCurveTreeItemIds, TSharedPtr<SCurveEditorTree> InCurveEditorTree) const = 0;
+	virtual TSharedRef<SWidget> CreateCurveThumbnail(const FRichCurve& Curve) const = 0;
+
 	virtual FLinearColor GetColorForExecutionCategory(FName ExecutionCategory) const = 0;
 };
 
diff --git a/Engine/Plugins/FX/Niagara/Source/NiagaraEditorWidgets/Private/DetailCustomizations/NiagaraDataInterfaceCurveDetails.cpp b/Engine/Plugins/FX/Niagara/Source/NiagaraEditorWidgets/Private/DetailCustomizations/NiagaraDataInterfaceCurveDetails.cpp
--- a/Engine/Plugins/FX/Niagara/Source/NiagaraEditorWidgets/Private/DetailCustomizations/NiagaraDataInterfaceCurveDetails.cpp
+++ b/Engine/Plugins/FX/Niagara/Source/NiagaraEditorWidgets/Private/DetailCustomizations/NiagaraDataInterfaceCurveDetails.cpp
@@ -13,6 +13,8 @@
 #include "NiagaraEditorUtilities.h"
 #include "ViewModels/NiagaraSystemViewModel.h"
 #include "ViewModels/NiagaraCurveSelectionViewModel.h"
+#include "SNiagaraCurveKeySelector.h"
+#include "SNiagaraCurveThumbnail.h"
 
 #include "PropertyHandle.h"
 #include "DetailLayoutBuilder.h"
@@ -179,636 +181,6 @@ FRichCurve* GetCurveFromPropertyHandle(TSharedPtr<IPropertyHandle> Handle)
 	return RawData.Num() == 1 ? static_cast<FRichCurve*>(RawData[0]) : nullptr;
 }
 
-class SNiagaraDataInterfaceCurveKeySelector : public SCompoundWidget
-{
-public:
-	SLATE_BEGIN_ARGS(SNiagaraDataInterfaceCurveKeySelector) {}
-	SLATE_END_ARGS()
-
-	void Construct(const FArguments& InArgs, TSharedPtr<FCurveEditor> InCurveEditor, const TArray<FCurveEditorTreeItemID>& InCurveTreeItemIds, TSharedPtr<SCurveEditorTree> InCurveEditorTree)
-	{
-		CurveEditor = InCurveEditor;
-		OrderedCurveTreeItemIds = InCurveTreeItemIds;
-		CurveEditorTree = InCurveEditorTree;
-		ChildSlot
-		[
-			SNew(SHorizontalBox)
-
-			// Zoom to fit button
-			+ SHorizontalBox::Slot()
-			.AutoWidth()
-			.Padding(0, 0, 4, 0)
-			[
-				SNew(SButton)
-				.ButtonStyle(FAppStyle::Get(), "SimpleButton")
-				.OnClicked(this, &SNiagaraDataInterfaceCurveKeySelector::ZoomToFitClicked)
-				.ToolTipText(LOCTEXT("ZoomToFitToolTip", "Zoom to fit all keys"))
-				.Content()
-				[
-					SNew(STextBlock)
-					.Font(FAppStyle::Get().GetFontStyle("FontAwesome.10"))
-					.ColorAndOpacity(FSlateColor::UseForeground())
-					.Text(FEditorFontGlyphs::Expand)
-				]
-			]
-
-			// Previous curve button
-			+ SHorizontalBox::Slot()
-			.AutoWidth()
-			.Padding(0, 0, 1, 0)
-			[
-				SNew(SButton)
-				.ButtonStyle(FAppStyle::Get(), "SimpleButton")
-				.OnClicked(this, &SNiagaraDataInterfaceCurveKeySelector::PreviousCurveClicked)
-				.ToolTipText(LOCTEXT("PreviousCurveToolTip", "Select the previous curve"))
-				.Visibility(this, &SNiagaraDataInterfaceCurveKeySelector::GetNextPreviousCurveButtonVisibility)
-				.Content()
-				[
-					SNew(SImage)
-					.ColorAndOpacity(FSlateColor::UseForeground())
-					.Image(FAppStyle::Get().GetBrush("Icons.ArrowUp"))
-				]
-			]
-
-			// Next curve button
-			+ SHorizontalBox::Slot()
-			.AutoWidth()
-			.Padding(0, 0, 4, 0)
-			[
-				SNew(SButton)
-				.ButtonStyle(FAppStyle::Get(), "SimpleButton")
-				.OnClicked(this, &SNiagaraDataInterfaceCurveKeySelector::NextCurveClicked)
-				.ToolTipText(LOCTEXT("NextCurveToolTip", "Select the next curve"))
-				.Visibility(this, &SNiagaraDataInterfaceCurveKeySelector::GetNextPreviousCurveButtonVisibility)
-				.Content()
-				[
-					SNew(SImage)
-					.Image(FAppStyle::Get().GetBrush("Icons.ArrowDown"))
-					.ColorAndOpacity(FSlateColor::UseForeground())
-				]
-			]
-
-			// Previous key button
-			+ SHorizontalBox::Slot()
-			.AutoWidth()
-			.Padding(0, 0, 1, 0)
-			[
-				SNew(SButton)
-				.ButtonStyle(FAppStyle::Get(), "SimpleButton")
-				.OnClicked(this, &SNiagaraDataInterfaceCurveKeySelector::PreviousKeyClicked)
-				.ToolTipText(LOCTEXT("PreviousKeyToolTip", "Select the previous key for the selected curve."))
-				.Content()
-				[
-					SNew(SImage)
-					.Image(FAppStyle::Get().GetBrush("Icons.ArrowLeft"))
-					.ColorAndOpacity(FSlateColor::UseForeground())
-				]
-			]
-
-			// Next key button
-			+ SHorizontalBox::Slot()
-			.AutoWidth()
-			.Padding(0, 0, 4, 0)
-			[
-				SNew(SButton)
-				.ButtonStyle(FAppStyle::Get(), "SimpleButton")
-				.OnClicked(this, &SNiagaraDataInterfaceCurveKeySelector::NextKeyClicked)
-				.ToolTipText(LOCTEXT("NextKeyToolTip", "Select the next key for the selected curve."))
-				.Content()
-				[
-					SNew(SImage)
-					.Image(FAppStyle::Get().GetBrush("Icons.ArrowRight"))
-					.ColorAndOpacity(FSlateColor::UseForeground())
-				]
-			]
-
-			// Add key button
-			+ SHorizontalBox::Slot()
-			.AutoWidth()
-			.Padding(0, 0, 1, 0)
-			[
-				SNew(SButton)
-				.ButtonStyle(FAppStyle::Get(), "SimpleButton")
-				.OnClicked(this, &SNiagaraDataInterfaceCurveKeySelector::AddKeyClicked)
-				.ToolTipText(LOCTEXT("AddKeyToolTip", "Add a key to the selected curve."))
-				.Content()
-				[
-					SNew(SImage)
-					.Image(FAppStyle::Get().GetBrush("Icons.PlusCircle"))
-					.ColorAndOpacity(FSlateColor::UseForeground())
-				]
-			]
-
-			// Delete key button
-			+ SHorizontalBox::Slot()
-			.AutoWidth()
-			[
-				SNew(SButton)
-				.ButtonStyle(FAppStyle::Get(), "SimpleButton")
-				.OnClicked(this, &SNiagaraDataInterfaceCurveKeySelector::DeleteKeyClicked)
-				.ToolTipText(LOCTEXT("DeleteKeyToolTip", "Delete the currently selected keys."))
-				.Content()
-				[
-					SNew(SImage)
-					.Image(FAppStyle::Get().GetBrush("Icons.Delete"))
-					.ColorAndOpacity(FSlateColor::UseForeground())
-				]
-			]
-		];
-	}
-
-private:
-	void GetActiveCurveModelAndSelectedKeys(TOptional<FCurveModelID>& OutActiveCurveModelId, TArray<FKeyHandle>& OutSelectedKeyHandles)
-	{
-		const TMap<FCurveEditorTreeItemID, ECurveEditorTreeSelectionState>& CurveTreeSelection = CurveEditor->GetTree()->GetSelection();
-		if (CurveTreeSelection.Num() != 0)
-		{
-			// If there are curves selected in the tree use those first.
-			FCurveEditorTreeItemID FirstSelectedCurveTreeItemId;
-			for (int32 i = OrderedCurveTreeItemIds.Num() - 1; i >= 0; i--)
-			{
-				const FCurveEditorTreeItemID& CurveTreeItemId = OrderedCurveTreeItemIds[i];
-				if (CurveTreeSelection.Contains(CurveTreeItemId))
-				{
-					FirstSelectedCurveTreeItemId = CurveTreeItemId;
-					break;
-				}
-			}
-
-			TArrayView<const FCurveModelID> CurveModelIds = CurveEditor->GetTreeItem(FirstSelectedCurveTreeItemId).GetOrCreateCurves(CurveEditor.Get());
-			if (CurveModelIds.Num() == 1)
-			{
-				OutActiveCurveModelId = CurveModelIds[0];
-
-				const FKeyHandleSet* SelectedKeyHandleSet = CurveEditor->GetSelection().GetAll().Find(OutActiveCurveModelId.GetValue());
-				if (SelectedKeyHandleSet != nullptr)
-				{
-					OutSelectedKeyHandles = SelectedKeyHandleSet->AsArray();
-				}
-			}
-		}
-		else
-		{
-			// Otherwise check if there are keys selected and if so use first curve with selected keys.
-			FCurveEditorSelection& CurveEditorSelection = CurveEditor->GetSelection();
-			if (CurveEditorSelection.IsEmpty() == false)
-			{
-				for (int32 i = OrderedCurveTreeItemIds.Num() - 1; i >= 0; i--)
-				{
-					const FCurveEditorTreeItemID& CurveTreeItemId = OrderedCurveTreeItemIds[i];
-					TArrayView<const FCurveModelID> CurveModelIds = CurveEditor->GetTreeItem(CurveTreeItemId).GetCurves();
-					if (CurveModelIds.Num() == 1)
-					{
-						const FKeyHandleSet* SelectedKeyHandleSet = CurveEditorSelection.GetAll().Find(CurveModelIds[0]);
-						if (SelectedKeyHandleSet != nullptr)
-						{
-							OutActiveCurveModelId = CurveModelIds[0];
-							OutSelectedKeyHandles = SelectedKeyHandleSet->AsArray();
-							break;
-						}
-					}
-				}
-			}
-			else
-			{
-				// Otherwise just use the first pinned curve.
-				const TSet<FCurveModelID>& PinnedCurveIds = CurveEditor->GetPinnedCurves();
-				if (PinnedCurveIds.Num() > 0)
-				{
-					for (FCurveEditorTreeItemID& CurveTreeItemId : OrderedCurveTreeItemIds)
-					{
-						TArrayView<const FCurveModelID> CurveModelIds = CurveEditor->GetTreeItem(CurveTreeItemId).GetOrCreateCurves(CurveEditor.Get());
-						if (CurveModelIds.Num() == 1 && PinnedCurveIds.Contains(CurveModelIds[0]))
-						{
-							OutActiveCurveModelId = CurveModelIds[0];
-							break;
-						}
-					}
-				}
-			}
-		}
-	}
-
-	struct FKeyHandlePositionPair
-	{
-		FKeyHandle Handle;
-		FKeyPosition Position;
-	};
-
-	void GetSortedKeyHandlessAndPositionsForModel(FCurveModel& InCurveModel, TArray<FKeyHandlePositionPair>& OutSortedKeyHandlesAndPositions)
-	{
-		const TArray<FKeyHandle> KeyHandles = InCurveModel.GetAllKeys();
-
-		TArray<FKeyPosition> KeyPositions;
-		KeyPositions.AddDefaulted(KeyHandles.Num());
-		InCurveModel.GetKeyPositions(KeyHandles, KeyPositions);
-
-		for (int32 i = 0; i < KeyHandles.Num(); i++)
-		{
-			FKeyHandlePositionPair& KeyHandlePositionPair = OutSortedKeyHandlesAndPositions.AddDefaulted_GetRef();
-			KeyHandlePositionPair.Handle = KeyHandles[i];
-			KeyHandlePositionPair.Position = KeyPositions[i];
-		}
-
-		OutSortedKeyHandlesAndPositions.Sort([](const FKeyHandlePositionPair& A, const FKeyHandlePositionPair& B) { return A.Position.InputValue < B.Position.InputValue; });
-	}
-
-	void GetOrderedActiveCurveModelIds(TArray<FCurveModelID>& OutOrderedActiveCurveModelIds)
-	{
-		if (CurveEditor->GetTreeSelection().Num() > 0)
-		{
-			// If there are curves selected in the tree then those are the only active ones.
-			for (const FCurveEditorTreeItemID& OrderedCurveTreeItemId : OrderedCurveTreeItemIds)
-			{
-				if(CurveEditor->GetTreeSelectionState(OrderedCurveTreeItemId) != ECurveEditorTreeSelectionState::None)
-				{
-					TArrayView<const FCurveModelID> CurveModelIds = CurveEditor->GetTreeItem(OrderedCurveTreeItemId).GetOrCreateCurves(CurveEditor.Get());
-					if (CurveModelIds.Num() == 1)
-					{
-						OutOrderedActiveCurveModelIds.Add(CurveModelIds[0]);
-					}
-				}
-			}
-		}
-		else
-		{
-			// Otherwise the active curves are the pinned curves.
-			for (const FCurveEditorTreeItemID& OrderedCurveTreeItemId : OrderedCurveTreeItemIds)
-			{
-				TArrayView<const FCurveModelID> CurveModelIds = CurveEditor->GetTreeItem(OrderedCurveTreeItemId).GetCurves();
-				if (CurveModelIds.Num() == 1 && CurveEditor->IsCurvePinned(CurveModelIds[0]))
-				{
-					OutOrderedActiveCurveModelIds.Add(CurveModelIds[0]);
-				}
-			}
-		}
-	}
-
-	enum class ENavigateDirection
-	{
-		Previous,
-		Next
-	};
-
-	void NavigateToAdjacentCurve(ENavigateDirection Direction)
-	{
-		if (CurveEditorTree.IsValid())
-		{
-			if (CurveEditor->GetTreeSelection().Num() == 0)
-			{
-				CurveEditorTree->SetItemSelection(OrderedCurveTreeItemIds[0], true);
-			}
-			else
-			{
-				int32 TargetSelectedTreeItemIndex = INDEX_NONE;
-				int32 StartIndex = Direction == ENavigateDirection::Previous ? 0 : OrderedCurveTreeItemIds.Num() - 1;
-				int32 EndIndex = Direction == ENavigateDirection::Previous ? OrderedCurveTreeItemIds.Num() : -1;
-				int32 IndexOffset = Direction == ENavigateDirection::Previous ? 1 : -1;
-				for (int32 i = StartIndex; i != EndIndex; i += IndexOffset)
-				{
-					if (CurveEditor->GetTreeSelectionState(OrderedCurveTreeItemIds[i]) != ECurveEditorTreeSelectionState::None)
-					{
-						TargetSelectedTreeItemIndex = i;
-						break;
-					}
-				}
-
-				if (TargetSelectedTreeItemIndex != INDEX_NONE)
-				{
-					int32 IndexToSelect = Direction == ENavigateDirection::Previous 
-						? TargetSelectedTreeItemIndex - 1
-						: TargetSelectedTreeItemIndex + 1;
-
-					if (IndexToSelect > OrderedCurveTreeItemIds.Num() - 1)
-					{
-						IndexToSelect = 0;
-					}
-					else if(IndexToSelect < 0)
-					{
-						IndexToSelect = OrderedCurveTreeItemIds.Num() - 1;
-					}
-
-					CurveEditorTree->ClearSelection();
-					CurveEditorTree->SetItemSelection(OrderedCurveTreeItemIds[IndexToSelect], true);
-				}
-			}
-		}
-	}
-
-	void NavigateToAdjacentKey(ENavigateDirection Direction)
-	{
-		TOptional<FCurveModelID> ActiveCurveModelId;
-		TArray<FKeyHandle> SelectedKeyHandles;
-		GetActiveCurveModelAndSelectedKeys(ActiveCurveModelId, SelectedKeyHandles);
-
-		TOptional<FCurveModelID> CurveModelIdToSelect;
-		TOptional<FKeyHandle> KeyHandleToSelect;
-		if (ActiveCurveModelId.IsSet())
-		{
-			FCurveModel* ActiveCurveModel = CurveEditor->GetCurves()[ActiveCurveModelId.GetValue()].Get();
-			TArray<FKeyHandlePositionPair> ActiveSortedKeyHandlePositionPairs;
-			GetSortedKeyHandlessAndPositionsForModel(*ActiveCurveModel, ActiveSortedKeyHandlePositionPairs);
-
-			if (SelectedKeyHandles.Num() != 0)
-			{
-				// If there's currently a selected key on the active curve then we want to use that as the target for navigating.
-				int32 TargetSelectedKeyIndex = INDEX_NONE;
-				int32 StartIndex = Direction == ENavigateDirection::Previous ? 0 : ActiveSortedKeyHandlePositionPairs.Num() - 1;
-				int32 EndIndex = Direction == ENavigateDirection::Previous ? ActiveSortedKeyHandlePositionPairs.Num() : -1;
-				int32 IndexOffset = Direction == ENavigateDirection::Previous ? 1 : -1;
-				for (int32 i = StartIndex; i != EndIndex; i += IndexOffset)
-				{
-					if (SelectedKeyHandles.Contains(ActiveSortedKeyHandlePositionPairs[i].Handle))
-					{
-						TargetSelectedKeyIndex = i;
-						break;
-					}
-				}
-
-				if (TargetSelectedKeyIndex != INDEX_NONE)
-				{
-					if(Direction == ENavigateDirection::Previous && TargetSelectedKeyIndex > 0)
-					{
-						// If we're navigating previous and we're not at the first key we can just select the previous key on this curve.
-						CurveModelIdToSelect = ActiveCurveModelId;
-						KeyHandleToSelect = ActiveSortedKeyHandlePositionPairs[TargetSelectedKeyIndex - 1].Handle;
-					}
-					else if (Direction == ENavigateDirection::Next && TargetSelectedKeyIndex < ActiveSortedKeyHandlePositionPairs.Num() - 1)
-					{
-						// If we're navigating next and we're not at the last key we can just select the next key on this curve.
-						CurveModelIdToSelect = ActiveCurveModelId;
-						KeyHandleToSelect = ActiveSortedKeyHandlePositionPairs[TargetSelectedKeyIndex + 1].Handle;
-					}
-					else
-					{
-						// Otherwise we're going to need to navigate to another curve since we're at the start or end of the current curve.
-						TArray<FCurveModelID> OrderedActiveCurveModelIds;
-						GetOrderedActiveCurveModelIds(OrderedActiveCurveModelIds);
-
-						// Find the adjacent curve with keys so we can select a key there.
-						FCurveModel* CurveModelToSelect = nullptr;
-						int32 CurrentIndex = OrderedActiveCurveModelIds.IndexOfByKey(ActiveCurveModelId);
-						int32 CurrentIndexOffset = Direction == ENavigateDirection::Previous ? -1 : 1;
-						while (CurveModelIdToSelect.IsSet() == false)
-						{
-							CurrentIndex += CurrentIndexOffset;
-							if (CurrentIndex < 0)
-							{
-								CurrentIndex = OrderedActiveCurveModelIds.Num() - 1;
-							}
-							else if (CurrentIndex > OrderedActiveCurveModelIds.Num() - 1)
-							{
-								CurrentIndex = 0;
-							}
-
-							FCurveModel* CurrentCurveModel = CurveEditor->GetCurves()[OrderedActiveCurveModelIds[CurrentIndex]].Get();
-							if (CurrentCurveModel->GetNumKeys() > 0)
-							{
-								CurveModelIdToSelect = OrderedActiveCurveModelIds[CurrentIndex];
-								CurveModelToSelect = CurrentCurveModel;
-							}
-						}
-
-						if (CurveModelIdToSelect == ActiveCurveModelId)
-						{
-							// There were no other active curves with keys to select from so just wrap around the currently active curve.
-							if (Direction == ENavigateDirection::Previous)
-							{
-								KeyHandleToSelect = ActiveSortedKeyHandlePositionPairs.Last().Handle;
-							}
-							else
-							{
-								KeyHandleToSelect = ActiveSortedKeyHandlePositionPairs[0].Handle;
-							}
-						}
-						else if(CurveModelToSelect != nullptr)
-						{
-							// We're selecting a key on a different curve so we need to sort the positions and select the first or last based
-							// on the navigation direction.
-							TArray<FKeyHandlePositionPair> SortedKeyHandlePositionPairs;
-							GetSortedKeyHandlessAndPositionsForModel(*CurveModelToSelect, SortedKeyHandlePositionPairs);
-							if (Direction == ENavigateDirection::Previous)
-							{
-								KeyHandleToSelect = SortedKeyHandlePositionPairs.Last().Handle;
-							}
-							else
-							{
-								KeyHandleToSelect = SortedKeyHandlePositionPairs[0].Handle;
-							}
-						}
-					}
-				}
-			}
-			else if(ActiveSortedKeyHandlePositionPairs.Num() > 0)
-			{
-				// There weren't any keys already selected so just select the first key on the active curve.
-				CurveModelIdToSelect = ActiveCurveModelId;
-				KeyHandleToSelect = ActiveSortedKeyHandlePositionPairs[0].Handle;
-			}
-		}
-
-		if (CurveModelIdToSelect.IsSet() && KeyHandleToSelect.IsSet())
-		{
-			FCurveEditorSelection& CurveEditorSelection = CurveEditor->GetSelection();
-			CurveEditorSelection.Clear();
-			CurveEditorSelection.Add(CurveModelIdToSelect.GetValue(), ECurvePointType::Key, KeyHandleToSelect.GetValue());
-
-			TArray<FCurveModelID> CurvesToFit;
-			CurvesToFit.Add(CurveModelIdToSelect.GetValue());
-			CurveEditor->ZoomToFitCurves(CurvesToFit);
-		}
-	}
-
-	FReply ZoomToFitClicked()
-	{
-		TArray<FCurveModelID> ActiveCurveModelIds;
-		GetOrderedActiveCurveModelIds(ActiveCurveModelIds);
-		CurveEditor->ZoomToFitCurves(ActiveCurveModelIds);
-		return FReply::Handled();
-	}
-
-	FReply PreviousCurveClicked()
-	{
-		NavigateToAdjacentCurve(ENavigateDirection::Previous);
-		return FReply::Handled();
-	}
-
-	FReply NextCurveClicked()
-	{
-		NavigateToAdjacentCurve(ENavigateDirection::Next);
-		return FReply::Handled();
-	}
-
-	FReply PreviousKeyClicked()
-	{
-		NavigateToAdjacentKey(ENavigateDirection::Previous);
-		return FReply::Handled();
-	}
-
-	FReply NextKeyClicked()
-	{
-		NavigateToAdjacentKey(ENavigateDirection::Next);
-		return FReply::Handled();
-	}
-
-	EVisibility GetNextPreviousCurveButtonVisibility() const
-	{
-		return OrderedCurveTreeItemIds.Num() == 1 ? EVisibility::Collapsed : EVisibility::Visible;
-	}
-
-	FReply AddKeyClicked()
-	{
-		TOptional<FCurveModelID> CurveModelIdForAdd;
-		TArray<FKeyHandle> SelectedKeyHandles;
-		GetActiveCurveModelAndSelectedKeys(CurveModelIdForAdd, SelectedKeyHandles);
-
-		if(CurveModelIdForAdd.IsSet())
-		{
-			FCurveModel* CurveModelForAdd = CurveEditor->GetCurves()[CurveModelIdForAdd.GetValue()].Get();
-
-			FKeyPosition NewKeyPosition;
-			FKeyAttributes NewKeyAttributes = CurveEditor->GetDefaultKeyAttributes().Get();
-			NewKeyAttributes.SetInterpMode(RCIM_Cubic);
-			NewKeyAttributes.SetTangentMode(RCTM_Auto);
-			if (CurveModelForAdd->GetNumKeys() == 0)
-			{
-				// If there are no keys, add one at 0, 0.
-				NewKeyPosition.InputValue = 0.0f;
-				NewKeyPosition.OutputValue = 0.0f;
-			}
-			else if (CurveModelForAdd->GetNumKeys() == 1)
-			{
-				// If there's a single key, add the new key at the same value, but time + 1.
-				const TArray<FKeyHandle> KeyHandles = CurveModelForAdd->GetAllKeys();
-				
-				TArray<FKeyPosition> KeyPositions;
-				KeyPositions.AddDefaulted();
-				CurveModelForAdd->GetKeyPositions(KeyHandles, KeyPositions);
-						
-				NewKeyPosition.InputValue = KeyPositions[0].InputValue + 1;
-				NewKeyPosition.OutputValue = KeyPositions[0].OutputValue;
-			}
-			else
-			{
-				TArray<FKeyHandlePositionPair> SortedKeyHandlePositionPairs;
-				GetSortedKeyHandlessAndPositionsForModel(*CurveModelForAdd, SortedKeyHandlePositionPairs);
-
-				int32 IndexToAddAfter = INDEX_NONE;
-				if (SelectedKeyHandles.Num() > 0)
-				{
-					for (int32 i = SortedKeyHandlePositionPairs.Num() - 1; i >= 0; i--)
-					{
-						if (SelectedKeyHandles.Contains(SortedKeyHandlePositionPairs[i].Handle))
-						{
-							IndexToAddAfter = i;
-							break;
-						}
-					}
-				}
-
-				if(IndexToAddAfter == INDEX_NONE)
-				{
-					IndexToAddAfter = SortedKeyHandlePositionPairs.Num() - 1;
-				}
-
-				if (IndexToAddAfter == SortedKeyHandlePositionPairs.Num() - 1)
-				{
-					const FKeyPosition& TargetKeyPosition = SortedKeyHandlePositionPairs[IndexToAddAfter].Position;
-					const FKeyPosition& PreviousKeyPosition = SortedKeyHandlePositionPairs[IndexToAddAfter - 1].Position;
-					NewKeyPosition.InputValue = TargetKeyPosition.InputValue + (TargetKeyPosition.InputValue - PreviousKeyPosition.InputValue);
-					NewKeyPosition.OutputValue = TargetKeyPosition.OutputValue + (TargetKeyPosition.OutputValue - PreviousKeyPosition.OutputValue);
-				}
-				else
-				{
-					const FKeyPosition& TargetKeyPosition = SortedKeyHandlePositionPairs[IndexToAddAfter].Position;
-					const FKeyPosition& NextKeyPosition = SortedKeyHandlePositionPairs[IndexToAddAfter + 1].Position;
-					NewKeyPosition.InputValue = TargetKeyPosition.InputValue + ((NextKeyPosition.InputValue - TargetKeyPosition.InputValue) / 2);
-					NewKeyPosition.OutputValue = TargetKeyPosition.OutputValue + ((NextKeyPosition.OutputValue - TargetKeyPosition.OutputValue) / 2);
-				}
-			}
-
-			FScopedTransaction Transaction(LOCTEXT("AddKey", "Add Key"));
-			TOptional<FKeyHandle> NewKeyHandle = CurveModelForAdd->AddKey(NewKeyPosition, NewKeyAttributes);
-			if(NewKeyHandle.IsSet())
-			{
-				FCurveEditorSelection& CurveEditorSelection = CurveEditor->GetSelection();
-				CurveEditorSelection.Clear();
-				CurveEditorSelection.Add(CurveModelIdForAdd.GetValue(), ECurvePointType::Key, NewKeyHandle.GetValue());
-
-				TArray<FCurveModelID> CurvesToFit;
-				CurvesToFit.Add(CurveModelIdForAdd.GetValue());
-				CurveEditor->ZoomToFitCurves(CurvesToFit);
-			}
-		}
-		return FReply::Handled();
-	}
-
-	FReply DeleteKeyClicked()
-	{
-		CurveEditor->DeleteSelection();
-		return FReply::Handled();
-	}
-
-private:
-	TSharedPtr<FCurveEditor> CurveEditor;
-	TArray<FCurveEditorTreeItemID> OrderedCurveTreeItemIds;
-	TSharedPtr<SCurveEditorTree> CurveEditorTree;
-};
-
-class SNiagaraCurveThumbnail : public SLeafWidget
-{
-	SLATE_BEGIN_ARGS(SNiagaraCurveThumbnail)
-		: _Width(16)
-		, _Height(8) 
-		{}
-		SLATE_ARGUMENT(float, Width)
-		SLATE_ARGUMENT(float, Height)
-	SLATE_END_ARGS()
-
-	void Construct(const FArguments& InArgs, const FRichCurve* CurveToDisplay)
-	{
-		Width = InArgs._Width;
-		Height = InArgs._Height;
-
-		float TimeMin;
-		float TimeMax;
-		float ValueMin;
-		float ValueMax;
-		CurveToDisplay->GetTimeRange(TimeMin, TimeMax);
-		CurveToDisplay->GetValueRange(ValueMin, ValueMax);
-
-		float TimeRange = TimeMax - TimeMin;
-		float ValueRange = ValueMax - ValueMin;
-
-		int32 Points = 13;
-		float TimeIncrement = TimeRange / (Points - 1);
-		for (int32 i = 0; i < Points; i++)
-		{
-			float Time = TimeMin + i * TimeIncrement;
-			float Value = CurveToDisplay->Eval(Time);
-
-			float NormalizedX = (Time - TimeMin) / TimeRange;
-			float NormalizedY = (Value - ValueMin) / ValueRange;
-			CurvePoints.Add(FVector2D(NormalizedX * Width, (1 - NormalizedY) * Height));
-		}
-	}
-
-	virtual int32 OnPaint(const FPaintArgs& Args, const FGeometry& AllottedGeometry, const FSlateRect& MyCullingRect, FSlateWindowElementList& OutDrawElements, int32 LayerId, const FWidgetStyle& InWidgetStyle, bool bParentEnabled) const override
-	{
-		FSlateDrawElement::MakeLines(OutDrawElements, LayerId, AllottedGeometry.ToPaintGeometry(), CurvePoints, ESlateDrawEffect::None, InWidgetStyle.GetForegroundColor(), true, 2.0f);
-		return LayerId;
-	}
-
-	virtual FVector2D ComputeDesiredSize(float) const override
-	{
-		return FVector2D(Width, Height);
-	}
-
-	TArray<FVector2D> CurvePoints;
-	float Width;
-	float Height;
-};
-
 class SNiagaraDataInterfaceCurveEditor : public SCompoundWidget
 {
 private:
@@ -923,7 +295,7 @@ public:
 				.AutoWidth()
 				.Padding(KeySelectorLeftPadding, 0, 8, 0)
 				[
-					SNew(SNiagaraDataInterfaceCurveKeySelector, CurveEditor, TreeItemIds, CurveEditorTree)
+					SNew(SNiagaraCurveKeySelector, CurveEditor, TreeItemIds, CurveEditorTree)
 				]
 				+ SHorizontalBox::Slot()
 				.VAlign(VAlign_Center)
@@ -1414,7 +786,7 @@ void FNiagaraDataInterfaceCurveDetailsBase::CustomizeDetails(IDetailLayoutBuilde
 				.Visibility_Lambda([this]() { return GetIsColorCurve() && IsGradientVisible() ? EVisibility::Collapsed : EVisibility::Visible; })
 				.Content()
 				[
-					SNew(SNiagaraCurveThumbnail, &FloatCurveAsset->FloatCurve)
+					SNew(SNiagaraCurveThumbnail, FloatCurveAsset->FloatCurve)
 				]);
 		}
 	}
diff --git a/Engine/Plugins/FX/Niagara/Source/NiagaraEditorWidgets/Private/NiagaraEditorWidgetsModule.cpp b/Engine/Plugins/FX/Niagara/Source/NiagaraEditorWidgets/Private/NiagaraEditorWidgetsModule.cpp
--- a/Engine/Plugins/FX/Niagara/Source/NiagaraEditorWidgets/Private/NiagaraEditorWidgetsModule.cpp
+++ b/Engine/Plugins/FX/Niagara/Source/NiagaraEditorWidgets/Private/NiagaraEditorWidgetsModule.cpp
@@ -24,6 +24,8 @@
 #include "Stack/SNiagaraStackIssueIcon.h"
 #include "SNiagaraScratchPadScriptManager.h"
 #include "SNiagaraCurveOverview.h"
+#include "SNiagaraCurveKeySelector.h"
+#include "SNiagaraCurveThumbnail.h"
 
 #include "Modules/ModuleManager.h"
 #include "PropertyEditorModule.h"
@@ -237,6 +239,16 @@ TSharedRef<SWidget> FNiagaraEditorWidgetsModule::FNiagaraEditorWidgetProvider::C
 	return SNew(SNiagaraCurveOverview, SystemViewModel);
 }
 
+TSharedRef<SWidget> FNiagaraEditorWidgetsModule::FNiagaraEditorWidgetProvider::CreateCurveKeySelector(TSharedPtr<FCurveEditor> InCurveEditor, const TArray<FCurveEditorTreeItemID>& InCurveTreeItemIds, TSharedPtr<SCurveEditorTree> InCurveEditorTree) const
+{
+	return SNew(SNiagaraCurveKeySelector, InCurveEditor, InCurveTreeItemIds, InCurveEditorTree);
+}
+
+TSharedRef<SWidget> FNiagaraEditorWidgetsModule::FNiagaraEditorWidgetProvider::CreateCurveThumbnail(const FRichCurve& Curve) const
+{
+	return SNew(SNiagaraCurveThumbnail, Curve);
+}
+
 FLinearColor FNiagaraEditorWidgetsModule::FNiagaraEditorWidgetProvider::GetColorForExecutionCategory(FName ExecutionCategory) const
 {
 	return FNiagaraEditorWidgetsStyle::Get().GetColor(FNiagaraStackEditorWidgetsUtilities::GetIconColorNameForExecutionCategory(ExecutionCategory));
diff --git a/Engine/Plugins/FX/Niagara/Source/NiagaraEditorWidgets/Private/SNiagaraCurveKeySelector.cpp b/Engine/Plugins/FX/Niagara/Source/NiagaraEditorWidgets/Private/SNiagaraCurveKeySelector.cpp
new file mode 100644
--- /dev/null
+++ b/Engine/Plugins/FX/Niagara/Source/NiagaraEditorWidgets/Private/SNiagaraCurveKeySelector.cpp
@@ -0,0 +1,567 @@
+// Copyright Epic Games, Inc. All Rights Reserved.
+
+#include "SNiagaraCurveKeySelector.h"
+
+#include "CurveEditor.h"
+#include "Tree/SCurveEditorTree.h"
+#include "Widgets/Input/SButton.h"
+#include "Widgets/Text/STextBlock.h"
+#include "EditorFontGlyphs.h"
+
+#define LOCTEXT_NAMESPACE "SNiagaraCurveKeySelector"
+
+void SNiagaraCurveKeySelector::Construct(const FArguments& InArgs, TSharedPtr<FCurveEditor> InCurveEditor, const TArray<FCurveEditorTreeItemID>& InCurveTreeItemIds, TSharedPtr<SCurveEditorTree> InCurveEditorTree)
+{
+	CurveEditor = InCurveEditor;
+	OrderedCurveTreeItemIds = InCurveTreeItemIds;
+	CurveEditorTree = InCurveEditorTree;
+	ChildSlot
+	[
+		SNew(SHorizontalBox)
+
+		// Zoom to fit button
+		+ SHorizontalBox::Slot()
+		.AutoWidth()
+		.Padding(0, 0, 4, 0)
+		[
+			SNew(SButton)
+			.ButtonStyle(FAppStyle::Get(), "SimpleButton")
+			.OnClicked(this, &SNiagaraCurveKeySelector::ZoomToFitClicked)
+			.ToolTipText(LOCTEXT("ZoomToFitToolTip", "Zoom to fit all keys"))
+			.Content()
+			[
+				SNew(STextBlock)
+				.Font(FAppStyle::Get().GetFontStyle("FontAwesome.10"))
+				.ColorAndOpacity(FSlateColor::UseForeground())
+				.Text(FEditorFontGlyphs::Expand)
+			]
+		]
+
+		// Previous curve button
+		+ SHorizontalBox::Slot()
+		.AutoWidth()
+		.Padding(0, 0, 1, 0)
+		[
+			SNew(SButton)
+			.ButtonStyle(FAppStyle::Get(), "SimpleButton")
+			.OnClicked(this, &SNiagaraCurveKeySelector::PreviousCurveClicked)
+			.ToolTipText(LOCTEXT("PreviousCurveToolTip", "Select the previous curve"))
+			.Visibility(this, &SNiagaraCurveKeySelector::GetNextPreviousCurveButtonVisibility)
+			.Content()
+			[
+				SNew(SImage)
+				.ColorAndOpacity(FSlateColor::UseForeground())
+				.Image(FAppStyle::Get().GetBrush("Icons.ArrowUp"))
+			]
+		]
+
+		// Next curve button
+		+ SHorizontalBox::Slot()
+		.AutoWidth()
+		.Padding(0, 0, 4, 0)
+		[
+			SNew(SButton)
+			.ButtonStyle(FAppStyle::Get(), "SimpleButton")
+			.OnClicked(this, &SNiagaraCurveKeySelector::NextCurveClicked)
+			.ToolTipText(LOCTEXT("NextCurveToolTip", "Select the next curve"))
+			.Visibility(this, &SNiagaraCurveKeySelector::GetNextPreviousCurveButtonVisibility)
+			.Content()
+			[
+				SNew(SImage)
+				.Image(FAppStyle::Get().GetBrush("Icons.ArrowDown"))
+				.ColorAndOpacity(FSlateColor::UseForeground())
+			]
+		]
+
+		// Previous key button
+		+ SHorizontalBox::Slot()
+		.AutoWidth()
+		.Padding(0, 0, 1, 0)
+		[
+			SNew(SButton)
+			.ButtonStyle(FAppStyle::Get(), "SimpleButton")
+			.OnClicked(this, &SNiagaraCurveKeySelector::PreviousKeyClicked)
+			.ToolTipText(LOCTEXT("PreviousKeyToolTip", "Select the previous key for the selected curve."))
+			.Content()
+			[
+				SNew(SImage)
+				.Image(FAppStyle::Get().GetBrush("Icons.ArrowLeft"))
+				.ColorAndOpacity(FSlateColor::UseForeground())
+			]
+		]
+
+		// Next key button
+		+ SHorizontalBox::Slot()
+		.AutoWidth()
+		.Padding(0, 0, 4, 0)
+		[
+			SNew(SButton)
+			.ButtonStyle(FAppStyle::Get(), "SimpleButton")
+			.OnClicked(this, &SNiagaraCurveKeySelector::NextKeyClicked)
+			.ToolTipText(LOCTEXT("NextKeyToolTip", "Select the next key for the selected curve."))
+			.Content()
+			[
+				SNew(SImage)
+				.Image(FAppStyle::Get().GetBrush("Icons.ArrowRight"))
+				.ColorAndOpacity(FSlateColor::UseForeground())
+			]
+		]
+
+		// Add key button
+		+ SHorizontalBox::Slot()
+		.AutoWidth()
+		.Padding(0, 0, 1, 0)
+		[
+			SNew(SButton)
+			.ButtonStyle(FAppStyle::Get(), "SimpleButton")
+			.OnClicked(this, &SNiagaraCurveKeySelector::AddKeyClicked)
+			.ToolTipText(LOCTEXT("AddKeyToolTip", "Add a key to the selected curve."))
+			.Content()
+			[
+				SNew(SImage)
+				.Image(FAppStyle::Get().GetBrush("Icons.PlusCircle"))
+				.ColorAndOpacity(FSlateColor::UseForeground())
+			]
+		]
+
+		// Delete key button
+		+ SHorizontalBox::Slot()
+		.AutoWidth()
+		[
+			SNew(SButton)
+			.ButtonStyle(FAppStyle::Get(), "SimpleButton")
+			.OnClicked(this, &SNiagaraCurveKeySelector::DeleteKeyClicked)
+			.ToolTipText(LOCTEXT("DeleteKeyToolTip", "Delete the currently selected keys."))
+			.Content()
+			[
+				SNew(SImage)
+				.Image(FAppStyle::Get().GetBrush("Icons.Delete"))
+				.ColorAndOpacity(FSlateColor::UseForeground())
+			]
+		]
+	];
+}
+
+void SNiagaraCurveKeySelector::GetActiveCurveModelAndSelectedKeys(TOptional<FCurveModelID>& OutActiveCurveModelId, TArray<FKeyHandle>& OutSelectedKeyHandles)
+{
+	const TMap<FCurveEditorTreeItemID, ECurveEditorTreeSelectionState>& CurveTreeSelection = CurveEditor->GetTree()->GetSelection();
+	if (CurveTreeSelection.Num() != 0)
+	{
+		// If there are curves selected in the tree use those first.
+		FCurveEditorTreeItemID FirstSelectedCurveTreeItemId;
+		for (int32 i = OrderedCurveTreeItemIds.Num() - 1; i >= 0; i--)
+		{
+			const FCurveEditorTreeItemID& CurveTreeItemId = OrderedCurveTreeItemIds[i];
+			if (CurveTreeSelection.Contains(CurveTreeItemId))
+			{
+				FirstSelectedCurveTreeItemId = CurveTreeItemId;
+				break;
+			}
+		}
+
+		TArrayView<const FCurveModelID> CurveModelIds = CurveEditor->GetTreeItem(FirstSelectedCurveTreeItemId).GetOrCreateCurves(CurveEditor.Get());
+		if (CurveModelIds.Num() == 1)
+		{
+			OutActiveCurveModelId = CurveModelIds[0];
+
+			const FKeyHandleSet* SelectedKeyHandleSet = CurveEditor->GetSelection().GetAll().Find(OutActiveCurveModelId.GetValue());
+			if (SelectedKeyHandleSet != nullptr)
+			{
+				OutSelectedKeyHandles = SelectedKeyHandleSet->AsArray();
+			}
+		}
+	}
+	else
+	{
+		// Otherwise check if there are keys selected and if so use first curve with selected keys.
+		FCurveEditorSelection& CurveEditorSelection = CurveEditor->GetSelection();
+		if (CurveEditorSelection.IsEmpty() == false)
+		{
+			for (int32 i = OrderedCurveTreeItemIds.Num() - 1; i >= 0; i--)
+			{
+				const FCurveEditorTreeItemID& CurveTreeItemId = OrderedCurveTreeItemIds[i];
+				TArrayView<const FCurveModelID> CurveModelIds = CurveEditor->GetTreeItem(CurveTreeItemId).GetCurves();
+				if (CurveModelIds.Num() == 1)
+				{
+					const FKeyHandleSet* SelectedKeyHandleSet = CurveEditorSelection.GetAll().Find(CurveModelIds[0]);
+					if (SelectedKeyHandleSet != nullptr)
+					{
+						OutActiveCurveModelId = CurveModelIds[0];
+						OutSelectedKeyHandles = SelectedKeyHandleSet->AsArray();
+						break;
+					}
+				}
+			}
+		}
+		else
+		{
+			// Otherwise just use the first pinned curve.
+			const TSet<FCurveModelID>& PinnedCurveIds = CurveEditor->GetPinnedCurves();
+			if (PinnedCurveIds.Num() > 0)
+			{
+				for (FCurveEditorTreeItemID& CurveTreeItemId : OrderedCurveTreeItemIds)
+				{
+					TArrayView<const FCurveModelID> CurveModelIds = CurveEditor->GetTreeItem(CurveTreeItemId).GetOrCreateCurves(CurveEditor.Get());
+					if (CurveModelIds.Num() == 1 && PinnedCurveIds.Contains(CurveModelIds[0]))
+					{
+						OutActiveCurveModelId = CurveModelIds[0];
+						break;
+					}
+				}
+			}
+		}
+	}
+}
+
+void SNiagaraCurveKeySelector::GetSortedKeyHandlessAndPositionsForModel(FCurveModel& InCurveModel, TArray<FKeyHandlePositionPair>& OutSortedKeyHandlesAndPositions)
+{
+	const TArray<FKeyHandle> KeyHandles = InCurveModel.GetAllKeys();
+
+	TArray<FKeyPosition> KeyPositions;
+	KeyPositions.AddDefaulted(KeyHandles.Num());
+	InCurveModel.GetKeyPositions(KeyHandles, KeyPositions);
+
+	for (int32 i = 0; i < KeyHandles.Num(); i++)
+	{
+		FKeyHandlePositionPair& KeyHandlePositionPair = OutSortedKeyHandlesAndPositions.AddDefaulted_GetRef();
+		KeyHandlePositionPair.Handle = KeyHandles[i];
+		KeyHandlePositionPair.Position = KeyPositions[i];
+	}
+
+	OutSortedKeyHandlesAndPositions.Sort([](const FKeyHandlePositionPair& A, const FKeyHandlePositionPair& B) { return A.Position.InputValue < B.Position.InputValue; });
+}
+
+void SNiagaraCurveKeySelector::GetOrderedActiveCurveModelIds(TArray<FCurveModelID>& OutOrderedActiveCurveModelIds)
+{
+	if (CurveEditor->GetTreeSelection().Num() > 0)
+	{
+		// If there are curves selected in the tree then those are the only active ones.
+		for (const FCurveEditorTreeItemID& OrderedCurveTreeItemId : OrderedCurveTreeItemIds)
+		{
+			if(CurveEditor->GetTreeSelectionState(OrderedCurveTreeItemId) != ECurveEditorTreeSelectionState::None)
+			{
+				TArrayView<const FCurveModelID> CurveModelIds = CurveEditor->GetTreeItem(OrderedCurveTreeItemId).GetOrCreateCurves(CurveEditor.Get());
+				if (CurveModelIds.Num() == 1)
+				{
+					OutOrderedActiveCurveModelIds.Add(CurveModelIds[0]);
+				}
+			}
+		}
+	}
+	else
+	{
+		// Otherwise the active curves are the pinned curves.
+		for (const FCurveEditorTreeItemID& OrderedCurveTreeItemId : OrderedCurveTreeItemIds)
+		{
+			TArrayView<const FCurveModelID> CurveModelIds = CurveEditor->GetTreeItem(OrderedCurveTreeItemId).GetCurves();
+			if (CurveModelIds.Num() == 1 && CurveEditor->IsCurvePinned(CurveModelIds[0]))
+			{
+				OutOrderedActiveCurveModelIds.Add(CurveModelIds[0]);
+			}
+		}
+	}
+}
+
+void SNiagaraCurveKeySelector::NavigateToAdjacentCurve(ENavigateDirection Direction)
+{
+	if (CurveEditorTree.IsValid())
+	{
+		if (CurveEditor->GetTreeSelection().Num() == 0)
+		{
+			CurveEditorTree->SetItemSelection(OrderedCurveTreeItemIds[0], true);
+		}
+		else
+		{
+			int32 TargetSelectedTreeItemIndex = INDEX_NONE;
+			int32 StartIndex = Direction == ENavigateDirection::Previous ? 0 : OrderedCurveTreeItemIds.Num() - 1;
+			int32 EndIndex = Direction == ENavigateDirection::Previous ? OrderedCurveTreeItemIds.Num() : -1;
+			int32 IndexOffset = Direction == ENavigateDirection::Previous ? 1 : -1;
+			for (int32 i = StartIndex; i != EndIndex; i += IndexOffset)
+			{
+				if (CurveEditor->GetTreeSelectionState(OrderedCurveTreeItemIds[i]) != ECurveEditorTreeSelectionState::None)
+				{
+					TargetSelectedTreeItemIndex = i;
+					break;
+				}
+			}
+
+			if (TargetSelectedTreeItemIndex != INDEX_NONE)
+			{
+				int32 IndexToSelect = Direction == ENavigateDirection::Previous 
+					? TargetSelectedTreeItemIndex - 1
+					: TargetSelectedTreeItemIndex + 1;
+
+				if (IndexToSelect > OrderedCurveTreeItemIds.Num() - 1)
+				{
+					IndexToSelect = 0;
+				}
+				else if(IndexToSelect < 0)
+				{
+					IndexToSelect = OrderedCurveTreeItemIds.Num() - 1;
+				}
+
+				CurveEditorTree->ClearSelection();
+				CurveEditorTree->SetItemSelection(OrderedCurveTreeItemIds[IndexToSelect], true);
+			}
+		}
+	}
+}
+
+void SNiagaraCurveKeySelector::NavigateToAdjacentKey(ENavigateDirection Direction)
+{
+	TOptional<FCurveModelID> ActiveCurveModelId;
+	TArray<FKeyHandle> SelectedKeyHandles;
+	GetActiveCurveModelAndSelectedKeys(ActiveCurveModelId, SelectedKeyHandles);
+
+	TOptional<FCurveModelID> CurveModelIdToSelect;
+	TOptional<FKeyHandle> KeyHandleToSelect;
+	if (ActiveCurveModelId.IsSet())
+	{
+		FCurveModel* ActiveCurveModel = CurveEditor->GetCurves()[ActiveCurveModelId.GetValue()].Get();
+		TArray<FKeyHandlePositionPair> ActiveSortedKeyHandlePositionPairs;
+		GetSortedKeyHandlessAndPositionsForModel(*ActiveCurveModel, ActiveSortedKeyHandlePositionPairs);
+
+		if (SelectedKeyHandles.Num() != 0)
+		{
+			// If there's currently a selected key on the active curve then we want to use that as the target for navigating.
+			int32 TargetSelectedKeyIndex = INDEX_NONE;
+			int32 StartIndex = Direction == ENavigateDirection::Previous ? 0 : ActiveSortedKeyHandlePositionPairs.Num() - 1;
+			int32 EndIndex = Direction == ENavigateDirection::Previous ? ActiveSortedKeyHandlePositionPairs.Num() : -1;
+			int32 IndexOffset = Direction == ENavigateDirection::Previous ? 1 : -1;
+			for (int32 i = StartIndex; i != EndIndex; i += IndexOffset)
+			{
+				if (SelectedKeyHandles.Contains(ActiveSortedKeyHandlePositionPairs[i].Handle))
+				{
+					TargetSelectedKeyIndex = i;
+					break;
+				}
+			}
+
+			if (TargetSelectedKeyIndex != INDEX_NONE)
+			{
+				if(Direction == ENavigateDirection::Previous && TargetSelectedKeyIndex > 0)
+				{
+					// If we're navigating previous and we're not at the first key we can just select the previous key on this curve.
+					CurveModelIdToSelect = ActiveCurveModelId;
+					KeyHandleToSelect = ActiveSortedKeyHandlePositionPairs[TargetSelectedKeyIndex - 1].Handle;
+				}
+				else if (Direction == ENavigateDirection::Next && TargetSelectedKeyIndex < ActiveSortedKeyHandlePositionPairs.Num() - 1)
+				{
+					// If we're navigating next and we're not at the last key we can just select the next key on this curve.
+					CurveModelIdToSelect = ActiveCurveModelId;
+					KeyHandleToSelect = ActiveSortedKeyHandlePositionPairs[TargetSelectedKeyIndex + 1].Handle;
+				}
+				else
+				{
+					// Otherwise we're going to need to navigate to another curve since we're at the start or end of the current curve.
+					TArray<FCurveModelID> OrderedActiveCurveModelIds;
+					GetOrderedActiveCurveModelIds(OrderedActiveCurveModelIds);
+
+					// Find the adjacent curve with keys so we can select a key there.
+					FCurveModel* CurveModelToSelect = nullptr;
+					int32 CurrentIndex = OrderedActiveCurveModelIds.IndexOfByKey(ActiveCurveModelId);
+					if (OrderedActiveCurveModelIds.Num() > 0)
+					{
+						int32 CurrentIndexOffset = Direction == ENavigateDirection::Previous ? -1 : 1;
+						while (CurveModelIdToSelect.IsSet() == false)
+						{
+							CurrentIndex += CurrentIndexOffset;
+							if (CurrentIndex < 0)
+							{
+								CurrentIndex = OrderedActiveCurveModelIds.Num() - 1;
+							}
+							else if (CurrentIndex > OrderedActiveCurveModelIds.Num() - 1)
+							{
+								CurrentIndex = 0;
+							}
+
+							FCurveModel* CurrentCurveModel = CurveEditor->GetCurves()[OrderedActiveCurveModelIds[CurrentIndex]].Get();
+							if (CurrentCurveModel->GetNumKeys() > 0)
+							{
+								CurveModelIdToSelect = OrderedActiveCurveModelIds[CurrentIndex];
+								CurveModelToSelect = CurrentCurveModel;
+							}
+						}
+					}
+
+					if (CurveModelIdToSelect == ActiveCurveModelId)
+					{
+						// There were no other active curves with keys to select from so just wrap around the currently active curve.
+						if (Direction == ENavigateDirection::Previous)
+						{
+							KeyHandleToSelect = ActiveSortedKeyHandlePositionPairs.Last().Handle;
+						}
+						else
+						{
+							KeyHandleToSelect = ActiveSortedKeyHandlePositionPairs[0].Handle;
+						}
+					}
+					else if(CurveModelToSelect != nullptr)
+					{
+						// We're selecting a key on a different curve so we need to sort the positions and select the first or last based
+						// on the navigation direction.
+						TArray<FKeyHandlePositionPair> SortedKeyHandlePositionPairs;
+						GetSortedKeyHandlessAndPositionsForModel(*CurveModelToSelect, SortedKeyHandlePositionPairs);
+						if (Direction == ENavigateDirection::Previous)
+						{
+							KeyHandleToSelect = SortedKeyHandlePositionPairs.Last().Handle;
+						}
+						else
+						{
+							KeyHandleToSelect = SortedKeyHandlePositionPairs[0].Handle;
+						}
+					}
+				}
+			}
+		}
+		else if(ActiveSortedKeyHandlePositionPairs.Num() > 0)
+		{
+			// There weren't any keys already selected so just select the first key on the active curve.
+			CurveModelIdToSelect = ActiveCurveModelId;
+			KeyHandleToSelect = ActiveSortedKeyHandlePositionPairs[0].Handle;
+		}
+	}
+
+	if (CurveModelIdToSelect.IsSet() && KeyHandleToSelect.IsSet())
+	{
+		FCurveEditorSelection& CurveEditorSelection = CurveEditor->GetSelection();
+		CurveEditorSelection.Clear();
+		CurveEditorSelection.Add(CurveModelIdToSelect.GetValue(), ECurvePointType::Key, KeyHandleToSelect.GetValue());
+
+		TArray<FCurveModelID> CurvesToFit;
+		CurvesToFit.Add(CurveModelIdToSelect.GetValue());
+		CurveEditor->ZoomToFitCurves(CurvesToFit);
+	}
+}
+
+FReply SNiagaraCurveKeySelector::ZoomToFitClicked()
+{
+	TArray<FCurveModelID> ActiveCurveModelIds;
+	GetOrderedActiveCurveModelIds(ActiveCurveModelIds);
+	CurveEditor->ZoomToFitCurves(ActiveCurveModelIds);
+	return FReply::Handled();
+}
+
+FReply SNiagaraCurveKeySelector::PreviousCurveClicked()
+{
+	NavigateToAdjacentCurve(ENavigateDirection::Previous);
+	return FReply::Handled();
+}
+
+FReply SNiagaraCurveKeySelector::NextCurveClicked()
+{
+	NavigateToAdjacentCurve(ENavigateDirection::Next);
+	return FReply::Handled();
+}
+
+FReply SNiagaraCurveKeySelector::PreviousKeyClicked()
+{
+	NavigateToAdjacentKey(ENavigateDirection::Previous);
+	return FReply::Handled();
+}
+
+FReply SNiagaraCurveKeySelector::NextKeyClicked()
+{
+	NavigateToAdjacentKey(ENavigateDirection::Next);
+	return FReply::Handled();
+}
+
+EVisibility SNiagaraCurveKeySelector::GetNextPreviousCurveButtonVisibility() const
+{
+	return OrderedCurveTreeItemIds.Num() == 1 ? EVisibility::Collapsed : EVisibility::Visible;
+}
+
+FReply SNiagaraCurveKeySelector::AddKeyClicked()
+{
+	TOptional<FCurveModelID> CurveModelIdForAdd;
+	TArray<FKeyHandle> SelectedKeyHandles;
+	GetActiveCurveModelAndSelectedKeys(CurveModelIdForAdd, SelectedKeyHandles);
+
+	if(CurveModelIdForAdd.IsSet())
+	{
+		FCurveModel* CurveModelForAdd = CurveEditor->GetCurves()[CurveModelIdForAdd.GetValue()].Get();
+
+		FKeyPosition NewKeyPosition;
+		FKeyAttributes NewKeyAttributes = CurveEditor->GetDefaultKeyAttributes().Get();
+		NewKeyAttributes.SetInterpMode(RCIM_Cubic);
+		NewKeyAttributes.SetTangentMode(RCTM_Auto);
+		if (CurveModelForAdd->GetNumKeys() == 0)
+		{
+			// If there are no keys, add one at 0, 0.
+			NewKeyPosition.InputValue = 0.0f;
+			NewKeyPosition.OutputValue = 0.0f;
+		}
+		else if (CurveModelForAdd->GetNumKeys() == 1)
+		{
+			// If there's a single key, add the new key at the same value, but time + 1.
+			const TArray<FKeyHandle> KeyHandles = CurveModelForAdd->GetAllKeys();
+				
+			TArray<FKeyPosition> KeyPositions;
+			KeyPositions.AddDefaulted();
+			CurveModelForAdd->GetKeyPositions(KeyHandles, KeyPositions);
+						
+			NewKeyPosition.InputValue = KeyPositions[0].InputValue + 1;
+			NewKeyPosition.OutputValue = KeyPositions[0].OutputValue;
+		}
+		else
+		{
+			TArray<FKeyHandlePositionPair> SortedKeyHandlePositionPairs;
+			GetSortedKeyHandlessAndPositionsForModel(*CurveModelForAdd, SortedKeyHandlePositionPairs);
+
+			int32 IndexToAddAfter = INDEX_NONE;
+			if (SelectedKeyHandles.Num() > 0)
+			{
+				for (int32 i = SortedKeyHandlePositionPairs.Num() - 1; i >= 0; i--)
+				{
+					if (SelectedKeyHandles.Contains(SortedKeyHandlePositionPairs[i].Handle))
+					{
+						IndexToAddAfter = i;
+						break;
+					}
+				}
+			}
+
+			if(IndexToAddAfter == INDEX_NONE)
+			{
+				IndexToAddAfter = SortedKeyHandlePositionPairs.Num() - 1;
+			}
+
+			if (IndexToAddAfter == SortedKeyHandlePositionPairs.Num() - 1)
+			{
+				const FKeyPosition& TargetKeyPosition = SortedKeyHandlePositionPairs[IndexToAddAfter].Position;
+				const FKeyPosition& PreviousKeyPosition = SortedKeyHandlePositionPairs[IndexToAddAfter - 1].Position;
+				NewKeyPosition.InputValue = TargetKeyPosition.InputValue + (TargetKeyPosition.InputValue - PreviousKeyPosition.InputValue);
+				NewKeyPosition.OutputValue = TargetKeyPosition.OutputValue + (TargetKeyPosition.OutputValue - PreviousKeyPosition.OutputValue);
+			}
+			else
+			{
+				const FKeyPosition& TargetKeyPosition = SortedKeyHandlePositionPairs[IndexToAddAfter].Position;
+				const FKeyPosition& NextKeyPosition = SortedKeyHandlePositionPairs[IndexToAddAfter + 1].Position;
+				NewKeyPosition.InputValue = TargetKeyPosition.InputValue + ((NextKeyPosition.InputValue - TargetKeyPosition.InputValue) / 2);
+				NewKeyPosition.OutputValue = TargetKeyPosition.OutputValue + ((NextKeyPosition.OutputValue - TargetKeyPosition.OutputValue) / 2);
+			}
+		}
+
+		FScopedTransaction Transaction(LOCTEXT("AddKey", "Add Key"));
+		TOptional<FKeyHandle> NewKeyHandle = CurveModelForAdd->AddKey(NewKeyPosition, NewKeyAttributes);
+		if(NewKeyHandle.IsSet())
+		{
+			FCurveEditorSelection& CurveEditorSelection = CurveEditor->GetSelection();
+			CurveEditorSelection.Clear();
+			CurveEditorSelection.Add(CurveModelIdForAdd.GetValue(), ECurvePointType::Key, NewKeyHandle.GetValue());
+
+			TArray<FCurveModelID> CurvesToFit;
+			CurvesToFit.Add(CurveModelIdForAdd.GetValue());
+			CurveEditor->ZoomToFitCurves(CurvesToFit);
+		}
+	}
+	return FReply::Handled();
+}
+
+FReply SNiagaraCurveKeySelector::DeleteKeyClicked()
+{
+	CurveEditor->DeleteSelection();
+	return FReply::Handled();
+}
+
+#undef LOCTEXT_NAMESPACE
diff --git a/Engine/Plugins/FX/Niagara/Source/NiagaraEditorWidgets/Private/SNiagaraCurveKeySelector.h b/Engine/Plugins/FX/Niagara/Source/NiagaraEditorWidgets/Private/SNiagaraCurveKeySelector.h
new file mode 100644
--- /dev/null
+++ b/Engine/Plugins/FX/Niagara/Source/NiagaraEditorWidgets/Private/SNiagaraCurveKeySelector.h
@@ -0,0 +1,60 @@
+// Copyright Epic Games, Inc. All Rights Reserved.
+
+#pragma once
+
+#include "CurveEditorTypes.h"
+#include "Curves/KeyHandle.h"
+#include "Input/Reply.h"
+#include "Widgets/SCompoundWidget.h"
+#include "CurveDataAbstraction.h"
+
+class FCurveEditor;
+struct FCurveEditorTreeItemID;
+class SCurveEditorTree;
+class FCurveModel;
+
+class SNiagaraCurveKeySelector : public SCompoundWidget
+{
+public:
+	SLATE_BEGIN_ARGS(SNiagaraCurveKeySelector) {}
+	SLATE_END_ARGS()
+
+	void Construct(const FArguments& InArgs, TSharedPtr<FCurveEditor> InCurveEditor, const TArray<FCurveEditorTreeItemID>& InCurveTreeItemIds, TSharedPtr<SCurveEditorTree> InCurveEditorTree);
+
+private:
+	void GetActiveCurveModelAndSelectedKeys(TOptional<FCurveModelID>& OutActiveCurveModelId, TArray<FKeyHandle>& OutSelectedKeyHandles);
+
+	struct FKeyHandlePositionPair
+	{
+		FKeyHandle Handle;
+		FKeyPosition Position;
+	};
+
+	void GetSortedKeyHandlessAndPositionsForModel(FCurveModel& InCurveModel, TArray<FKeyHandlePositionPair>& OutSortedKeyHandlesAndPositions);
+	void GetOrderedActiveCurveModelIds(TArray<FCurveModelID>& OutOrderedActiveCurveModelIds);
+
+	enum class ENavigateDirection
+	{
+		Previous,
+		Next
+	};
+
+	void NavigateToAdjacentCurve(ENavigateDirection Direction);
+	void NavigateToAdjacentKey(ENavigateDirection Direction);
+
+	FReply ZoomToFitClicked();
+	FReply PreviousCurveClicked();
+	FReply NextCurveClicked();
+	FReply PreviousKeyClicked();
+	FReply NextKeyClicked();
+
+	EVisibility GetNextPreviousCurveButtonVisibility() const;
+
+	FReply AddKeyClicked();
+	FReply DeleteKeyClicked();
+
+private:
+	TSharedPtr<FCurveEditor> CurveEditor;
+	TArray<FCurveEditorTreeItemID> OrderedCurveTreeItemIds;
+	TSharedPtr<SCurveEditorTree> CurveEditorTree;
+};
diff --git a/Engine/Plugins/FX/Niagara/Source/NiagaraEditorWidgets/Private/SNiagaraCurveThumbnail.cpp b/Engine/Plugins/FX/Niagara/Source/NiagaraEditorWidgets/Private/SNiagaraCurveThumbnail.cpp
new file mode 100644
--- /dev/null
+++ b/Engine/Plugins/FX/Niagara/Source/NiagaraEditorWidgets/Private/SNiagaraCurveThumbnail.cpp
@@ -0,0 +1,49 @@
+// Copyright Epic Games, Inc. All Rights Reserved.
+
+#include "SNiagaraCurveThumbnail.h"
+#include "Curves/RichCurve.h"
+
+#define LOCTEXT_NAMESPACE "SNiagaraCurveThumbnail"
+
+void SNiagaraCurveThumbnail::Construct(const FArguments& InArgs, const FRichCurve& CurveToDisplay)
+{
+	Width = InArgs._Width;
+	Height = InArgs._Height;
+
+	float TimeMin;
+	float TimeMax;
+	float ValueMin;
+	float ValueMax;
+	CurveToDisplay.GetTimeRange(TimeMin, TimeMax);
+	CurveToDisplay.GetValueRange(ValueMin, ValueMax);
+
+	const float TimeRange = TimeMax - TimeMin;
+	const float ValueRange = ValueMax - ValueMin;
+	if (!FMath::IsNearlyZero(TimeRange) && !FMath::IsNearlyZero(ValueRange))
+	{
+		constexpr int32 Points = 13;
+		const float TimeIncrement = TimeRange / (Points - 1);
+		for (int32 i = 0; i < Points; i++)
+		{
+			const float Time = TimeMin + i * TimeIncrement;
+			const float Value = CurveToDisplay.Eval(Time);
+
+			const float NormalizedX = (Time - TimeMin) / TimeRange;
+			const float NormalizedY = (Value - ValueMin) / ValueRange;
+			CurvePoints.Add(FVector2D(NormalizedX * Width, (1 - NormalizedY) * Height));
+		}
+	}
+}
+
+int32 SNiagaraCurveThumbnail::OnPaint(const FPaintArgs& Args, const FGeometry& AllottedGeometry, const FSlateRect& MyCullingRect, FSlateWindowElementList& OutDrawElements, int32 LayerId, const FWidgetStyle& InWidgetStyle, bool bParentEnabled) const
+{
+	FSlateDrawElement::MakeLines(OutDrawElements, LayerId, AllottedGeometry.ToPaintGeometry(), CurvePoints, ESlateDrawEffect::None, InWidgetStyle.GetForegroundColor(), true, 2.0f);
+	return LayerId;
+}
+
+FVector2D SNiagaraCurveThumbnail::ComputeDesiredSize(float) const
+{
+	return FVector2D(Width, Height);
+}
+
+#undef LOCTEXT_NAMESPACE
diff --git a/Engine/Plugins/FX/Niagara/Source/NiagaraEditorWidgets/Private/SNiagaraCurveThumbnail.h b/Engine/Plugins/FX/Niagara/Source/NiagaraEditorWidgets/Private/SNiagaraCurveThumbnail.h
new file mode 100644
--- /dev/null
+++ b/Engine/Plugins/FX/Niagara/Source/NiagaraEditorWidgets/Private/SNiagaraCurveThumbnail.h
@@ -0,0 +1,33 @@
+// Copyright Epic Games, Inc. All Rights Reserved.
+
+#pragma once
+
+#include "Styling/SlateTypes.h"
+#include "Input/Reply.h"
+#include "Widgets/SLeafWidget.h"
+
+class UCurveBase;
+class UCurveFloat;
+struct FAssetData;
+struct FRichCurve;
+
+class SNiagaraCurveThumbnail : public SLeafWidget
+{
+	SLATE_BEGIN_ARGS(SNiagaraCurveThumbnail)
+		: _Width(16)
+		, _Height(8)
+		{
+		}
+		SLATE_ARGUMENT(float, Width)
+		SLATE_ARGUMENT(float, Height)
+	SLATE_END_ARGS()
+
+	void Construct(const FArguments& InArgs, const FRichCurve& CurveToDisplay);
+
+	virtual int32 OnPaint(const FPaintArgs& Args, const FGeometry& AllottedGeometry, const FSlateRect& MyCullingRect, FSlateWindowElementList& OutDrawElements, int32 LayerId, const FWidgetStyle& InWidgetStyle, bool bParentEnabled) const override;
+	virtual FVector2D ComputeDesiredSize(float) const override;
+
+	TArray<FVector2D> CurvePoints;
+	float Width;
+	float Height;
+};
diff --git a/Engine/Plugins/FX/Niagara/Source/NiagaraEditorWidgets/Public/NiagaraEditorWidgetsModule.h b/Engine/Plugins/FX/Niagara/Source/NiagaraEditorWidgets/Public/NiagaraEditorWidgetsModule.h
--- a/Engine/Plugins/FX/Niagara/Source/NiagaraEditorWidgets/Public/NiagaraEditorWidgetsModule.h
+++ b/Engine/Plugins/FX/Niagara/Source/NiagaraEditorWidgets/Public/NiagaraEditorWidgetsModule.h
@@ -8,6 +8,11 @@
 #include "UObject/ObjectKey.h"
 #include "Templates/SharedPointer.h"
 
+class UObject;
+class FCurveEditor;
+struct FCurveEditorTreeItemID;
+class SCurveEditorTree;
+
 class FNiagaraStackCurveEditorOptions
 {
 public:
@@ -42,8 +47,6 @@ private:
 	float Height;
 };
 
-class UObject;
-
 /** A module containing widgets for editing niagara data. */
 class FNiagaraEditorWidgetsModule : public IModuleInterface
 {
@@ -56,6 +59,8 @@ private:
 		virtual TSharedRef<SWidget> CreateStackIssueIcon(UNiagaraStackViewModel& StackViewModel, UNiagaraStackEntry& StackEntry) const override;
 		virtual TSharedRef<SWidget> CreateScriptScratchPadManager(UNiagaraScratchPadViewModel& ScriptScratchPadViewModel) const override;
 		virtual TSharedRef<SWidget> CreateCurveOverview(TSharedRef<FNiagaraSystemViewModel> SystemViewModel) const override;
+		virtual TSharedRef<SWidget> CreateCurveKeySelector(TSharedPtr<FCurveEditor> InCurveEditor, const TArray<FCurveEditorTreeItemID>& InCurveTreeItemIds, TSharedPtr<SCurveEditorTree> InCurveEditorTree) const override;
+		virtual TSharedRef<SWidget> CreateCurveThumbnail(const FRichCurve& Curve) const override;
 		virtual FLinearColor GetColorForExecutionCategory(FName ExecutionCategory) const override;
 	};
 
