- Fix audio data interace hitting a check when inside serialize

#rnx

diff --git a/Engine/Plugins/FX/Niagara/Source/Niagara/Private/NiagaraDataInterfaceAudio.cpp b/Engine/Plugins/FX/Niagara/Source/Niagara/Private/NiagaraDataInterfaceAudio.cpp
--- a/Engine/Plugins/FX/Niagara/Source/Niagara/Private/NiagaraDataInterfaceAudio.cpp
+++ b/Engine/Plugins/FX/Niagara/Source/Niagara/Private/NiagaraDataInterfaceAudio.cpp
@@ -367,7 +367,23 @@ bool UNiagaraDataInterfaceAudioSubmix::CopyToInternal(UNiagaraDataInterface* Des
 	if (CastedDestination)
 	{
 		CastedDestination->Submix = Submix;
-		CastedDestination->GetProxyAs<FNiagaraDataInterfaceProxySubmix>()->OnUpdateSubmix(Submix);
+		if (IsInGameThread())
+		{
+			CastedDestination->GetProxyAs<FNiagaraDataInterfaceProxySubmix>()->OnUpdateSubmix(Submix);
+		}
+		else
+		{
+			AsyncTask(
+				ENamedThreads::GameThread,
+				[WeakSubmixDI=MakeWeakObjectPtr(CastedDestination)]()
+				{
+					if (UNiagaraDataInterfaceAudioSubmix* SubmixDI = WeakSubmixDI.Get())
+					{
+						SubmixDI->GetProxyAs<FNiagaraDataInterfaceProxySubmix>()->OnUpdateSubmix(SubmixDI->Submix);
+					}
+				}
+			);
+		}
 	}
 	
 	return true;
