- Capture audio / collision information into CSV stats
- This doesn't capture async tracing cost as we don't have a way to hook into that data currently

#rnx

diff --git a/Engine/Plugins/FX/Niagara/Source/Niagara/Private/NiagaraCollision.cpp b/Engine/Plugins/FX/Niagara/Source/Niagara/Private/NiagaraCollision.cpp
--- a/Engine/Plugins/FX/Niagara/Source/Niagara/Private/NiagaraCollision.cpp
+++ b/Engine/Plugins/FX/Niagara/Source/Niagara/Private/NiagaraCollision.cpp
@@ -26,6 +26,7 @@ static FAutoConsoleVariableRef CVarNiagaraCollisionEnabled(
 void FNiagaraDICollisionQueryBatch::DispatchQueries()
 {
 	check(IsInGameThread());
+	CSV_SCOPED_TIMING_STAT(Particles, CoreSystems_Collision);
 
 	// swap the buffers, and use our new read buffer for issuing all of our accumulated queries now that we're on the main thread
 	FlipBuffers();
@@ -60,6 +61,7 @@ void FNiagaraDICollisionQueryBatch::DispatchQueries()
 void FNiagaraDICollisionQueryBatch::CollectResults()
 {
 	check(IsInGameThread());
+	CSV_SCOPED_TIMING_STAT(Particles, CoreSystems_Collision);
 
 	// locks are not used here because we assume that the per instance ticking is done from the main thread
 	if (CollisionWorld)
@@ -147,6 +149,8 @@ bool FNiagaraDICollisionQueryBatch::PerformQuery(FVector StartPos, FVector EndPo
 	}
 
 	SCOPE_CYCLE_COUNTER(STAT_NiagaraCollision);
+	CSV_SCOPED_TIMING_STAT(Particles, CoreSystems_Collision);
+
 	FCollisionQueryParams QueryParams(TraceTagName, StatId, false);
 	QueryParams.OwnerTag = CollisionTagName;
 	QueryParams.bFindInitialOverlaps = false;
diff --git a/Engine/Plugins/FX/Niagara/Source/Niagara/Private/NiagaraDataInterfaceAudioPlayer.cpp b/Engine/Plugins/FX/Niagara/Source/Niagara/Private/NiagaraDataInterfaceAudioPlayer.cpp
--- a/Engine/Plugins/FX/Niagara/Source/Niagara/Private/NiagaraDataInterfaceAudioPlayer.cpp
+++ b/Engine/Plugins/FX/Niagara/Source/Niagara/Private/NiagaraDataInterfaceAudioPlayer.cpp
@@ -113,6 +113,8 @@ bool UNiagaraDataInterfaceAudioPlayer::InitPerInstanceData(void* PerInstanceData
 
 void UNiagaraDataInterfaceAudioPlayer::DestroyPerInstanceData(void* PerInstanceData, FNiagaraSystemInstance* SystemInstance)
 {
+	CSV_SCOPED_TIMING_STAT(Particles, CoreSystems_Audio);
+
 	FAudioPlayerInterface_InstanceData* InstData = (FAudioPlayerInterface_InstanceData*)PerInstanceData;
 	for (const auto& Entry : InstData->PersistentAudioMapping)
 	{
@@ -126,6 +128,8 @@ void UNiagaraDataInterfaceAudioPlayer::DestroyPerInstanceData(void* PerInstanceD
 
 bool UNiagaraDataInterfaceAudioPlayer::PerInstanceTick(void* PerInstanceData, FNiagaraSystemInstance* SystemInstance, float DeltaSeconds)
 {
+	CSV_SCOPED_TIMING_STAT(Particles, CoreSystems_Audio);
+
 	FAudioPlayerInterface_InstanceData* PIData = static_cast<FAudioPlayerInterface_InstanceData*>(PerInstanceData);
 	if (!PIData)
 	{
@@ -192,6 +196,8 @@ FAudioParameter ConvertVariableToAudioParam(const FNiagaraVariable& NiagaraVaria
 
 bool UNiagaraDataInterfaceAudioPlayer::PerInstanceTickPostSimulate(void* PerInstanceData, FNiagaraSystemInstance* SystemInstance, float DeltaSeconds)
 {
+	CSV_SCOPED_TIMING_STAT(Particles, CoreSystems_Audio);
+
 	FAudioPlayerInterface_InstanceData* PIData = static_cast<FAudioPlayerInterface_InstanceData*>(PerInstanceData);
 	UNiagaraSystem* System = SystemInstance->GetSystem();
 	UWorld* World = SystemInstance->GetWorldManager()->GetWorld();
