User Asset Tagging: v3

Tags are simply Names stored in FMetaData, serialized to Asset Registry.
This CL contains:
- User Asset Tags, Editor, API
-- Hit Ctrl T on selected assets to summon the User Asset Tags Editor that allows you to assign tags
-- User Asset Tag Providers, with a Local Favorites and a Project Settings Tag provider implementations

- Tagged Asset Browser
-- Asset to configure filters
-- Added section support
-- Integrated Data Hierarchy Editor to organize filters, sections etc.
-- Filters are extendable; you can write new filters, assign icons etc.

- Changes to Data Hierarchy Editor
-- Reworked sections; from hardcoded sections in categories to a metadata approach. Any element that can be dropped onto a section will create a section association
-- Reworked various parts of the DHT, like removing unnecessary functions, streamlining code more, adding default User Feedback for various actions
-- New 'Add new element via UI' API - adding new elements via "Add" menu

- Niagara Changes
-- Moving Niagara Asset Tags to User Asset Tags; deprecating old related classes and functions
-- Emitters and Systems look up the display names of their old tags, and move them to the new tags contained in FMetaData during PostLoad

- UPackage Changes
-- Added "transfer metadata requested delegate" that is called from ObjectTools when renaming/moving a package. The UAT system hooks into it to transfer the UAT to the new metadata object

diff --git a/Engine/Config/BaseUserAssetTags.ini b/Engine/Config/BaseUserAssetTags.ini
new file mode 100644
--- /dev/null
+++ b/Engine/Config/BaseUserAssetTags.ini
@@ -0,0 +1,3 @@
+[/Script/UserAssetTagsEditor.ProjectUserAssetTagSettings]
+FavoriteUserAssetTagsPerClass=(("/Script/Niagara.NiagaraSystem", (FavoriteUserAssetTags=("Template","Lightweight","Learning Content"))),("/Script/Engine.Material", (FavoriteUserAssetTags=("Fire","Energy","Electricity","Distortion"))),("/Script/Niagara.NiagaraEmitter", (FavoriteUserAssetTags=("Template","Lightweight","Learning Content"))))
+
diff --git a/Engine/Plugins/FX/Niagara/Source/Niagara/Classes/NiagaraEmitter.h b/Engine/Plugins/FX/Niagara/Source/Niagara/Classes/NiagaraEmitter.h
--- a/Engine/Plugins/FX/Niagara/Source/Niagara/Classes/NiagaraEmitter.h
+++ b/Engine/Plugins/FX/Niagara/Source/Niagara/Classes/NiagaraEmitter.h
@@ -10,7 +10,6 @@
 #include "NiagaraMessageDataBase.h"
 #include "NiagaraMessageStore.h"
 #include "INiagaraMergeManager.h"
-#include "NiagaraAssetTagDefinitions.h"
 #include "NiagaraEffectType.h"
 #include "NiagaraEmitterBase.h"
 #include "NiagaraDataInterfaceEmitterBinding.h"
@@ -25,6 +24,7 @@
 
 #include "NiagaraEmitter.generated.h"
 
+struct FNiagaraAssetTagDefinitionReference;
 class UNiagaraScratchPadContainer;
 class UMaterial;
 class UNiagaraEmitter;
@@ -746,9 +746,6 @@ private:
 public:
 	
 #if WITH_EDITORONLY_DATA
-	UPROPERTY()
-	TArray<FNiagaraAssetTagDefinitionReference> AssetTags;
-
 	/** If an emitter is inheritable, new emitters based on an inheritable emitter, or Niagara Systems using an inheritable emitter, will automatically inherit changes made to the original emitter. */
 	UPROPERTY(EditAnywhere, Category = "Asset Options", AssetRegistrySearchable)
 	bool bIsInheritable = true;
@@ -883,6 +880,9 @@ public:
 	/** Use property in struct returned from GetEmitterData() instead */ 
 	UPROPERTY(meta = (DeprecatedProperty))
 	uint32 MaxGPUParticlesSpawnPerFrame_DEPRECATED;
+
+	UPROPERTY(meta = (DeprecatedProperty))
+	TArray<FNiagaraAssetTagDefinitionReference> AssetTags_DEPRECATED;
 #endif
 	
 	NIAGARA_API void UpdateEmitterAfterLoad();
diff --git a/Engine/Plugins/FX/Niagara/Source/Niagara/Classes/NiagaraScript.h b/Engine/Plugins/FX/Niagara/Source/Niagara/Classes/NiagaraScript.h
--- a/Engine/Plugins/FX/Niagara/Source/Niagara/Classes/NiagaraScript.h
+++ b/Engine/Plugins/FX/Niagara/Source/Niagara/Classes/NiagaraScript.h
@@ -18,7 +18,6 @@
 #include "NiagaraVersionedObject.h"
 #include "HAL/CriticalSection.h"
 #include "VectorVM.h"
-#include "NiagaraAssetTagDefinitions.h"
 
 #include "NiagaraScript.generated.h"
 
@@ -644,9 +643,6 @@ public:
 	UPROPERTY(EditAnywhere, Category = Script)
 	FText InlineOverviewDisplayName;
 
-	UPROPERTY(EditAnywhere, Category = "Asset Options")
-	TArray<FNiagaraAssetTagDefinitionReference> AssetTagDefinitionReferences;
-
 	/** If true, this script will be added to a 'Suggested' category at the top of menus during searches */
 	UPROPERTY(AssetRegistrySearchable, EditAnywhere, Category = Script)
 	bool bSuggested = false;
diff --git a/Engine/Plugins/FX/Niagara/Source/Niagara/Classes/NiagaraSystem.h b/Engine/Plugins/FX/Niagara/Source/Niagara/Classes/NiagaraSystem.h
--- a/Engine/Plugins/FX/Niagara/Source/Niagara/Classes/NiagaraSystem.h
+++ b/Engine/Plugins/FX/Niagara/Source/Niagara/Classes/NiagaraSystem.h
@@ -3,7 +3,6 @@
 #pragma once
 
 #include "CoreMinimal.h"
-#include "NiagaraAssetTagDefinitions.h"
 
 #include "NiagaraDataSetCompiledData.h"
 #include "NiagaraDataSetAccessor.h"
@@ -20,6 +19,7 @@
 #include "Particles/FXBudget.h"
 #include "NiagaraSystem.generated.h"
 
+struct FNiagaraAssetTagDefinitionReference;
 class FNiagaraActiveCompilation;
 class UNiagaraEffectType;
 enum class ENiagaraCullProxyMode : uint32;
@@ -493,6 +493,9 @@ public:
 	UPROPERTY()
 	ENiagaraScriptTemplateSpecification TemplateSpecification_DEPRECATED;
 
+	UPROPERTY(meta=(DeprecatedProperty))
+	TArray<FNiagaraAssetTagDefinitionReference> AssetTags_DEPRECATED;
+	
 	UPROPERTY(EditAnywhere, AdvancedDisplay, Category = "Asset Options", AssetRegistrySearchable, DisplayName="Asset Description", meta = (SkipSystemResetOnChange = "true"))
 	FText TemplateAssetDescription;
 
@@ -500,9 +503,6 @@ public:
 	UPROPERTY(AssetRegistrySearchable, meta = (SkipSystemResetOnChange = "true"))
 	FText Category;
 
-	UPROPERTY(meta = (SkipSystemResetOnChange = "true"))
-	TArray<FNiagaraAssetTagDefinitionReference> AssetTags;
-
 	UPROPERTY(EditAnywhere, AdvancedDisplay, Category="Asset Options", AssetRegistrySearchable, meta = (SkipSystemResetOnChange = "true", AllowedClasses="/Script/MediaAssets.FileMediaSource"))
 	FSoftObjectPath PreviewMoviePath;
 	
diff --git a/Engine/Plugins/FX/Niagara/Source/Niagara/Classes/NiagaraValidationRule.h b/Engine/Plugins/FX/Niagara/Source/Niagara/Classes/NiagaraValidationRule.h
--- a/Engine/Plugins/FX/Niagara/Source/Niagara/Classes/NiagaraValidationRule.h
+++ b/Engine/Plugins/FX/Niagara/Source/Niagara/Classes/NiagaraValidationRule.h
@@ -70,7 +70,7 @@ class UNiagaraValidationRule : public UObject
 	GENERATED_BODY()
 
 public:
-	bool IsEnabled() const { return bIsConfigDisabled == false; }
+	bool IsEnabled() const { return bIsConfigDisabled == false && GetClass()->HasAnyClassFlags(CLASS_Deprecated) == false; }
 #if WITH_EDITOR
 	NIAGARA_API virtual void CheckValidity(const FNiagaraValidationContext& Context, TArray<FNiagaraValidationResult>& OutResults) const;
 #endif
diff --git a/Engine/Plugins/FX/Niagara/Source/Niagara/Private/NiagaraAssetTagDefinitions.cpp b/Engine/Plugins/FX/Niagara/Source/Niagara/Private/NiagaraAssetTagDefinitions.cpp
--- a/Engine/Plugins/FX/Niagara/Source/Niagara/Private/NiagaraAssetTagDefinitions.cpp
+++ b/Engine/Plugins/FX/Niagara/Source/Niagara/Private/NiagaraAssetTagDefinitions.cpp
@@ -94,7 +94,7 @@ void FNiagaraAssetTagDefinitionReference::AddTagToAssetRegistryTags(FAssetRegist
 	}
 }
 
-FText UNiagaraAssetTagDefinitions::GetDisplayName() const
+FText UDEPRECATED_NiagaraAssetTagDefinitions::GetDisplayName() const
 {
 	if(DisplayName.IsEmptyOrWhitespace())
 	{
@@ -104,12 +104,12 @@ FText UNiagaraAssetTagDefinitions::GetDisplayName() const
 	return DisplayName;
 }
 
-FText UNiagaraAssetTagDefinitions::GetDescription() const
+FText UDEPRECATED_NiagaraAssetTagDefinitions::GetDescription() const
 {
 	return Description;
 }
 
-bool UNiagaraAssetTagDefinitions::DoesAssetDataContainAnyTag(const FAssetData& AssetData) const
+bool UDEPRECATED_NiagaraAssetTagDefinitions::DoesAssetDataContainAnyTag(const FAssetData& AssetData) const
 {
 	for(const FNiagaraAssetTagDefinition& AssetTagDefinition : TagDefinitions)
 	{
diff --git a/Engine/Plugins/FX/Niagara/Source/Niagara/Private/NiagaraEmitter.cpp b/Engine/Plugins/FX/Niagara/Source/Niagara/Private/NiagaraEmitter.cpp
--- a/Engine/Plugins/FX/Niagara/Source/Niagara/Private/NiagaraEmitter.cpp
+++ b/Engine/Plugins/FX/Niagara/Source/Niagara/Private/NiagaraEmitter.cpp
@@ -631,15 +631,6 @@ void UNiagaraEmitter::PostLoad()
 		{
 			bIsInheritable = false;
 		}
-
-		if(CurrentTemplateSpecification == ENiagaraScriptTemplateSpecification::Template)
-		{			
-			AssetTags.AddUnique(INiagaraModule::TemplateTagDefinition);
-		}
-		else if(CurrentTemplateSpecification == ENiagaraScriptTemplateSpecification::Behavior)
-		{
-			AssetTags.AddUnique(INiagaraModule::LearningContentTagDefinition);
-		}
 	}
 	
 	if (MessageKeyToMessageMap_DEPRECATED.IsEmpty() == false)
@@ -658,6 +649,12 @@ void UNiagaraEmitter::PostLoad()
 	}
 #endif
 
+#if WITH_EDITORONLY_DATA	
+	if(NiagaraVer < FNiagaraCustomVersion::Type::UserAssetTagStorageInMetaData && AssetTags_DEPRECATED.Num() > 0)
+	{
+		INiagaraModule::Get().GetEditorOnlyDataUtilities().MoveUserAssetTagsToMetaData(this);
+	}
+#endif
 	const int32 UE5MainVer = GetLinkerCustomVersion(FUE5MainStreamObjectVersion::GUID);
 
 	for (FVersionedNiagaraEmitterData& Data : VersionData)
@@ -1049,11 +1046,6 @@ void UNiagaraEmitter::GetAssetRegistryTags(FAssetRegistryTagsContext Context) co
 	INiagaraModule& NiagaraModule = FModuleManager::GetModuleChecked<INiagaraModule>("Niagara");
 	Context.AddTag(NiagaraModule.GetEditorOnlyDataUtilities().CreateClassUsageAssetRegistryTag(this));
 
-	// Asset Library Tags
-	for(const FNiagaraAssetTagDefinitionReference& Tag : AssetTags)
-	{
-		Tag.AddTagToAssetRegistryTags(Context);
-	}
 #endif
 	Super::GetAssetRegistryTags(Context);
 }
diff --git a/Engine/Plugins/FX/Niagara/Source/Niagara/Private/NiagaraModule.cpp b/Engine/Plugins/FX/Niagara/Source/Niagara/Private/NiagaraModule.cpp
--- a/Engine/Plugins/FX/Niagara/Source/Niagara/Private/NiagaraModule.cpp
+++ b/Engine/Plugins/FX/Niagara/Source/Niagara/Private/NiagaraModule.cpp
@@ -40,6 +40,7 @@
 #include "Particles/FXBudget.h"
 #include "Engine/StaticMesh.h"
 #include "UObject/UObjectGlobals.h"
+#include "NiagaraAssetTagDefinitions.h"
 
 #include "DataInterface/NiagaraDataInterfaceDataChannelCommon.h"
 #include "DataInterface/NiagaraDataInterfaceDataChannelWrite.h"
@@ -66,6 +67,7 @@ static FAutoConsoleVariableRef CVarLogCompileIdGeneration(
 float INiagaraModule::EngineGlobalSpawnCountScale = 1.0f;
 float INiagaraModule::EngineGlobalSystemCountScale = 1.0f;
 
+PRAGMA_DISABLE_DEPRECATION_WARNINGS
 const FNiagaraAssetTagDefinition INiagaraModule::LightweightTagDefinition = {
 	LOCTEXT("LightweightDisplayName", "Lightweight")
 	, (int32) (ENiagaraAssetLibraryAssetTypes::Systems)
@@ -105,6 +107,7 @@ const FNiagaraAssetTagDefinition INiagaraModule::DeprecatedTagDefinition = {
 	, ENiagaraAssetTagDefinitionImportance::Internal
 	, FLinearColor::Black
 	, FGuid(0xF1F2F2FB, 0x50C64AF9, 0xB6BCE631, 0x32C4CC61)};
+PRAGMA_ENABLE_DEPRECATION_WARNINGS
 
 std::atomic<bool> INiagaraModule::bDataChannelRefreshRequested = false;
 
@@ -512,7 +515,9 @@ void INiagaraModule::StartupModule()
 #endif
 #endif
 
+PRAGMA_DISABLE_DEPRECATION_WARNINGS
 	RegisterInternalAssetTagDefinitions();
+PRAGMA_ENABLE_DEPRECATION_WARNINGS
 	
 	FCoreDelegates::OnPostEngineInit.AddRaw(this, &INiagaraModule::OnPostEngineInit);
 	FCoreDelegates::OnPreExit.AddRaw(this, &INiagaraModule::OnPreExit);
@@ -606,6 +611,13 @@ void INiagaraModule::RefreshDataChannels()
 	}
 }
 
+const TArray<const FNiagaraAssetTagDefinition*>& INiagaraModule::GetInternalAssetTagDefinitions() const
+{
+PRAGMA_DISABLE_DEPRECATION_WARNINGS
+	return InternalAssetTagDefinitions;
+PRAGMA_ENABLE_DEPRECATION_WARNINGS
+}
+
 void INiagaraModule::OnPostGarbageCollect()
 {
 	FNiagaraTypeHelper::TickTypeRemap();
@@ -1910,12 +1922,14 @@ FSHAHash FNiagaraEmitterParameters::GetStructHash()
 
 void INiagaraModule::RegisterInternalAssetTagDefinitions()
 {
+PRAGMA_DISABLE_DEPRECATION_WARNINGS
 	// All internal tags have to be manually assigned a unique and stable guid so that names can be changed without affecting assigned assets.
 	InternalAssetTagDefinitions.Add(&TemplateTagDefinition);
 	InternalAssetTagDefinitions.Add(&LightweightTagDefinition);
 	InternalAssetTagDefinitions.Add(&LearningContentTagDefinition);
 	InternalAssetTagDefinitions.Add(&HiddenAssetTagDefinition);
 	InternalAssetTagDefinitions.Add(&DeprecatedTagDefinition);
+PRAGMA_ENABLE_DEPRECATION_WARNINGS
 }
 
 #if NIAGARA_PERF_BASELINES
diff --git a/Engine/Plugins/FX/Niagara/Source/Niagara/Private/NiagaraScript.cpp b/Engine/Plugins/FX/Niagara/Source/Niagara/Private/NiagaraScript.cpp
--- a/Engine/Plugins/FX/Niagara/Source/Niagara/Private/NiagaraScript.cpp
+++ b/Engine/Plugins/FX/Niagara/Source/Niagara/Private/NiagaraScript.cpp
@@ -3425,13 +3425,6 @@ void UNiagaraScript::GetAssetRegistryTags(FAssetRegistryTagsContext Context) con
 	FString& SuggestedTag = CustomAssetRegistryTagCache->FindOrAdd(SuggestedName);
 	SuggestedTag = FString::FromInt(bSuggested);
 	Context.AddTag(FAssetRegistryTag(SuggestedName, SuggestedTag, FAssetRegistryTag::TT_Hidden));
-
-	// Asset Tags
-	TArray<FNiagaraAssetTagDefinitionReference> AssetTagsDefinitionReferences = ScriptData ? ScriptData->AssetTagDefinitionReferences : TArray<FNiagaraAssetTagDefinitionReference>();
-	for(const FNiagaraAssetTagDefinitionReference& TagDefinitionReference : AssetTagsDefinitionReferences)
-	{
-		TagDefinitionReference.AddTagToAssetRegistryTags(Context);
-	}
 	
 	// Add the current custom version to the tags so that tags can be fixed up in the future without having to load
 	// the whole asset.
diff --git a/Engine/Plugins/FX/Niagara/Source/Niagara/Private/NiagaraSystem.cpp b/Engine/Plugins/FX/Niagara/Source/Niagara/Private/NiagaraSystem.cpp
--- a/Engine/Plugins/FX/Niagara/Source/Niagara/Private/NiagaraSystem.cpp
+++ b/Engine/Plugins/FX/Niagara/Source/Niagara/Private/NiagaraSystem.cpp
@@ -1416,23 +1416,9 @@ void UNiagaraSystem::PostLoad()
 #endif
 
 #if WITH_EDITORONLY_DATA
-	// see the equivalent in NiagaraEmitter for details
-	ENiagaraScriptTemplateSpecification CurrentTemplateSpecification = TemplateSpecification_DEPRECATED;
-	if(bIsTemplateAsset_DEPRECATED)
+	if(NiagaraVer < FNiagaraCustomVersion::UserAssetTagStorageInMetaData && AssetTags_DEPRECATED.Num() > 0)
 	{
-		CurrentTemplateSpecification = ENiagaraScriptTemplateSpecification::Template;
-	}
-
-	if(NiagaraVer < FNiagaraCustomVersion::InheritanceUxRefactor)
-	{
-		if(CurrentTemplateSpecification == ENiagaraScriptTemplateSpecification::Template)
-		{
-			AssetTags.AddUnique(INiagaraModule::TemplateTagDefinition);
-		}
-		else if(CurrentTemplateSpecification == ENiagaraScriptTemplateSpecification::Behavior)
-		{
-			AssetTags.AddUnique(INiagaraModule::LearningContentTagDefinition);
-		}
+		INiagaraModule::Get().GetEditorOnlyDataUtilities().MoveUserAssetTagsToMetaData(this);
 	}
 #endif // WITH_EDITORONLY_DATA
 
@@ -1922,12 +1908,6 @@ void UNiagaraSystem::GetAssetRegistryTags(FAssetRegistryTagsContext Context) con
 		DI->GetClass()->GetName(ClassName);
 		Context.AddTag(FAssetRegistryTag(*(TEXT("bHas")+ClassName), TEXT("True"), FAssetRegistryTag::TT_Alphabetical));
 	}*/
-
-	// Asset Library Tags
-	for(const FNiagaraAssetTagDefinitionReference& Tag : AssetTags)
-	{
-		Tag.AddTagToAssetRegistryTags(Context);
-	}
 	
 	//Context.AddTag(FAssetRegistryTag("CPUCollision", UsesCPUCollision() ? TEXT("True") : TEXT("False"), FAssetRegistryTag::TT_Alphabetical));
 	//Context.AddTag(FAssetRegistryTag("Looping", bAnyEmitterLoopsForever ? TEXT("True") : TEXT("False"), FAssetRegistryTag::TT_Alphabetical));
diff --git a/Engine/Plugins/FX/Niagara/Source/Niagara/Public/INiagaraEditorOnlyDataUtlities.h b/Engine/Plugins/FX/Niagara/Source/Niagara/Public/INiagaraEditorOnlyDataUtlities.h
--- a/Engine/Plugins/FX/Niagara/Source/Niagara/Public/INiagaraEditorOnlyDataUtlities.h
+++ b/Engine/Plugins/FX/Niagara/Source/Niagara/Public/INiagaraEditorOnlyDataUtlities.h
@@ -35,4 +35,6 @@ public:
 	virtual TOptional<FNiagaraSystemStateData> GetSystemStateData(const UNiagaraSystem& System) const = 0;
 
 	virtual bool GetResetSimulationOnParameterChange() const = 0;
+	
+	virtual void MoveUserAssetTagsToMetaData(UObject* NiagaraAsset) const = 0;
 };
diff --git a/Engine/Plugins/FX/Niagara/Source/Niagara/Public/NiagaraAssetTagDefinitions.h b/Engine/Plugins/FX/Niagara/Source/Niagara/Public/NiagaraAssetTagDefinitions.h
--- a/Engine/Plugins/FX/Niagara/Source/Niagara/Public/NiagaraAssetTagDefinitions.h
+++ b/Engine/Plugins/FX/Niagara/Source/Niagara/Public/NiagaraAssetTagDefinitions.h
@@ -8,6 +8,8 @@
 
 #define UE_API NIAGARA_API
 
+/** DEPRECATED - All classes and data in here are deprecated in favor of the generic User Asset Tags system. */
+
 UENUM(meta=(Bitflags, UseEnumValuesAsMaskValuesInEditor=true))
 enum class ENiagaraAssetLibraryAssetTypes : uint8
 {
@@ -113,14 +115,15 @@ private:
 	FGuid AssetTagDefinitionGuid;
 };
 
-/** An Asset Tag Definition defines a tag that can be added to various Niagara assets for sorting & filtering purposes.
+/**
+ * An Asset Tag Definition defines a tag that can be added to various Niagara assets for sorting & filtering purposes.
  * For example, custom tags will show up in the Create Niagara System dialog to filter available emitters.
  * They can also be used to filter assets in the content browser, when used with the custom filter option.
  *
  * You can modify asset tags in the content browser by right-clicking on a Niagara asset, then use the "Manage Tags" submenu to add or remove them.
  */
-UCLASS(MinimalAPI)
-class UNiagaraAssetTagDefinitions : public UObject
+UCLASS(MinimalAPI, Deprecated)
+class UDEPRECATED_NiagaraAssetTagDefinitions : public UObject
 {
 	GENERATED_BODY()
 public:	
diff --git a/Engine/Plugins/FX/Niagara/Source/Niagara/Public/NiagaraModule.h b/Engine/Plugins/FX/Niagara/Source/Niagara/Public/NiagaraModule.h
--- a/Engine/Plugins/FX/Niagara/Source/Niagara/Public/NiagaraModule.h
+++ b/Engine/Plugins/FX/Niagara/Source/Niagara/Public/NiagaraModule.h
@@ -26,6 +26,8 @@ struct FNiagaraParameterStore;
 class FCommonViewportClient;
 class FNiagaraDebuggerClient;
 struct FNiagaraSystemAsyncCompileResults;
+struct FNiagaraAssetTagReference;
+struct FNiagaraAssetTagDefinition;
 
 extern NIAGARA_API int32 GEnableVerboseNiagaraChangeIdLogging;
 
@@ -136,7 +138,9 @@ public:
 	static void RequestRefreshDataChannels() { bDataChannelRefreshRequested = true; }
 	static NIAGARA_API void RefreshDataChannels();
 
-	const TArray<const FNiagaraAssetTagDefinition*>& GetInternalAssetTagDefinitions() { return InternalAssetTagDefinitions; }
+	UE_DEPRECATED(5.7, "Niagara Asset Tags are deprecated in favor of the generic User Asset Tags. Please update your code as this function will be removed. ")
+	NIAGARA_API const TArray<const FNiagaraAssetTagDefinition*>& GetInternalAssetTagDefinitions() const;
+	UE_DEPRECATED(5.7, "Niagara Asset Tags are deprecated in favor of the generic User Asset Tags. Please update your code as this function will be removed. ")
 	void RegisterInternalAssetTagDefinitions();
 	
 #if NIAGARA_PERF_BASELINES
@@ -296,13 +300,18 @@ public:
 	static NIAGARA_API bool bUseGlobalFXBudget;
 	static NIAGARA_API bool bDataChannelsEnabled;
 
+	UE_DEPRECATED(5.7, "Niagara Asset Tags are deprecated in favor of the generic User Asset Tags. Access the new User Asset Tags using FMetaData functions. Please update your code as this property will be removed.")
 	static const FNiagaraAssetTagDefinition LightweightTagDefinition;
-
+	UE_DEPRECATED(5.7, "Niagara Asset Tags are deprecated in favor of the generic User Asset Tags. Access the new User Asset Tags using FMetaData functions. Please update your code as this property will be removed.")
 	static const FNiagaraAssetTagDefinition TemplateTagDefinition;
+	UE_DEPRECATED(5.7, "Niagara Asset Tags are deprecated in favor of the generic User Asset Tags. Access the new User Asset Tags using FMetaData functions. Please update your code as this property will be removed.")
 	static const FNiagaraAssetTagDefinition LearningContentTagDefinition;
+	UE_DEPRECATED(5.7, "Niagara Asset Tags are deprecated in favor of the generic User Asset Tags. Access the new User Asset Tags using FMetaData functions. Please update your code as this property will be removed.")
 	static NIAGARA_API const FNiagaraAssetTagDefinition HiddenAssetTagDefinition;
+	UE_DEPRECATED(5.7, "Niagara Asset Tags are deprecated in favor of the generic User Asset Tags. Access the new User Asset Tags using FMetaData functions. Please update your code as this property will be removed.")
 	static NIAGARA_API const FNiagaraAssetTagDefinition DeprecatedTagDefinition;
 
+	UE_DEPRECATED(5.7, "Niagara Asset Tags are deprecated in favor of the generic User Asset Tags. Access the new User Asset Tags using FMetaData functions. Please update your code as this property will be removed.")
 	TArray<const FNiagaraAssetTagDefinition*> InternalAssetTagDefinitions;
 private:
 	static NIAGARA_API FNiagaraVariable Engine_WorldDeltaTime;
diff --git a/Engine/Plugins/FX/Niagara/Source/NiagaraCore/Public/NiagaraCustomVersion.h b/Engine/Plugins/FX/Niagara/Source/NiagaraCore/Public/NiagaraCustomVersion.h
--- a/Engine/Plugins/FX/Niagara/Source/NiagaraCore/Public/NiagaraCustomVersion.h
+++ b/Engine/Plugins/FX/Niagara/Source/NiagaraCore/Public/NiagaraCustomVersion.h
@@ -214,6 +214,8 @@ struct FNiagaraCustomVersion
 		SerializeUsageBitMaskToGPUFunctionInfo, // Make sure that we serialize the UsageBitMask with the GPUFunctionInfo
 		
 		AddedNewInterpolatedSpawnMode, // Changed the interpolated spawn property into an enum
+
+		UserAssetTagStorageInMetaData, // Adding user asset tag storage in UMetaData for all UPackages
 		
 		// DO NOT ADD A NEW VERSION UNLESS YOU HAVE TALKED TO THE NIAGARA LEAD. Mismanagement of these versions can lead to data loss if it is adjusted in multiple streams simultaneously.
 		// -----<new versions can be added above this line>  -------------------------------------------------
diff --git a/Engine/Plugins/FX/Niagara/Source/NiagaraEditor/Private/Customizations/UserParameters/NiagaraUserParametersBuilderBase.cpp b/Engine/Plugins/FX/Niagara/Source/NiagaraEditor/Private/Customizations/UserParameters/NiagaraUserParametersBuilderBase.cpp
--- a/Engine/Plugins/FX/Niagara/Source/NiagaraEditor/Private/Customizations/UserParameters/NiagaraUserParametersBuilderBase.cpp
+++ b/Engine/Plugins/FX/Niagara/Source/NiagaraEditor/Private/Customizations/UserParameters/NiagaraUserParametersBuilderBase.cpp
@@ -171,9 +171,10 @@ void FNiagaraUserParameterNodeBuilder::GenerateUserParameterRows(IDetailChildren
 	{
 		if(const UHierarchyCategory* HierarchyCategory = Cast<UHierarchyCategory>(Child))
 		{
-			if(HierarchyCategory->GetSection() != nullptr)
+			const FDataHierarchyElementMetaData_SectionAssociation* SectionAssociation = HierarchyCategory->FindMetaDataOfType<FDataHierarchyElementMetaData_SectionAssociation>();
+			if(SectionAssociation && SectionAssociation->Section.IsValid())
 			{
-				SectionCategoryMap.FindOrAdd(HierarchyCategory->GetSection()).Add(HierarchyCategory);	
+				SectionCategoryMap.FindOrAdd(SectionAssociation->Section.Get()).Add(HierarchyCategory);	
 			}
 		}
 	}
@@ -184,7 +185,7 @@ void FNiagaraUserParameterNodeBuilder::GenerateUserParameterRows(IDetailChildren
 		TSharedRef<SWrapBox> SectionsBox = SNew(SWrapBox).UseAllottedSize(true);
 
 		// maps don't guarantee order, so we iterate over the original section data instead
-		for(auto& Section : Root->GetSectionData())
+		for(const TObjectPtr<UHierarchySection>& Section : Root->GetSectionData())
 		{
 			// if the map doesn't contain at the entry, it means there are 0 categories for that section, so we skip it
 			if(!SectionCategoryMap.Contains(Section))
@@ -230,8 +231,8 @@ void FNiagaraUserParameterNodeBuilder::GenerateUserParameterRows(IDetailChildren
 							[
 								SNew(STextBlock)
 								.TextStyle(FAppStyle::Get(), "SmallText")
-								.Text_UObject(Section, &UHierarchySection::GetSectionNameAsText)
-								.ToolTipText_UObject(Section, &UHierarchySection::GetTooltip)
+								.Text_UObject(Section.Get(), &UHierarchySection::GetSectionNameAsText)
+								.ToolTipText_UObject(Section.Get(), &UHierarchySection::GetTooltip)
 							]
 						]
 					];
@@ -282,7 +283,16 @@ void FNiagaraUserParameterNodeBuilder::GenerateUserParameterRows(IDetailChildren
 	{		
 		if(const UHierarchyCategory* HierarchyCategory = Cast<UHierarchyCategory>(Child))
 		{
-			if((ActiveSection == nullptr || HierarchyCategory->GetSection() == ActiveSection) && HierarchyCategory->DoesOneChildExist<UNiagaraHierarchyUserParameter>(true))
+			const UHierarchySection* Section = nullptr;
+			if(const FDataHierarchyElementMetaData_SectionAssociation* SectionAssociation = HierarchyCategory->FindMetaDataOfType<FDataHierarchyElementMetaData_SectionAssociation>())
+			{
+				if(SectionAssociation->Section.IsValid())
+				{
+					Section = SectionAssociation->Section.Get();
+				}
+			}
+			
+			if((ActiveSection == nullptr || Section == ActiveSection) && HierarchyCategory->DoesOneChildExist<UNiagaraHierarchyUserParameter>(true))
 			{
 				GenerateCategory(ChildrenBuilder, *HierarchyCategory);
 			}
diff --git a/Engine/Plugins/FX/Niagara/Source/NiagaraEditor/Private/NiagaraEditorModule.cpp b/Engine/Plugins/FX/Niagara/Source/NiagaraEditor/Private/NiagaraEditorModule.cpp
--- a/Engine/Plugins/FX/Niagara/Source/NiagaraEditor/Private/NiagaraEditorModule.cpp
+++ b/Engine/Plugins/FX/Niagara/Source/NiagaraEditor/Private/NiagaraEditorModule.cpp
@@ -170,24 +170,21 @@
 
 #include "LevelEditor.h"
 #include "LevelEditorOutlinerSettings.h"
+#include "NiagaraAssetTagDefinitions.h"
 #include "Filters/CustomClassFilterData.h"
 
 #include "Widgets/SNiagaraDebugger.h"
-#include "Widgets/AssetBrowser/NiagaraAssetBrowserConfig.h"
-#include "Widgets/AssetBrowser/SNiagaraAssetBrowser.h"
 
 #include "NiagaraDebugVis.h"
 #include "NiagaraPerfBaseline.h"
 #include "NiagaraGraphDataCache.h"
 #include "NiagaraDecalRendererProperties.h"
-#include "NiagaraEditorMenuHelpers.h"
 #include "NiagaraLightRendererProperties.h"
 #include "NiagaraRibbonRendererProperties.h"
 #include "NiagaraSpriteRendererProperties.h"
 #include "NiagaraVolumeRendererProperties.h"
 
 #include "Engine/AssetManager.h"
-#include "Widgets/AssetBrowser/SNiagaraSelectedAssetDetails.h"
 #include "NiagaraRecentAndFavoritesManager.h"
 #include "Customizations/Stack/NiagaraStackObjectPropertyCustomization_StatelessModule_DynamicMaterialParameters.h"
 #include "Stateless/Modules/NiagaraStatelessModule_DynamicMaterialParameters.h"
@@ -196,6 +193,12 @@
 #include "TraversalCache/TraversalBuilder.h"
 #include "TraversalCache/TraversalCache.h"
 
+#include "UserAssetTagMenuHelpers.h"
+#include "Config/TaggedAssetBrowserConfig.h"
+#include "UserAssetTagEditorUtilities.h"
+#include "Widgets/STaggedAssetBrowser.h"
+#include "Widgets/SUserAssetTagBrowserSelectedAssetDetails.h"
+
 #include UE_INLINE_GENERATED_CPP_BY_NAME(NiagaraEditorModule)
 
 IMPLEMENT_MODULE( FNiagaraEditorModule, NiagaraEditor );
@@ -374,6 +377,72 @@ class FNiagaraEditorOnlyDataUtilities : public INiagaraEditorOnlyDataUtilities
 		const UNiagaraEditorSettings* NiagaraEditorSettings = GetDefault<UNiagaraEditorSettings>();
 		return NiagaraEditorSettings->GetResetSimulationOnChange();
 	}
+
+	virtual void MoveUserAssetTagsToMetaData(UObject* NiagaraAsset) const override
+	{
+		check(NiagaraAsset->IsA<UNiagaraEmitter>() || NiagaraAsset->IsA<UNiagaraSystem>());
+
+		TArray<FNiagaraAssetTagDefinitionReference>* References = nullptr;
+		if(UNiagaraEmitter* Emitter = Cast<UNiagaraEmitter>(NiagaraAsset))
+		{
+			References = &Emitter->AssetTags_DEPRECATED;
+		}
+		else if(UNiagaraSystem* System = Cast<UNiagaraSystem>(NiagaraAsset))
+		{
+			References = &System->AssetTags_DEPRECATED;
+		}
+
+		if(References && References->Num() > 0)
+		{
+			TMap<FGuid, const FNiagaraAssetTagDefinition*> AllTags;
+			FAssetRegistryModule& AssetRegistryModule = FModuleManager::Get().LoadModuleChecked<FAssetRegistryModule>("AssetRegistry");
+
+			FARFilter Filter;
+			Filter.ClassPaths.Add(UDEPRECATED_NiagaraAssetTagDefinitions::StaticClass()->GetClassPathName());
+
+PRAGMA_DISABLE_DEPRECATION_WARNINGS
+			for(const FNiagaraAssetTagDefinition* InternalTagDefinition : INiagaraModule::Get().GetInternalAssetTagDefinitions())
+			{
+				AllTags.Add(InternalTagDefinition->TagGuid, InternalTagDefinition);
+			}
+PRAGMA_ENABLE_DEPRECATION_WARNINGS
+
+			TArray<FAssetData> AssetTagDefinitionAssets;
+			AssetRegistryModule.Get().GetAssets(Filter, AssetTagDefinitionAssets);
+
+			for(const FAssetData& AssetTagDefinitionAsset : AssetTagDefinitionAssets)
+			{
+				UDEPRECATED_NiagaraAssetTagDefinitions* TagDefinitionsAsset = Cast<UDEPRECATED_NiagaraAssetTagDefinitions>(AssetTagDefinitionAsset.GetAsset());
+
+				for(const FNiagaraAssetTagDefinition& TagDefinition : TagDefinitionsAsset->GetAssetTagDefinitions())
+				{
+					AllTags.Add(TagDefinition.TagGuid, &TagDefinition);
+				}
+			}
+
+			TArray<const FNiagaraAssetTagDefinition*> FoundTagDefinitions;
+			for(const FNiagaraAssetTagDefinitionReference& AssetTagDefinitionReference : *References)
+			{
+				FGuid TagGuid = AssetTagDefinitionReference.GetTagDefinitionReferenceGuid();
+				if(AllTags.Contains(TagGuid))
+				{
+					FoundTagDefinitions.Add(AllTags[TagGuid]);
+				}
+			}
+
+			for(const FNiagaraAssetTagDefinition* Definition : FoundTagDefinitions)
+			{
+				if(Definition == nullptr)
+				{
+					continue;
+				}
+
+				UE::UserAssetTags::AddUserAssetTag(NiagaraAsset, FName(Definition->AssetTag.ToString()));
+			}
+
+			References->Empty();
+		}
+	}
 };
 
 class FNiagaraScriptGraphPanelPinFactory : public FGraphPanelPinFactory
@@ -1108,16 +1177,13 @@ void FNiagaraEditorModule::StartupModule()
 
 	RecentAndFavoritesManager = MakeShared<FNiagaraRecentAndFavoritesManager>();
 	RecentAndFavoritesManager->Initialize();
-
-	UNiagaraAssetBrowserConfig::Initialize();
-
 	// We have to call shutdown before the exit actually happens to be able to save to disk
 	FEditorDelegates::OnEditorPreExit.AddLambda([]()
 	{
 		UNiagaraFavoriteActionsConfig::Shutdown();
 	});
 	
-	FNiagaraAssetDetailDatabase::Init();
+	InitializeTaggedAssetBrowserDetails();
 
 	TempPackage = NewObject<UPackage>(nullptr, TEXT("/Temp/NiagaraEditor"), RF_Transient);
 	TempPackage->AddToRoot();
@@ -1357,9 +1423,6 @@ void FNiagaraEditorModule::StartupModule()
 
 	RegisterDefaultStackObjectCustomizations();
 	
-	NiagaraEditorMenuHelpers::RegisterToolMenus();
-	NiagaraEditorMenuHelpers::RegisterMenuExtensions();
-
 	FNiagaraEditorStyle::Register();
 	ReinitializeStyleCommand = IConsoleManager::Get().RegisterConsoleCommand(
 		TEXT("fx.NiagaraEditor.ReinitializeStyle"),
@@ -2339,6 +2402,151 @@ void FNiagaraEditorModule::UnregisterSettings()
 	}
 }
 
+void FNiagaraEditorModule::InitializeTaggedAssetBrowserDetails()
+{	
+	// Now we add information on how to display emitters & systems in the tagged asset browser
+	// Emitters
+	{
+		FAssetDetailsDisplayInfo EmitterClassInfo;
+	
+		// Inheritance
+		FDisplayedPropertyData InheritanceProperty;
+		InheritanceProperty.ShouldDisplayPropertyDelegate = FDisplayedPropertyData::FShouldDisplayProperty::CreateLambda([](const FAssetData& AssetData)
+		{
+			if(AssetData.FindTag(GET_MEMBER_NAME_CHECKED(UNiagaraEmitter, bIsInheritable)))
+			{
+				return true;
+			}
+			
+			return false; 
+		});
+		InheritanceProperty.NameWidgetDelegate = FDisplayedPropertyData::FGenerateWidget::CreateLambda([](const FAssetData& AssetData) -> TSharedRef<SWidget>
+		{
+			return SNew(STextBlock)
+			.Text(LOCTEXT("Inheritance", "Inheritance"));
+		});
+		InheritanceProperty.ValueWidgetDelegate = FDisplayedPropertyData::FGenerateWidget::CreateLambda([](const FAssetData& AssetData) -> TSharedRef<SWidget>
+		{
+			bool bUseInheritance = false;
+			if(FNiagaraEditorUtilities::GetIsInheritableFromAssetRegistryTags(AssetData, bUseInheritance))
+			{
+				return SNew(STextBlock).Text(bUseInheritance
+					? LOCTEXT("Emitter_UseInheritance_Yes", "Yes")
+					: LOCTEXT("Emitter_UseInheritance_No", "No")
+					);
+			}
+	
+			return SNullWidget::NullWidget;
+		});
+		
+		// Simulation Target
+		FName SimulationTargetTag("HasGPUEmitter");
+	
+		FDisplayedPropertyData SimulationTargetProperty;
+		SimulationTargetProperty.ShouldDisplayPropertyDelegate = FDisplayedPropertyData::FShouldDisplayProperty::CreateLambda([SimulationTargetTag](const FAssetData& Asset)
+		{
+			return Asset.FindTag(SimulationTargetTag);
+		});
+		SimulationTargetProperty.NameWidgetDelegate = FDisplayedPropertyData::FGenerateWidget::CreateLambda([](const FAssetData& AssetData) -> TSharedRef<SWidget>
+		{
+			return SNew(STextBlock)
+			.Text(LOCTEXT("SimulationTargetLabel", "Runs on"));
+		});
+		SimulationTargetProperty.ValueWidgetDelegate = FDisplayedPropertyData::FGenerateWidget::CreateLambda([SimulationTargetTag](const FAssetData& Asset)
+		{
+			return SNew(SImage)
+			.Image_Lambda([Asset, SimulationTargetTag]() -> const FSlateBrush*
+			{
+				if(Asset.FindTag(SimulationTargetTag))
+				{
+					FString Value;
+					if(Asset.GetTagValue(SimulationTargetTag, Value))
+					{
+						if(Value == "True")
+						{
+							return FNiagaraEditorStyle::Get().GetBrush("NiagaraEditor.Stack.GPUIcon");
+						}
+						else
+						{
+							return FNiagaraEditorStyle::Get().GetBrush("NiagaraEditor.Stack.CPUIcon");
+						}
+					}
+				}
+	
+				return FAppStyle::GetNoBrush();
+			});
+		});
+	
+		EmitterClassInfo.DisplayedProperties =
+		{
+			InheritanceProperty,
+			SimulationTargetProperty
+		};
+		
+		EmitterClassInfo.GetDescriptionDelegate = FAssetDetailsDisplayInfo::FGetDescription::CreateLambda([](const FAssetData& AssetData)
+		{
+			FName DescriptionTagName("TemplateAssetDescription");
+			if(AssetData.FindTag(DescriptionTagName))
+			{
+				FString Description;
+				AssetData.GetTagValue(DescriptionTagName, Description);
+				return FText::FromString(Description);
+			}
+	
+			return FText::GetEmpty();
+		});
+
+		FTaggedAssetBrowserDetailsDisplayDatabase::RegisterClass(UNiagaraEmitter::StaticClass(), EmitterClassInfo);
+	}
+	
+	// Systems
+	{
+		FAssetDetailsDisplayInfo SystemClassInfo;
+	
+		// Number of emitters
+		FDisplayedPropertyData NumEmittersProperty;
+		NumEmittersProperty.ShouldDisplayPropertyDelegate = FDisplayedPropertyData::FShouldDisplayProperty::CreateLambda([](const FAssetData& AssetData) -> bool
+		{
+			return AssetData.FindTag("NumEmitters");
+		});
+		NumEmittersProperty.NameWidgetDelegate = FDisplayedPropertyData::FGenerateWidget::CreateLambda([](const FAssetData& AssetData) -> TSharedRef<SWidget>
+		{
+			return SNew(STextBlock)
+			.Text(LOCTEXT("NumberOfEmitters", "Number of Emitters"));
+		});
+		NumEmittersProperty.ValueWidgetDelegate =  FDisplayedPropertyData::FGenerateWidget::CreateLambda([](const FAssetData& AssetData) -> TSharedRef<SWidget>
+		{
+			int32 NumEmitters = INDEX_NONE;
+			if(AssetData.GetTagValue("NumEmitters", NumEmitters))
+			{
+				return SNew(STextBlock).Text(FText::AsNumber(NumEmitters));
+			}
+	
+			return SNullWidget::NullWidget;
+		});
+		
+		SystemClassInfo.DisplayedProperties =
+		{
+			NumEmittersProperty
+		};
+		
+		SystemClassInfo.GetDescriptionDelegate = FAssetDetailsDisplayInfo::FGetDescription::CreateLambda([](const FAssetData& AssetData)
+		{
+			FName DescriptionTagName("TemplateAssetDescription");
+			if(AssetData.FindTag(DescriptionTagName))
+			{
+				FString Description;
+				AssetData.GetTagValue(DescriptionTagName, Description);
+				return FText::FromString(Description);
+			}
+		
+			return FText::GetEmpty();
+		});
+		
+		FTaggedAssetBrowserDetailsDisplayDatabase::RegisterClass(UNiagaraSystem::StaticClass(), SystemClassInfo);
+	}
+}
+
 void FNiagaraEditorModule::AddReferencedObjects(FReferenceCollector& Collector)
 {
 	if (SequencerSettings)
diff --git a/Engine/Plugins/FX/Niagara/Source/NiagaraEditor/Private/NiagaraEditorStyle.cpp b/Engine/Plugins/FX/Niagara/Source/NiagaraEditor/Private/NiagaraEditorStyle.cpp
--- a/Engine/Plugins/FX/Niagara/Source/NiagaraEditor/Private/NiagaraEditorStyle.cpp
+++ b/Engine/Plugins/FX/Niagara/Source/NiagaraEditor/Private/NiagaraEditorStyle.cpp
@@ -106,31 +106,6 @@ void FNiagaraEditorStyle::InitAssetPicker()
 
 }
 
-void FNiagaraEditorStyle::InitAssetBrowser()
-{
-	const FTextBlockStyle NormalText = FAppStyle::GetWidgetStyle<FTextBlockStyle>("NormalText");
-
-	Set("NiagaraEditor.AssetBrowser.PropertySeparator", new FSlateColorBrush(FStyleColors::White25));
-
-	FTextBlockStyle AssetBrowserAssetTitleStyle = FAppStyle::GetWidgetStyle<FTextBlockStyle>("NormalText.Important");
-	AssetBrowserAssetTitleStyle.Font.Size = 12.f;
-	Set("NiagaraEditor.AssetBrowser.AssetTitle", AssetBrowserAssetTitleStyle);
-	
-	FTextBlockStyle AssetBrowserAssetTypeStyle = FAppStyle::GetWidgetStyle<FTextBlockStyle>("NormalText.Subdued");
-	AssetBrowserAssetTypeStyle.Font.Size = 10.f;
-	Set("NiagaraEditor.AssetBrowser.AssetType", AssetBrowserAssetTypeStyle);
-	
-	Set("NiagaraEditor.AssetBrowser.AssetTag.OuterBorder", new FSlateRoundedBoxBrush(FStyleColors::AccentGray, 4.f));
-	Set("NiagaraEditor.AssetBrowser.AssetTag.InnerBorder", new FSlateRoundedBoxBrush(FStyleColors::Black, 4.f));
-
-	FTextBlockStyle AssetTagText = NormalText;
-
-	FSlateFontInfo Font = FStyleFonts::Get().Normal;
-	Font.Size = 10.f;
-	AssetTagText.SetFont(Font);
-	Set("NiagaraEditor.AssetBrowser.AssetTag.Text", AssetTagText);
-}
-
 void FNiagaraEditorStyle::InitActionMenu()
 {
 	const FTextBlockStyle NormalText = FAppStyle::GetWidgetStyle<FTextBlockStyle>("NormalText");
@@ -899,7 +874,6 @@ FNiagaraEditorStyle::FNiagaraEditorStyle() : FSlateStyleSet("NiagaraEditorStyle"
 
 	InitStats();
 	InitAssetPicker();
-	InitAssetBrowser();
 	InitActionMenu();
 	InitEmitterHeader();
 	InitParameters();
diff --git a/Engine/Plugins/FX/Niagara/Source/NiagaraEditor/Private/NiagaraEditorUtilities.cpp b/Engine/Plugins/FX/Niagara/Source/NiagaraEditor/Private/NiagaraEditorUtilities.cpp
--- a/Engine/Plugins/FX/Niagara/Source/NiagaraEditor/Private/NiagaraEditorUtilities.cpp
+++ b/Engine/Plugins/FX/Niagara/Source/NiagaraEditor/Private/NiagaraEditorUtilities.cpp
@@ -5283,213 +5283,6 @@ TMap<FNiagaraVariableBase, FGuid> FNiagaraEditorUtilities::Scripts::Validation::
 	return RemappedGuids;
 }
 
-TArray<FNiagaraEditorUtilities::AssetBrowser::FStructuredAssetTagDefinitionLookupData> FNiagaraEditorUtilities::AssetBrowser::GetStructuredSortedAssetTagDefinitions(bool bSortTagsByName)
-{
-	TArray<FStructuredAssetTagDefinitionLookupData> Result;
-	
-	// Internal Tags
-	{
-		FStructuredAssetTagDefinitionLookupData TagDefinitionData;
-		TagDefinitionData.DefinitionsAsset = nullptr;
-
-		for(const FNiagaraAssetTagDefinition* InternalAssetTagDefinition : INiagaraModule::Get().GetInternalAssetTagDefinitions())
-		{
-			TagDefinitionData.AssetTagDefinitions.Add(*InternalAssetTagDefinition);
-		}
-		
-		if(bSortTagsByName)
-		{
-			TagDefinitionData.AssetTagDefinitions.Sort();
-		}
-		TagDefinitionData.Source = EAssetTagSectionSource::NiagaraInternal;
-		Result.Add(TagDefinitionData);
-	}
-
-	FAssetRegistryModule& AssetRegistryModule = FModuleManager::Get().LoadModuleChecked<FAssetRegistryModule>("AssetRegistry");
-	FARFilter Filter;
-	Filter.ClassPaths.Add(UNiagaraAssetTagDefinitions::StaticClass()->GetClassPathName());
-
-	TArray<FAssetData> FilteredTagDefinitionAssets;
-	AssetRegistryModule.Get().GetAssets(Filter, FilteredTagDefinitionAssets);
-
-	FAssetToolsModule& AssetToolsModule = FModuleManager::LoadModuleChecked<FAssetToolsModule>("AssetTools");
-	for(const FAssetData& AssetData : FilteredTagDefinitionAssets)
-	{
-		if(AssetToolsModule.Get().IsAssetVisible(AssetData, true) == false)
-		{
-			continue;
-		}
-		
-		UNiagaraAssetTagDefinitions* TagDefinitions = Cast<UNiagaraAssetTagDefinitions>(AssetData.GetAsset());
-		FStructuredAssetTagDefinitionLookupData TagDefinitionData;
-		TagDefinitionData.DefinitionsAsset = TagDefinitions;
-		TagDefinitionData.AssetTagDefinitions = TagDefinitions->GetAssetTagDefinitions();
-		if(bSortTagsByName)
-		{
-			TagDefinitionData.AssetTagDefinitions.Sort();
-		}
-		TagDefinitionData.Source = GetAssetTagDefinitionSource(AssetData);
-		Result.Add(TagDefinitionData);
-	}
-
-	Result.Sort([](const FStructuredAssetTagDefinitionLookupData& CandidateA, const FStructuredAssetTagDefinitionLookupData& CandidateB)
-	{
-		if((int32) CandidateA.Source != (int32) CandidateB.Source)
-		{
-			return (int32) CandidateA.Source < (int32) CandidateB.Source;
-		}
-
-		if(CandidateA.DefinitionsAsset == nullptr || CandidateB.DefinitionsAsset == nullptr)
-		{
-			if(CandidateB.DefinitionsAsset != nullptr)
-			{
-				return true;
-			}
-			else if(CandidateA.DefinitionsAsset != nullptr)
-			{
-				return false;
-			}
-			
-			return CandidateA.AssetTagDefinitions.Num() < CandidateB.AssetTagDefinitions.Num();
-		}
-		
-		return CandidateA.DefinitionsAsset->GetSortOrder() < CandidateB.DefinitionsAsset->GetSortOrder();
-	});
-	
-	return Result;
-}
-
-TArray<FNiagaraAssetTagDefinition> FNiagaraEditorUtilities::AssetBrowser::GetFlatSortedAssetTagDefinitions(bool bSortEndResultTagsByName)
-{
-	TArray<FNiagaraAssetTagDefinition> Result;
-	TArray<FStructuredAssetTagDefinitionLookupData> IntermediateResult;
-	FAssetRegistryModule& AssetRegistryModule = FModuleManager::Get().LoadModuleChecked<FAssetRegistryModule>("AssetRegistry");
-
-	for(const FNiagaraAssetTagDefinition* InternalAssetTagDefinition : INiagaraModule::Get().GetInternalAssetTagDefinitions())
-	{
-		Result.Add(*InternalAssetTagDefinition);
-	}
-
-	FARFilter Filter;
-	Filter.ClassPaths.Add(UNiagaraAssetTagDefinitions::StaticClass()->GetClassPathName());
-
-	TArray<FAssetData> AllTagDefinitionAssets;
-	AssetRegistryModule.Get().GetAssets(Filter, AllTagDefinitionAssets);
-
-	for(const FAssetData& AssetData : AllTagDefinitionAssets)
-	{
-		UNiagaraAssetTagDefinitions* TagDefinitions = Cast<UNiagaraAssetTagDefinitions>(AssetData.GetAsset());
-		FStructuredAssetTagDefinitionLookupData TagDefinitionData;
-		TagDefinitionData.DefinitionsAsset = TagDefinitions;
-		TagDefinitionData.AssetTagDefinitions = TagDefinitions->GetAssetTagDefinitions();
-		TagDefinitionData.Source = GetAssetTagDefinitionSource(AssetData);
-		IntermediateResult.Add(TagDefinitionData);
-	}
-	
-	IntermediateResult.Sort([](const FStructuredAssetTagDefinitionLookupData& CandidateA, const FStructuredAssetTagDefinitionLookupData& CandidateB)
-	{
-		if((int32) CandidateA.Source != (int32) CandidateB.Source)
-		{
-			return (int32) CandidateA.Source < (int32) CandidateB.Source;
-		}
-
-		if(CandidateA.DefinitionsAsset == nullptr || CandidateB.DefinitionsAsset == nullptr)
-		{
-			if(CandidateB.DefinitionsAsset != nullptr)
-			{
-				return true;
-			}
-			else if(CandidateA.DefinitionsAsset != nullptr)
-			{
-				return false;
-			}
-			
-			return CandidateA.AssetTagDefinitions.Num() < CandidateB.AssetTagDefinitions.Num();
-		}
-		
-		return CandidateA.DefinitionsAsset->GetSortOrder() < CandidateB.DefinitionsAsset->GetSortOrder();
-	});
-
-	for(const FStructuredAssetTagDefinitionLookupData& IntermediateData : IntermediateResult)
-	{
-		Result.Append(IntermediateData.AssetTagDefinitions);
-	}
-
-	if(bSortEndResultTagsByName)
-	{
-		Result.Sort();
-	}
-	
-	return Result;
-}
-
-const FNiagaraAssetTagDefinition& FNiagaraEditorUtilities::AssetBrowser::FindTagDefinitionForReference(const FNiagaraAssetTagDefinitionReference& Reference)
-{
-	TMap<FGuid, FNiagaraAssetTagDefinition> Result;
-	FAssetRegistryModule& AssetRegistryModule = FModuleManager::Get().LoadModuleChecked<FAssetRegistryModule>("AssetRegistry");
-
-	FARFilter Filter;
-	Filter.ClassPaths.Add(UNiagaraAssetTagDefinitions::StaticClass()->GetClassPathName());
-
-	const TArray<const FNiagaraAssetTagDefinition*>& InternalTagDefinitions = INiagaraModule::Get().GetInternalAssetTagDefinitions();
-	
-	for(const FNiagaraAssetTagDefinition* InternalTagDefinition : InternalTagDefinitions)
-	{
-		Result.Add(InternalTagDefinition->TagGuid, *InternalTagDefinition);
-	}
-
-	TArray<FAssetData> AssetTagDefinitionAssets;
-	AssetRegistryModule.Get().GetAssets(Filter, AssetTagDefinitionAssets);
-
-	for(const FAssetData& AssetTagDefinitionAsset : AssetTagDefinitionAssets)
-	{
-		UNiagaraAssetTagDefinitions* TagDefinitionsAsset = Cast<UNiagaraAssetTagDefinitions>(AssetTagDefinitionAsset.GetAsset());
-
-		for(const FNiagaraAssetTagDefinition& TagDefinition : TagDefinitionsAsset->GetAssetTagDefinitions())
-		{
-			Result.Add(TagDefinition.TagGuid, TagDefinition);
-		}
-	}
-
-	if(Result.Contains(Reference.GetTagDefinitionReferenceGuid()))
-	{
-		return Result[Reference.GetTagDefinitionReferenceGuid()];
-	}
-
-	static FNiagaraAssetTagDefinition DummyTagDefinition;
-	return DummyTagDefinition;
-}
-
-FNiagaraEditorUtilities::AssetBrowser::EAssetTagSectionSource FNiagaraEditorUtilities::AssetBrowser::GetAssetTagDefinitionSource(const FAssetData& AssetData)
-{	
-	ensure(AssetData.GetClass() == UNiagaraAssetTagDefinitions::StaticClass());
-
-	if(AssetData.PackagePath.ToString().StartsWith("/Niagara"))
-	{
-		return EAssetTagSectionSource::NiagaraInternal;
-	}
-	else if(AssetData.PackagePath.ToString().StartsWith("/Game"))
-	{
-		return EAssetTagSectionSource::Project;
-	}
-
-	return EAssetTagSectionSource::Other;
-}
-
-FText FNiagaraEditorUtilities::AssetBrowser::GetAssetTagSectionNameFromSource(EAssetTagSectionSource Source)
-{
-	if(Source == EAssetTagSectionSource::NiagaraInternal)
-	{
-		return FText::FromString("Niagara");
-	}
-	else if(Source == EAssetTagSectionSource::Project)
-	{
-		return FText::FromString(FApp::GetProjectName());
-	}
-
-	return LOCTEXT("NiagaraAssetTagSourceSectionName", "Misc");
-}
-
 void FNiagaraParameterUtilities::FilterToRelevantStaticVariables(TConstArrayView<FNiagaraVariable> InVars, TArray<FNiagaraVariable>& OutVars, FName InOldEmitterAlias, FName InNewEmitterAlias, bool bFilterByEmitterAliasAndConvertToUnaliased)
 {
 	FNiagaraAliasContext RenameContext(FNiagaraAliasContext::ERapidIterationParameterMode::StaticVariables);
diff --git a/Engine/Plugins/FX/Niagara/Source/NiagaraEditor/Private/NiagaraEmitterEditorData.cpp b/Engine/Plugins/FX/Niagara/Source/NiagaraEditor/Private/NiagaraEmitterEditorData.cpp
--- a/Engine/Plugins/FX/Niagara/Source/NiagaraEditor/Private/NiagaraEmitterEditorData.cpp
+++ b/Engine/Plugins/FX/Niagara/Source/NiagaraEditor/Private/NiagaraEmitterEditorData.cpp
@@ -91,19 +91,26 @@ void UNiagaraEmitterEditorData::PostLoad_TransferSummaryDataToNewFormat()
 		for(FNiagaraStackSection& StackSection : SummarySections_DEPRECATED)
 		{
 			FText SectionDisplayName = StackSection.SectionDisplayName.IsEmptyOrWhitespace() ? FText::FromName(StackSection.SectionIdentifier) : StackSection.SectionDisplayName;
-			UHierarchySection* Section = SummaryViewRoot->AddSection(SectionDisplayName);
+			UHierarchySection* Section = NewObject<UHierarchySection>(SummaryViewRoot, UHierarchySection::StaticClass());
+			Section->SetSectionNameAsText(SectionDisplayName);
+			
 			FHierarchyElementIdentity SectionIdentity;
 			SectionIdentity.Names.Add(StackSection.SectionIdentifier);
 			Section->SetIdentity(SectionIdentity);
 			
+			SummaryViewRoot->GetSectionDataMutable().Insert(Section, 0);
+			
 			Sections.Add(StackSection.SectionIdentifier, Section);
 			
 			for(const FText& CategoryIdentity : StackSection.Categories)
 			{
 				FName CategoryName = FName(CategoryIdentity.ToString());
 				UHierarchyCategory* HierarchyCategory = SummaryViewRoot->AddChild<UHierarchyCategory>();
-				HierarchyCategory->SetSection(Section);
 				HierarchyCategory->SetCategoryName(CategoryName);
+				
+				FDataHierarchyElementMetaData_SectionAssociation* SectionAssociation = HierarchyCategory->FindOrAddMetaDataOfType<FDataHierarchyElementMetaData_SectionAssociation>();
+				SectionAssociation->Section = Section;
+				
 				IdentityCategoryMap.Add(CategoryName, HierarchyCategory);
 				SectionCategoryMapping.Add(CategoryName, Section->GetSectionName());
 			}
@@ -498,7 +505,7 @@ FSimpleMulticastDelegate& UNiagaraEmitterEditorData::OnSummaryViewStateChanged()
 	return OnSummaryViewStateChangedDelegate;
 }
 
-const TArray<UHierarchySection*>& UNiagaraEmitterEditorData::GetSummarySections() const
+TConstArrayView<const TObjectPtr<UHierarchySection>> UNiagaraEmitterEditorData::GetSummarySections() const
 {
 	return SummaryViewRoot->GetSectionData();
 }
diff --git a/Engine/Plugins/FX/Niagara/Source/NiagaraEditor/Private/NiagaraEmitterFactoryNew.cpp b/Engine/Plugins/FX/Niagara/Source/NiagaraEditor/Private/NiagaraEmitterFactoryNew.cpp
--- a/Engine/Plugins/FX/Niagara/Source/NiagaraEditor/Private/NiagaraEmitterFactoryNew.cpp
+++ b/Engine/Plugins/FX/Niagara/Source/NiagaraEditor/Private/NiagaraEmitterFactoryNew.cpp
@@ -15,11 +15,15 @@
 #include "NiagaraEmitterEditorData.h"
 #include "NiagaraNodeAssignment.h"
 #include "NiagaraRecentAndFavoritesManager.h"
+#include "UserAssetTagEditorUtilities.h"
+#include "Assets/TaggedAssetBrowserConfiguration.h"
 #include "Misc/MessageDialog.h"
 #include "Modules/ModuleManager.h"
 #include "Interfaces/IMainFrameModule.h"
 #include "Framework/Application/SlateApplication.h"
-#include "Widgets/AssetBrowser/SNiagaraAssetBrowser.h"
+#include "Widgets/STaggedAssetBrowser.h"
+#include "Widgets/STaggedAssetBrowserCreateAssetWindow.h"
+#include "Widgets/AssetBrowser/SNiagaraAddEmitterToSystemWindow.h"
 
 #include UE_INLINE_GENERATED_CPP_BY_NAME(NiagaraEmitterFactoryNew)
 
@@ -37,37 +41,66 @@ UNiagaraEmitterFactoryNew::UNiagaraEmitterFactoryNew(const FObjectInitializer& O
 }
 
 bool UNiagaraEmitterFactoryNew::ConfigureProperties()
-{
-	IMainFrameModule& MainFrame = FModuleManager::LoadModuleChecked<IMainFrameModule>("MainFrame");
-	TSharedPtr<SWindow>	ParentWindow = MainFrame.GetParentWindow();
-
-	SNiagaraAssetBrowser::FArguments AssetBrowserArgs;
-	AssetBrowserArgs.AvailableClasses({UNiagaraEmitter::StaticClass()})
-	.RecentAndFavoritesList(FNiagaraEditorModule::Get().GetRecentsManager()->GetRecentEmitterAndSystemsList())
-	.EmptySelectionMessage(LOCTEXT("EmptyEmitterFactorySelectionMessage", "Select an emitter as a starting point for your new emitter.\n"));
-
-	SNiagaraAssetBrowserWindow::FArguments AssetBrowserWindowArgs;
-	AssetBrowserWindowArgs.AssetBrowserArgs(AssetBrowserArgs)
-	.WindowTitle(LOCTEXT("EmitterAssetBrowserWindowTitle", "Create Niagara Emitter - Select an emitter as a base"));
+{	
+	TArray<FAssetData> SelectedAssetData;
 	
-	TSharedRef<SNiagaraCreateAssetWindow> CreateAssetBrowserWindow = SNew(SNiagaraCreateAssetWindow, *UNiagaraEmitter::StaticClass()).AssetBrowserWindowArgs(AssetBrowserWindowArgs);
-	
-	FSlateApplication::Get().AddModalWindow(CreateAssetBrowserWindow, ParentWindow);
-
-	if(CreateAssetBrowserWindow->ShouldProceedWithAction() == false)
+	FSoftObjectPath Path = GetDefault<UNiagaraEditorSettings>()->EmitterAssetWizardConfiguration;
+    UObject* ConfigUObject = Path.TryLoad();
+	if(UTaggedAssetBrowserConfiguration* ConfigurationAsset = Cast<UTaggedAssetBrowserConfiguration>(ConfigUObject))
 	{
-		return false;	
-	}
+		FDefaultDetailsTabConfiguration DefaultDetailsTabConfiguration;
+		DefaultDetailsTabConfiguration.bUseDefaultDetailsTab = true;
+		DefaultDetailsTabConfiguration.EmptySelectionMessage = LOCTEXT("EmptyEmitterFactorySelectionMessage", "Select an emitter as a starting point for your new emitter.\n");
+ 
+		IMainFrameModule& MainFrame = FModuleManager::LoadModuleChecked<IMainFrameModule>("MainFrame");
+		TSharedPtr<SWindow>	ParentWindow = MainFrame.GetParentWindow();
+
+		STaggedAssetBrowser::FInterfaceOverrideProfiles OverrideProfiles;
+		OverrideProfiles.AssetViewOptionsProfileName = FName("TaggedAssetBrowser");
+		OverrideProfiles.DefaultFilterMenuExpansion = EAssetTypeCategories::FX;
+		OverrideProfiles.FilterBarSaveName = FName("TaggedAssetBrowser.NiagaraEmitterFactory");
+		
+		STaggedAssetBrowser::FArguments AssetBrowserArgs;
+		AssetBrowserArgs
+			.AvailableClasses({ UNiagaraEmitter::StaticClass() })
+			.DefaultDetailsTabConfiguration(DefaultDetailsTabConfiguration)
+			.InterfaceOverrideProfiles(OverrideProfiles)
+			.RecentAndFavoritesList(FNiagaraEditorModule::Get().GetRecentsManager()->GetRecentEmitterAndSystemsList())
+			.OnAdditionalShouldFilterAsset_UObject(this, &UNiagaraEmitterFactoryNew::OnAdditionalShouldFilterAsset)
+			.OnGetExtraFrontendFilters_UObject(this, &UNiagaraEmitterFactoryNew::OnGetExtraFrontendFilters)
+			.OnExtendAddFilterMenu_UObject(this, &UNiagaraEmitterFactoryNew::OnExtendAddFilterMenu);
+
+		SWindow::FArguments WindowArgs;
+		WindowArgs.Title(LOCTEXT("EmitterAssetBrowserWindowTitle", "Create Niagara Emitter - Select an emitter as a base"));
+		WindowArgs.SupportsMaximize(false);
+		WindowArgs.SupportsMinimize(false);
+		WindowArgs.ClientSize(FVector2D(1400, 750));
+		WindowArgs.SizingRule(ESizingRule::UserSized);
+		
+		STaggedAssetBrowserWindow::FArguments AssetBrowserWindowArgs;
+		AssetBrowserWindowArgs.AssetBrowserArgs(AssetBrowserArgs);
+		AssetBrowserWindowArgs.WindowArgs(WindowArgs);
+		
+		TSharedRef<STaggedAssetBrowserCreateAssetWindow> CreateAssetBrowserWindow = SNew(STaggedAssetBrowserCreateAssetWindow, *ConfigurationAsset, *UNiagaraEmitter::StaticClass())
+			.AssetBrowserWindowArgs(AssetBrowserWindowArgs);
+		
+		FSlateApplication::Get().AddModalWindow(CreateAssetBrowserWindow, ParentWindow);
 
-	TArray<FAssetData> SelectedAssetData = CreateAssetBrowserWindow->GetSelectedAssets();
+		if(CreateAssetBrowserWindow->ShouldProceedWithAction() == false)
+		{
+			return false;	
+		}
 
-	// Give the default empty emitter a chance to be selected if the user chose 'Empty Emitter'
+		SelectedAssetData = CreateAssetBrowserWindow->GetSelectedAssets();
+	}
+	
+	// Give the minimal emitter a chance to be selected if the user chose 'Minimal Emitter'
 	if(SelectedAssetData.Num() == 0)
 	{
-		FSoftObjectPath DefaultEmptyEmitter = GetDefault<UNiagaraEditorSettings>()->DefaultEmptyEmitter;
-		if(DefaultEmptyEmitter.IsValid() && DefaultEmptyEmitter.IsAsset())
+		FSoftObjectPath MinimalEmitter = GetDefault<UNiagaraEditorSettings>()->DefaultEmptyEmitter;
+		if(MinimalEmitter.IsValid() && MinimalEmitter.IsAsset())
 		{
-			SelectedAssetData.Add(FAssetData(DefaultEmptyEmitter.TryLoad()));
+			SelectedAssetData.Add(FAssetData(MinimalEmitter.TryLoad()));
 		}
 	}
 
@@ -131,22 +164,21 @@ UObject* UNiagaraEmitterFactoryNew::FactoryCreateNew(UClass* Class, UObject* InP
 
 	UNiagaraEmitter* NewEmitter;
 
-	if (EmitterToCopy != nullptr)
+	if (EmitterToCopy.IsValid())
 	{
 		if (bUseInheritance)
 		{
-			NewEmitter = UNiagaraEmitter::CreateWithParentAndOwner(FVersionedNiagaraEmitter(EmitterToCopy, EmitterToCopy->GetExposedVersion().VersionGuid), InParent, Name, Flags);
+			NewEmitter = UNiagaraEmitter::CreateWithParentAndOwner(FVersionedNiagaraEmitter(EmitterToCopy.Get(), EmitterToCopy->GetExposedVersion().VersionGuid), InParent, Name, Flags);
 		}
 		else
 		{
-			NewEmitter = Cast<UNiagaraEmitter>(StaticDuplicateObject(EmitterToCopy, InParent, Name, Flags, Class));
+			NewEmitter = Cast<UNiagaraEmitter>(StaticDuplicateObject(EmitterToCopy.Get(), InParent, Name, Flags, Class));
 			NewEmitter->SetUniqueEmitterName(Name.GetPlainNameString());
 			NewEmitter->DisableVersioning(EmitterToCopy->GetExposedVersion().VersionGuid);
 		}
 
 		NewEmitter->bIsInheritable = true;
 		NewEmitter->TemplateAssetDescription = FText();
-		NewEmitter->AssetTags.Empty();
 		NewEmitter->Category = FText();
 	}
 	else
@@ -244,5 +276,53 @@ void UNiagaraEmitterFactoryNew::InitializeEmitter(UNiagaraEmitter* NewEmitter, b
 	}
 }
 
+bool UNiagaraEmitterFactoryNew::OnAdditionalShouldFilterAsset(const FAssetData& AssetData) const
+{
+	const UNiagaraEditorSettings* Settings = GetDefault<UNiagaraEditorSettings>();
+	if(Settings->IsAllowedAssetInNiagaraAssetBrowser(AssetData) == false)
+	{
+		return true;
+	}
+	
+	if(Settings->IsAllowedAssetByClassUsage(AssetData) == false)
+	{
+		return true;
+	}
+
+	return false;
+}
+
+void UNiagaraEmitterFactoryNew::OnExtendAddFilterMenu(UToolMenu* ToolMenu) const
+{
+	for (FToolMenuSection& Section : ToolMenu->Sections)
+	{
+		Section.Blocks.RemoveAll([](FToolMenuEntry& ToolMenuEntry)
+		{
+			return ToolMenuEntry.Name == FName("Common");
+		});
+	}
+	
+	TArray<FName> SectionsToKeep { FName("FilterBarAdvanced"), FName("FilterBarResetFilters"), FName("FX Filters"), FName("AssetFilterBarFilterAdvancedAsset"), FName("Niagara Filters"), FName("Niagara Tags") };
+	ToolMenu->Sections.RemoveAll([&SectionsToKeep](FToolMenuSection& ToolMenuSection)
+	{
+		return SectionsToKeep.Contains(ToolMenuSection.Name) == false;
+	});
+}
+
+TArray<TSharedRef<FFrontendFilter>> UNiagaraEmitterFactoryNew::OnGetExtraFrontendFilters() const
+{
+	TArray<TSharedRef<FFrontendFilter>> Result;
+	
+	TSharedRef<FFrontendFilterCategory> NiagaraAdditionalFiltersCategory = MakeShared<FFrontendFilterCategory>(LOCTEXT("NiagaraPropertyFilterCategoryLabel", "Niagara Filters"), LOCTEXT("NiagaraAdditionalFiltersTooltip", "Additional filters for filtering Niagara assets"));
+	
+	Result.Add(MakeShared<FFrontendFilter_NiagaraEmitterInheritance>(true, NiagaraAdditionalFiltersCategory));
+	Result.Add(MakeShared<FFrontendFilter_NiagaraEmitterInheritance>(false, NiagaraAdditionalFiltersCategory));
+	
+	Result.Add(MakeShared<FFrontendFilter_NiagaraSystemEffectType>(true, NiagaraAdditionalFiltersCategory));
+	Result.Add(MakeShared<FFrontendFilter_NiagaraSystemEffectType>(false, NiagaraAdditionalFiltersCategory));
+
+	return Result;
+}
+
 #undef LOCTEXT_NAMESPACE
 
diff --git a/Engine/Plugins/FX/Niagara/Source/NiagaraEditor/Private/NiagaraGraph.cpp b/Engine/Plugins/FX/Niagara/Source/NiagaraEditor/Private/NiagaraGraph.cpp
--- a/Engine/Plugins/FX/Niagara/Source/NiagaraEditor/Private/NiagaraGraph.cpp
+++ b/Engine/Plugins/FX/Niagara/Source/NiagaraEditor/Private/NiagaraGraph.cpp
@@ -4366,13 +4366,16 @@ void UNiagaraGraph::MigrateParameterScriptDataToHierarchyRoot(FVersionedNiagaraS
 	for(const FNiagaraStackSection& StackSection : OwnerData.InputSections_DEPRECATED)
 	{
 		FText SectionDisplayName = StackSection.SectionDisplayName.IsEmptyOrWhitespace() ? FText::FromName(StackSection.SectionIdentifier) : StackSection.SectionDisplayName;
-		UHierarchySection* HierarchySection = ParameterRoot->AddSection(SectionDisplayName);
+		UHierarchySection* HierarchySection = NewObject<UHierarchySection>(ParameterRoot, UHierarchySection::StaticClass());
+		HierarchySection->SetSectionNameAsText(SectionDisplayName);
 		HierarchySection->SetTooltip(StackSection.Tooltip);
 
 		FHierarchyElementIdentity SectionIdentity;
 		SectionIdentity.Names.Add(StackSection.SectionIdentifier);
 		HierarchySection->SetIdentity(SectionIdentity);
 			
+		ParameterRoot->GetSectionDataMutable().Insert(HierarchySection, 0);
+		
 		HierarchySections.Add(StackSection.SectionIdentifier, HierarchySection);
 
 		for(const FText& CategoryIdentity : StackSection.Categories)
@@ -4385,8 +4388,11 @@ void UNiagaraGraph::MigrateParameterScriptDataToHierarchyRoot(FVersionedNiagaraS
 				
 			FName CategoryName = FName(CategoryIdentity.ToString());
 			UNiagaraHierarchyScriptCategory* HierarchyCategory = ParameterRoot->AddChild<UNiagaraHierarchyScriptCategory>();
-			HierarchyCategory->SetSection(HierarchySection);
 			HierarchyCategory->SetCategoryName(CategoryName);
+			
+			FDataHierarchyElementMetaData_SectionAssociation* SectionAssociation = HierarchyCategory->FindOrAddMetaDataOfType<FDataHierarchyElementMetaData_SectionAssociation>();
+			SectionAssociation->Section = HierarchySection;
+			
 			IdentityCategoryMap.Add(CategoryName, HierarchyCategory);
 			SectionCategoryMapping.Add(CategoryName, HierarchySection->GetSectionName());
 		}
diff --git a/Engine/Plugins/FX/Niagara/Source/NiagaraEditor/Private/NiagaraScriptMergeManager.cpp b/Engine/Plugins/FX/Niagara/Source/NiagaraEditor/Private/NiagaraScriptMergeManager.cpp
--- a/Engine/Plugins/FX/Niagara/Source/NiagaraEditor/Private/NiagaraScriptMergeManager.cpp
+++ b/Engine/Plugins/FX/Niagara/Source/NiagaraEditor/Private/NiagaraScriptMergeManager.cpp
@@ -2474,9 +2474,9 @@ void FNiagaraScriptMergeManager::DiffEmitterSummary(const UNiagaraEmitterEditorD
 
 	DiffResults.NewShouldShowSummaryViewValue = OtherEditorData->ShouldShowSummaryView();
 	
-	TArray<UHierarchyElement*> AddedItemsInOther;
-	TArray<UHierarchyElement*> RemovedItemsInBase;
-	TArray<UHierarchySection*> AddedSections;
+	TArray<const UHierarchyElement*> AddedItemsInOther;
+	TArray<const UHierarchyElement*> RemovedItemsInBase;
+	TArray<const UHierarchySection*> AddedSections;
 	
 	UHierarchyRoot* BaseRoot = BaseEditorData->GetSummaryRoot();
 	UHierarchyRoot* OtherRoot = OtherEditorData->GetSummaryRoot();	
@@ -2494,13 +2494,13 @@ void FNiagaraScriptMergeManager::DiffEmitterSummary(const UNiagaraEmitterEditorD
 	
 	DiffResults.AddedSummaryEntriesInOther.Append(SummaryItemDiff.AddedOtherValues);
 	
-	TArray<UHierarchySection*> BaseSections = BaseRoot->GetSectionData();
-	TArray<UHierarchySection*> OtherSections = OtherRoot->GetSectionData();
+	TConstArrayView<const TObjectPtr<UHierarchySection>> BaseSections = BaseRoot->GetSectionData();
+	TConstArrayView<const TObjectPtr<UHierarchySection>> OtherSections = OtherRoot->GetSectionData();
 
-	for (UHierarchySection* OtherSection : OtherSections)
+	for (const UHierarchySection* OtherSection : OtherSections)
 	{
 		FHierarchyElementIdentity OtherSectionIdentity = OtherSection->GetPersistentIdentity();
-		if(!BaseSections.ContainsByPredicate([OtherSectionIdentity](UHierarchySection* BaseSection)
+		if(!BaseSections.ContainsByPredicate([OtherSectionIdentity](const UHierarchySection* BaseSection)
 		{
 			return BaseSection->GetPersistentIdentity() == OtherSectionIdentity;
 		}))
@@ -3736,7 +3736,7 @@ FNiagaraScriptMergeManager::FApplyDiffResults FNiagaraScriptMergeManager::ApplyE
 		BaseRoot->DuplicateSectionFromOtherRoot(*AddedOtherSection);
 	}
 
-	for (UHierarchyElement* AddedInput : DiffResults.AddedSummaryEntriesInOther)
+	for (const UHierarchyElement* AddedInput : DiffResults.AddedSummaryEntriesInOther)
 	{
 		if (UHierarchyElement* SummaryItemParent = AddedInput->GetTypedOuter<UHierarchyElement>())
 		{
@@ -3757,10 +3757,21 @@ FNiagaraScriptMergeManager::FApplyDiffResults FNiagaraScriptMergeManager::ApplyE
 					{
 						UHierarchyElement* NewChild = BaseRoot->CopyAndAddItemAsChild(*AddedInput);
 						NewChild->GetChildrenMutable().Empty();
-						// for root level category, we have to fixup the section objects they point to as after duplicating them they will still point to the original section
-						if(UHierarchyCategory* AsCategory = Cast<UHierarchyCategory>(NewChild))
+
+						if(FDataHierarchyElementMetaData_SectionAssociation* SectionAssociation = NewChild->FindOrAddMetaDataOfType<FDataHierarchyElementMetaData_SectionAssociation>())
 						{
-							AsCategory->FixupSectionLinkage();
+							if(SectionAssociation->Section.IsValid())
+							{
+								const TObjectPtr<UHierarchySection>* NewSection = BaseRoot->GetSectionData().FindByPredicate([OriginalSection = SectionAssociation->Section](UHierarchySection* SectionCandidate)
+									{
+										return SectionCandidate->GetSectionName().IsEqual(OriginalSection->GetSectionName());
+									});
+
+								if(NewSection)
+								{
+									SectionAssociation->Section = (*NewSection);
+								}
+							}
 						}
 					}				
 				}
diff --git a/Engine/Plugins/FX/Niagara/Source/NiagaraEditor/Private/NiagaraScriptMergeManager.h b/Engine/Plugins/FX/Niagara/Source/NiagaraEditor/Private/NiagaraScriptMergeManager.h
--- a/Engine/Plugins/FX/Niagara/Source/NiagaraEditor/Private/NiagaraScriptMergeManager.h
+++ b/Engine/Plugins/FX/Niagara/Source/NiagaraEditor/Private/NiagaraScriptMergeManager.h
@@ -404,8 +404,8 @@ struct FNiagaraEmitterDiffResults
 	TArray<TSharedRef<FNiagaraRendererMergeAdapter>> ModifiedBaseRenderers;
 	TArray<TSharedRef<FNiagaraRendererMergeAdapter>> ModifiedOtherRenderers;
 
-	TArray<UHierarchyElement*> AddedSummaryEntriesInOther;	
-	TArray<UHierarchySection*> AddedSummarySectionsInOther;
+	TArray<const UHierarchyElement*> AddedSummaryEntriesInOther;	
+	TArray<const UHierarchySection*> AddedSummarySectionsInOther;
 
 	TOptional<bool> NewShouldShowSummaryViewValue;
 	
diff --git a/Engine/Plugins/FX/Niagara/Source/NiagaraEditor/Private/NiagaraSystemFactoryNew.cpp b/Engine/Plugins/FX/Niagara/Source/NiagaraEditor/Private/NiagaraSystemFactoryNew.cpp
--- a/Engine/Plugins/FX/Niagara/Source/NiagaraEditor/Private/NiagaraSystemFactoryNew.cpp
+++ b/Engine/Plugins/FX/Niagara/Source/NiagaraEditor/Private/NiagaraSystemFactoryNew.cpp
@@ -19,7 +19,11 @@
 #include "NiagaraRecentAndFavoritesManager.h"
 #include "NiagaraSettings.h"
 #include "ObjectTools.h"
-#include "Widgets/AssetBrowser/SNiagaraAssetBrowser.h"
+#include "UserAssetTagEditorUtilities.h"
+#include "Assets/TaggedAssetBrowserConfiguration.h"
+#include "Widgets/STaggedAssetBrowser.h"
+#include "Widgets/STaggedAssetBrowserCreateAssetWindow.h"
+#include "Widgets/AssetBrowser/SNiagaraAddEmitterToSystemWindow.h"
 
 #include UE_INLINE_GENERATED_CPP_BY_NAME(NiagaraSystemFactoryNew)
 
@@ -35,28 +39,56 @@ UNiagaraSystemFactoryNew::UNiagaraSystemFactoryNew(const FObjectInitializer& Obj
 
 bool UNiagaraSystemFactoryNew::ConfigureProperties()
 {
-	IMainFrameModule& MainFrame = FModuleManager::LoadModuleChecked<IMainFrameModule>("MainFrame");
-	TSharedPtr<SWindow>	ParentWindow = MainFrame.GetParentWindow();
+	TArray<FAssetData> SelectedAssetData;
+	
+	FSoftObjectPath Path = GetDefault<UNiagaraEditorSettings>()->SystemAssetWizardConfiguration;
+	UObject* ConfigUObject = Path.TryLoad();
+	if(UTaggedAssetBrowserConfiguration* ConfigurationAsset = Cast<UTaggedAssetBrowserConfiguration>(ConfigUObject))
+	{
+		IMainFrameModule& MainFrame = FModuleManager::LoadModuleChecked<IMainFrameModule>("MainFrame");
+		TSharedPtr<SWindow>	ParentWindow = MainFrame.GetParentWindow();
+		
+		STaggedAssetBrowser::FInterfaceOverrideProfiles OverrideProfiles;
+		OverrideProfiles.AssetViewOptionsProfileName = FName("TaggedAssetBrowser");
+		OverrideProfiles.DefaultFilterMenuExpansion = EAssetTypeCategories::FX;
+		OverrideProfiles.FilterBarSaveName = FName("TaggedAssetBrowser.NiagaraSystemFactory");
+		
+		FDefaultDetailsTabConfiguration DefaultDetailsTabConfiguration;
+		DefaultDetailsTabConfiguration.bUseDefaultDetailsTab = true;
+		DefaultDetailsTabConfiguration.EmptySelectionMessage = LOCTEXT("EmptySystemFactorySelectionMessage", "Select an emitter or system as a starting point for your new system.\nA system consists of one or more emitters.");
+		
+		STaggedAssetBrowser::FArguments AssetBrowserArgs;
+		AssetBrowserArgs
+			.AvailableClasses({UNiagaraEmitter::StaticClass(), UNiagaraSystem::StaticClass()})
+			.DefaultDetailsTabConfiguration(DefaultDetailsTabConfiguration)
+			.InterfaceOverrideProfiles(OverrideProfiles)
+			.RecentAndFavoritesList(FNiagaraEditorModule::Get().GetRecentsManager()->GetRecentEmitterAndSystemsList())
+			.OnGetExtraFrontendFilters_UObject(this, &UNiagaraSystemFactoryNew::OnGetExtraFrontendFilters)
+			.OnExtendAddFilterMenu_UObject(this, &UNiagaraSystemFactoryNew::OnExtendAddFilterMenu);
 
-	SNiagaraAssetBrowser::FArguments AssetBrowserArgs;
-	AssetBrowserArgs.AvailableClasses({UNiagaraEmitter::StaticClass(), UNiagaraSystem::StaticClass()})
-	.RecentAndFavoritesList(FNiagaraEditorModule::Get().GetRecentsManager()->GetRecentEmitterAndSystemsList())
-	.EmptySelectionMessage(LOCTEXT("EmptySystemFactorySelectionMessage", "Select an emitter or system as a starting point for your new system.\nA system consists of one or more emitters."));
+		SWindow::FArguments WindowArgs;
+		WindowArgs.Title(LOCTEXT("SystemAssetBrowserWindowTitle", "Create Niagara System - Select an emitter or system as a base"));
+		WindowArgs.SupportsMaximize(false);
+		WindowArgs.SupportsMinimize(false);
+		WindowArgs.ClientSize(FVector2D(1400, 750));
+		WindowArgs.SizingRule(ESizingRule::UserSized);
+		
+		STaggedAssetBrowserWindow::FArguments AssetBrowserWindowArgs;
+		AssetBrowserWindowArgs.AssetBrowserArgs(AssetBrowserArgs);
+		AssetBrowserWindowArgs.WindowArgs(WindowArgs);
+		
+		TSharedRef<STaggedAssetBrowserCreateAssetWindow> CreateAssetBrowserWindow = SNew(STaggedAssetBrowserCreateAssetWindow, *ConfigurationAsset, *UNiagaraSystem::StaticClass())
+			.AssetBrowserWindowArgs(AssetBrowserWindowArgs);
 
-	SNiagaraAssetBrowserWindow::FArguments AssetBrowserWindowArgs;
-	AssetBrowserWindowArgs.AssetBrowserArgs(AssetBrowserArgs)
-	.WindowTitle(LOCTEXT("SystemAssetBrowserWindowTitle", "Create Niagara System - Select an emitter or system as a base"));
-	
-	TSharedRef<SNiagaraCreateAssetWindow> CreateAssetBrowserWindow = SNew(SNiagaraCreateAssetWindow, *UNiagaraSystem::StaticClass()).AssetBrowserWindowArgs(AssetBrowserWindowArgs);
+		FSlateApplication::Get().AddModalWindow(CreateAssetBrowserWindow, ParentWindow);
 
-	FSlateApplication::Get().AddModalWindow(CreateAssetBrowserWindow, ParentWindow);
+		if(CreateAssetBrowserWindow->ShouldProceedWithAction() == false)
+		{
+			return false;
+		}
 
-	if(CreateAssetBrowserWindow->ShouldProceedWithAction() == false)
-	{
-		return false;
+		SelectedAssetData = CreateAssetBrowserWindow->GetSelectedAssets();
 	}
-
-	TArray<FAssetData> SelectedAssetData = CreateAssetBrowserWindow->GetSelectedAssets();
 	
 	if(SelectedAssetData.Num() == 1)
 	{
@@ -98,7 +130,7 @@ UObject* UNiagaraSystemFactoryNew::FactoryCreateNew(UClass* Class, UObject* InPa
 		NewSystem = Cast<UNiagaraSystem>(StaticDuplicateObject(SystemToCopy, InParent, Name, Flags, Class));
 		NewSystem->TemplateAssetDescription = FText();
 		NewSystem->Category = FText();
-		NewSystem->AssetTags.Empty();
+		//NewSystem->AssetTags.Empty();
 
 		// if the new system doesn't have a thumbnail image, check the thumbnail map of the original asset's upackage
 		if(NewSystem->ThumbnailImage == nullptr)
@@ -193,5 +225,53 @@ void UNiagaraSystemFactoryNew::TryAssignDefaultEffectType(UNiagaraSystem* System
 	}
 }
 
+bool UNiagaraSystemFactoryNew::OnAdditionalShouldFilterAsset(const FAssetData& AssetData) const
+{
+	const UNiagaraEditorSettings* Settings = GetDefault<UNiagaraEditorSettings>();
+	if(Settings->IsAllowedAssetInNiagaraAssetBrowser(AssetData) == false)
+	{
+		return true;
+	}
+	
+	if(Settings->IsAllowedAssetByClassUsage(AssetData) == false)
+	{
+		return true;
+	}
+
+	return false;
+}
+
+void UNiagaraSystemFactoryNew::OnExtendAddFilterMenu(UToolMenu* ToolMenu) const
+{
+	for (FToolMenuSection& Section : ToolMenu->Sections)
+	{
+		Section.Blocks.RemoveAll([](FToolMenuEntry& ToolMenuEntry)
+		{
+			return ToolMenuEntry.Name == FName("Common");
+		});
+	}
+	
+	TArray<FName> SectionsToKeep { FName("FilterBarAdvanced"), FName("FilterBarResetFilters"), FName("FX Filters"), FName("AssetFilterBarFilterAdvancedAsset"), FName("Niagara Filters"), FName("Niagara Tags") };
+	ToolMenu->Sections.RemoveAll([&SectionsToKeep](FToolMenuSection& ToolMenuSection)
+	{
+		return SectionsToKeep.Contains(ToolMenuSection.Name) == false;
+	});
+}
+
+TArray<TSharedRef<FFrontendFilter>> UNiagaraSystemFactoryNew::OnGetExtraFrontendFilters() const
+{
+	TArray<TSharedRef<FFrontendFilter>> Result;
+	
+	TSharedRef<FFrontendFilterCategory> NiagaraAdditionalFiltersCategory = MakeShared<FFrontendFilterCategory>(LOCTEXT("NiagaraPropertyFilterCategoryLabel", "Niagara Filters"), LOCTEXT("NiagaraAdditionalFiltersTooltip", "Additional filters for filtering Niagara assets"));
+	
+	Result.Add(MakeShared<FFrontendFilter_NiagaraEmitterInheritance>(true, NiagaraAdditionalFiltersCategory));
+	Result.Add(MakeShared<FFrontendFilter_NiagaraEmitterInheritance>(false, NiagaraAdditionalFiltersCategory));
+	
+	Result.Add(MakeShared<FFrontendFilter_NiagaraSystemEffectType>(true, NiagaraAdditionalFiltersCategory));
+	Result.Add(MakeShared<FFrontendFilter_NiagaraSystemEffectType>(false, NiagaraAdditionalFiltersCategory));
+
+	return Result;
+}
+
 #undef LOCTEXT_NAMESPACE
 
diff --git a/Engine/Plugins/FX/Niagara/Source/NiagaraEditor/Private/NiagaraValidationRules.cpp b/Engine/Plugins/FX/Niagara/Source/NiagaraEditor/Private/NiagaraValidationRules.cpp
--- a/Engine/Plugins/FX/Niagara/Source/NiagaraEditor/Private/NiagaraValidationRules.cpp
+++ b/Engine/Plugins/FX/Niagara/Source/NiagaraEditor/Private/NiagaraValidationRules.cpp
@@ -1122,24 +1122,9 @@ void UNiagaraValidationRule_HasEffectType::CheckValidity(const FNiagaraValidatio
 	}
 }
 
-void UNiagaraValidationRule_CheckDeprecatedEmitters::CheckValidity(const FNiagaraValidationContext& Context, TArray<FNiagaraValidationResult>& OutResults) const
+void UDEPRECATED_NiagaraValidationRule_CheckDeprecatedEmitters::CheckValidity(const FNiagaraValidationContext& Context, TArray<FNiagaraValidationResult>& OutResults) const
 {
-	TArray<TSharedRef<FNiagaraEmitterHandleViewModel>> EmitterHandleViewModels = Context.ViewModel->GetEmitterHandleViewModels();
-	for (TSharedRef<FNiagaraEmitterHandleViewModel> EmitterHandleViewModel : EmitterHandleViewModels)
-	{
-		if(EmitterHandleViewModel->GetEmitterHandle()->GetInstance().Emitter->AssetTags.ContainsByPredicate([](const FNiagaraAssetTagDefinitionReference& AssetTagReferenceCandidate)
-		{
-			return INiagaraModule::Get().DeprecatedTagDefinition.TagGuid == AssetTagReferenceCandidate.GetTagDefinitionReferenceGuid();
-		}))
-		{
-			UNiagaraStackEmitterPropertiesItem* EmitterProperties = NiagaraValidation::GetStackEntry<UNiagaraStackEmitterPropertiesItem>(EmitterHandleViewModel.Get().GetEmitterStackViewModel());
-
-			FNiagaraValidationResult Result(Severity, FText::Format(
-				LOCTEXT("DeprecatedEmitterUsedTitle", "Emitter '{0}' is deprecated."), FText::FromName(EmitterHandleViewModel->GetEmitterHandle()->GetName())),
-				LOCTEXT("DeprecatedEmitterUsedDescription", "The emitter is deprecated and should not be used. Consider replacing it."), EmitterProperties);
-			OutResults.Add(Result);
-		}
-	}
+	// TODO(ME) Removed validity check before we delete the class
 }
 
 void UNiagaraValidationRule_LWC::CheckValidity(const FNiagaraValidationContext& Context, TArray<FNiagaraValidationResult>& Results)  const
diff --git a/Engine/Plugins/FX/Niagara/Source/NiagaraEditor/Private/Toolkits/NiagaraSystemToolkit.cpp b/Engine/Plugins/FX/Niagara/Source/NiagaraEditor/Private/Toolkits/NiagaraSystemToolkit.cpp
--- a/Engine/Plugins/FX/Niagara/Source/NiagaraEditor/Private/Toolkits/NiagaraSystemToolkit.cpp
+++ b/Engine/Plugins/FX/Niagara/Source/NiagaraEditor/Private/Toolkits/NiagaraSystemToolkit.cpp
@@ -53,7 +53,9 @@
 #include "Modules/ModuleManager.h"
 #include "ViewModels/NiagaraMessageLogViewModel.h"
 #include "NiagaraVersionMetaData.h"
+#include "Assets/TaggedAssetBrowserConfiguration.h"
 #include "Interfaces/IMainFrameModule.h"
+#include "Logging/StructuredLog.h"
 #include "Toolkits/SystemToolkitModes/NiagaraSystemToolkitModeBase.h"
 #include "SystemToolkitModes/NiagaraSystemToolkitMode_Default.h"
 #include "SystemToolkitModes/NiagaraSystemToolkitMode_Scalability.h"
@@ -63,7 +65,7 @@
 #include "Widgets/Layout/SBox.h"
 #include "Widgets/SBoxPanel.h"
 #include "Widgets/SNiagaraEmitterVersionWidget.h"
-#include "Widgets/AssetBrowser/SNiagaraAssetBrowser.h"
+#include "Widgets/AssetBrowser/SNiagaraAddEmitterToSystemWindow.h"
 
 #define LOCTEXT_NAMESPACE "NiagaraSystemEditor"
 
@@ -606,8 +608,17 @@ void FNiagaraSystemToolkit::OpenAddEmitterMenu()
 	IMainFrameModule& MainFrame = FModuleManager::LoadModuleChecked<IMainFrameModule>("MainFrame");
 	TSharedPtr<SWindow>	ParentWindow = MainFrame.GetParentWindow();
 
-	TSharedRef<SWindow> AddEmitterWindow = SNew(SNiagaraAddEmitterToSystemWindow, GetSystemViewModel().ToSharedRef());
-	FSlateApplication::Get().AddModalWindow(AddEmitterWindow, ParentWindow);
+	FSoftObjectPath Path = GetDefault<UNiagaraEditorSettings>()->EmitterAssetWizardConfiguration;
+	UObject* ConfigUObject = Path.TryLoad();
+	if(UTaggedAssetBrowserConfiguration* ConfigurationAsset = Cast<UTaggedAssetBrowserConfiguration>(ConfigUObject))
+	{
+		TSharedRef<SWindow> AddEmitterWindow = SNew(SNiagaraAddEmitterToSystemWindow, *ConfigurationAsset, GetSystemViewModel().ToSharedRef());
+		FSlateApplication::Get().AddModalWindow(AddEmitterWindow, ParentWindow);
+	}
+	else
+	{
+		UE_LOGFMT(LogNiagaraEditor, Warning, "Could not find emitter wizard configuration asset. Check Niagara Editor settings under Project Settings");
+	}
 }
 
 bool FNiagaraSystemToolkit::CanAddEmitters() const
diff --git a/Engine/Plugins/FX/Niagara/Source/NiagaraEditor/Private/ViewModels/HierarchyEditor/NiagaraScriptParametersHierarchyViewModel.cpp b/Engine/Plugins/FX/Niagara/Source/NiagaraEditor/Private/ViewModels/HierarchyEditor/NiagaraScriptParametersHierarchyViewModel.cpp
--- a/Engine/Plugins/FX/Niagara/Source/NiagaraEditor/Private/ViewModels/HierarchyEditor/NiagaraScriptParametersHierarchyViewModel.cpp
+++ b/Engine/Plugins/FX/Niagara/Source/NiagaraEditor/Private/ViewModels/HierarchyEditor/NiagaraScriptParametersHierarchyViewModel.cpp
@@ -279,7 +279,7 @@ FHierarchyElementViewModel::FResultWithUserFeedback FNiagaraHierarchyScriptCateg
 bool FNiagaraHierarchyScriptParameterViewModel::DoesExternalDataStillExist(const UHierarchyDataRefreshContext* Context) const
 {
 	// GetScriptVariable checks the owning graph if it still exists
-	const UNiagaraScriptVariable* ScriptVariable = Cast<UNiagaraHierarchyScriptParameter>(GetDataMutable())->GetScriptVariable();
+	const UNiagaraScriptVariable* ScriptVariable = Cast<UNiagaraHierarchyScriptParameter>(GetData())->GetScriptVariable();
 	
 	if(ScriptVariable == nullptr)
 	{
