User Asset Tagging: v3

diff --git a/Engine/Plugins/FX/Niagara/Source/NiagaraEditor/Private/ViewModels/HierarchyEditor/NiagaraSummaryViewViewModel.cpp b/Engine/Plugins/FX/Niagara/Source/NiagaraEditor/Private/ViewModels/HierarchyEditor/NiagaraSummaryViewViewModel.cpp
--- a/Engine/Plugins/FX/Niagara/Source/NiagaraEditor/Private/ViewModels/HierarchyEditor/NiagaraSummaryViewViewModel.cpp
+++ b/Engine/Plugins/FX/Niagara/Source/NiagaraEditor/Private/ViewModels/HierarchyEditor/NiagaraSummaryViewViewModel.cpp
@@ -243,7 +243,7 @@ void FNiagaraFunctionViewModel::RefreshChildrenDataInternal()
 	RefreshChildrenInputs(false);
 }
 
-void FNiagaraFunctionViewModel::RefreshChildrenInputs(bool bClearCache) const
+void FNiagaraFunctionViewModel::RefreshChildrenInputs(bool bClearCache)
 {
 	TWeakObjectPtr<UNiagaraNodeFunctionCall> FunctionNodeWeak = GetFunctionCallNode();
 	if(FunctionNodeWeak.IsValid())
@@ -322,7 +322,7 @@ void FNiagaraFunctionViewModel::RefreshChildrenInputs(bool bClearCache) const
 
 				for(const FNiagaraVariable& Variable : AssignmentNode->GetAssignmentTargets())
 				{
-					const bool bChildExists = GetDataMutable()->GetChildrenMutable().ContainsByPredicate([Variable, bClearCache, this](UHierarchyElement* Candidate)
+					const bool bChildExists = GetData()->GetChildren().ContainsByPredicate([Variable, bClearCache, this](UHierarchyElement* Candidate)
 					{
 						if(UNiagaraHierarchyAssignmentInput* AssignmentInput = Cast<UNiagaraHierarchyAssignmentInput>(Candidate))
 						{
@@ -361,16 +361,6 @@ FHierarchyElementViewModel::FResultWithUserFeedback FNiagaraFunctionViewModel::C
 	return FHierarchyItemViewModel::CanDropOnInternal(NiagaraHierarchyItemViewModelBase, ItemDropZone);
 }
 
-const UHierarchySection* FNiagaraFunctionViewModel::GetSectionInternal() const
-{
-	if(bIsForHierarchy)
-	{
-		return nullptr;
-	}
-
-	return Section.IsValid() ? Section.Get() : nullptr;
-}
-
 FHierarchyElementViewModel::FResultWithUserFeedback FNiagaraModuleInputViewModel::CanDropOnInternal(TSharedPtr<FHierarchyElementViewModel> DraggedItem, EItemDropZone ItemDropZone)
 {
 	// if the input isn't editable, we don't allow any drops on/above/below the item.
@@ -904,11 +894,6 @@ FHierarchyElementViewModel::FResultWithUserFeedback FNiagaraHierarchyRendererVie
 	return CanEditResults;
 }
 
-const UHierarchySection* FNiagaraHierarchyRendererViewModel::GetSectionInternal() const
-{
-	return Section.IsValid() ? Section.Get() : nullptr;
-}
-
 FString FNiagaraHierarchyEmitterPropertiesViewModel::ToString() const
 {
 	return TEXT("Emitter Properties");
@@ -932,11 +917,6 @@ FHierarchyElementViewModel::FResultWithUserFeedback FNiagaraHierarchyEmitterProp
 	return CanEditResults;
 }
 
-const UHierarchySection* FNiagaraHierarchyEmitterPropertiesViewModel::GetSectionInternal() const
-{
-	return Section.IsValid() ? Section.Get() : nullptr;
-}
-
 FString FNiagaraHierarchyEventHandlerViewModel::ToString() const
 {
 	FNiagaraEventScriptProperties* ScriptProperties = GetEventScriptProperties();
@@ -1039,16 +1019,6 @@ FHierarchyElementViewModel::FResultWithUserFeedback FNiagaraHierarchyEventHandle
 	return CanEditResults;
 }
 
-const UHierarchySection* FNiagaraHierarchyEventHandlerViewModel::GetSectionInternal() const
-{
-	if(Section.IsValid())
-	{
-		return Section.Get();
-	}
-
-	return nullptr;
-}
-
 UNiagaraRendererProperties* FNiagaraHierarchyRendererViewModel::GetRendererProperties() const
 {
 	TWeakObjectPtr<UDataHierarchyViewModelBase> ViewModel = GetHierarchyViewModel();
@@ -1254,11 +1224,6 @@ FHierarchyElementViewModel::FResultWithUserFeedback FNiagaraHierarchySimStageVie
 	return CanEditResults;
 }
 
-const UHierarchySection* FNiagaraHierarchySimStageViewModel::GetSectionInternal() const
-{
-	return Section.IsValid() ? Section.Get() : nullptr;
-}
-
 FString FNiagaraHierarchySimStagePropertiesViewModel::ToString() const
 {
 	return GetSimStage()->SimulationStageName.ToString().Append(" Properties");
@@ -1415,12 +1380,86 @@ void UNiagaraSummaryViewViewModel::PrepareSourceItems(UHierarchyRoot* SourceRoot
 	TSharedPtr<FNiagaraEmitterViewModel> EmitterViewModel = EmitterViewModelWeak.Pin();
 
 	TArray<UHierarchyElement*> NewItems;
-	TArray<UHierarchySection*> NewSections;
+	TArray<UHierarchySection*> Sections;
 	
 	FunctionCallCache.Empty();
-		
-	// we keep track of all category items (modules etc.) belonging to a certain usage so we can set the sections later 
-	TMap<ENiagaraScriptUsage, TArray<UNiagaraHierarchyModule*>> UsageMap;
+
+	const FName EmitterSpawnSectionName = "Emitter Spawn";
+	const FName EmitterUpdateSectionName = "Emitter Update";
+	const FName ParticleSpawnSectionName = "Particle Spawn";
+	const FName ParticleUpdateSectionName = "Particle Update";
+	const FName EventsSectionName = "Events";
+	const FName SimStagesSectionName = "Sim Stages";
+	const FName RenderersSectionName = "Renderers";
+	
+	TMap<ENiagaraScriptUsage, FName> ScriptUsageSectionNameMap;
+	ScriptUsageSectionNameMap.Add(ENiagaraScriptUsage::EmitterSpawnScript, EmitterSpawnSectionName);
+	ScriptUsageSectionNameMap.Add(ENiagaraScriptUsage::EmitterUpdateScript, EmitterUpdateSectionName);
+	ScriptUsageSectionNameMap.Add(ENiagaraScriptUsage::ParticleSpawnScript, ParticleSpawnSectionName);
+	ScriptUsageSectionNameMap.Add(ENiagaraScriptUsage::ParticleUpdateScript, ParticleUpdateSectionName);
+
+	TMap<FName, FName> ScriptUsageSectionIconMap;
+	ScriptUsageSectionIconMap.Add(ScriptUsageSectionNameMap[ENiagaraScriptUsage::EmitterSpawnScript], "NiagaraEditor.Emitter.SpawnIcon");
+	ScriptUsageSectionIconMap.Add(ScriptUsageSectionNameMap[ENiagaraScriptUsage::EmitterUpdateScript], "NiagaraEditor.Emitter.UpdateIcon");
+	ScriptUsageSectionIconMap.Add(ScriptUsageSectionNameMap[ENiagaraScriptUsage::ParticleSpawnScript], "NiagaraEditor.Particle.SpawnIcon");
+	ScriptUsageSectionIconMap.Add(ScriptUsageSectionNameMap[ENiagaraScriptUsage::ParticleUpdateScript], "NiagaraEditor.Particle.UpdateIcon");
+	ScriptUsageSectionIconMap.Add(EventsSectionName, "NiagaraEditor.EventIcon");
+	ScriptUsageSectionIconMap.Add(SimStagesSectionName, "NiagaraEditor.SimulationStageIcon");
+	ScriptUsageSectionIconMap.Add(RenderersSectionName, "NiagaraEditor.RenderIcon");
+
+	TMap<FName, TSharedPtr<FHierarchySectionViewModel>> SectionViewModelMap;
+
+	auto CreateAndInitializeSectionData = [this, SourceRoot, &Sections](FName SectionName) -> UHierarchySection*
+	{
+		UHierarchySection* Section = nullptr;
+
+		auto* FoundHierarchySection = SourceRoot->GetSectionDataMutable().FindByPredicate([SectionName](UHierarchySection* Candidate)
+		{
+			return Candidate->GetSectionName().IsEqual(SectionName);
+		});
+
+		if(FoundHierarchySection == nullptr)
+		{
+			Section = SourceRoot->AddChild<UHierarchySection>(GetSectionDataClass(), INDEX_NONE);
+			Section->SetSectionName(SectionName);
+		}
+		else
+		{
+			Section = *FoundHierarchySection;
+		}
+
+		ensure(Section);
+		Sections.Add(Section);
+		return Section;
+	};
+
+	TMap<ENiagaraScriptUsage, UHierarchySection*> ScriptUsageSectionMap;
+	for(auto& SectionNamePair : ScriptUsageSectionNameMap)
+	{
+		ScriptUsageSectionMap.Add(SectionNamePair.Key, CreateAndInitializeSectionData(SectionNamePair.Value));
+	}
+	
+	CreateAndInitializeSectionData(EventsSectionName);
+	CreateAndInitializeSectionData(SimStagesSectionName);
+	CreateAndInitializeSectionData(RenderersSectionName);
+
+	// We create the view models now
+	SourceRootViewModel->SyncViewModelsToData();
+	
+	TArray<TSharedPtr<FHierarchySectionViewModel>> SectionViewModels;
+	SourceRootViewModel->GetChildrenViewModelsForType<UHierarchySection, FHierarchySectionViewModel>(SectionViewModels);
+
+	for(TSharedPtr<FHierarchySectionViewModel> SectionViewModel : SectionViewModels)
+	{
+		SectionViewModelMap.Add(SectionViewModel->GetSectionName(), SectionViewModel);
+	}
+
+	// Now we set their images
+	for(const auto& SectionViewModelMapEntry : SectionViewModelMap)
+	{
+		SectionViewModelMapEntry.Value->SetSectionImage(FNiagaraEditorStyle::Get().GetBrush(ScriptUsageSectionIconMap[SectionViewModelMapEntry.Key]));
+	}
+	
 	// we keep track of renderers & sim stages here for the same reasons
 	TArray<UNiagaraHierarchyRenderer*> HierarchyRenderers;
 	TArray<UNiagaraHierarchySimStage*> HierarchySimStages;
@@ -1440,8 +1479,7 @@ void UNiagaraSummaryViewViewModel::PrepareSourceItems(UHierarchyRoot* SourceRoot
 	NewItems.Add(EmitterProperties);
 	
 	// We create hierarchy modules here. We attempt to maintain as many previous elements as possible in order to maintain UI state
-	TArray<UNiagaraNodeFunctionCall*> ModuleNodes = FNiagaraStackGraphUtilities::GetAllModuleNodes(EmitterViewModel.ToSharedRef());
-	for(UNiagaraNodeFunctionCall* ModuleNode : ModuleNodes)
+	for(UNiagaraNodeFunctionCall* ModuleNode : FNiagaraStackGraphUtilities::GetAllModuleNodes(EmitterViewModel.ToSharedRef()))
 	{
 		// we skip over sim stage modules here as we want to add them to their respective sim stage group items instead
 		if(SimStageModules.Contains(ModuleNode) || EventHandlerModules.Contains(ModuleNode))
@@ -1449,332 +1487,111 @@ void UNiagaraSummaryViewViewModel::PrepareSourceItems(UHierarchyRoot* SourceRoot
 			continue;
 		}
 		
-		UNiagaraHierarchyModule* HierarchyModule = SourceRoot->AddChild<UNiagaraHierarchyModule>();		
-		HierarchyModule->Initialize(*ModuleNode);		
+		ENiagaraScriptUsage ScriptUsage = FNiagaraStackGraphUtilities::GetOutputNodeUsage(*ModuleNode);
 		
-		NewItems.Add(HierarchyModule);
-		UsageMap.FindOrAdd(FNiagaraStackGraphUtilities::GetOutputNodeUsage(*ModuleNode)).Add(HierarchyModule);
+		if(TSharedPtr<FHierarchyElementViewModel> ParentElement = SectionViewModelMap[ScriptUsageSectionNameMap[ScriptUsage]])
+		{
+			UNiagaraHierarchyModule* HierarchyModule = ScriptUsageSectionMap[ScriptUsage]->AddChild<UNiagaraHierarchyModule>();		
+			HierarchyModule->Initialize(*ModuleNode);		
+		
+			NewItems.Add(HierarchyModule);
+		}
 	}
-
-	const TArray<FNiagaraEventScriptProperties> ScriptProperties = EmitterViewModel->GetEmitter().GetEmitterData()->GetEventHandlers();
-	for(const FNiagaraEventScriptProperties& ScriptPropertiesItem : ScriptProperties)
+	
+	for(const FNiagaraEventScriptProperties& ScriptPropertiesItem : EmitterViewModel->GetEmitter().GetEmitterData()->GetEventHandlers())
 	{
 		UNiagaraHierarchyEventHandler* HierarchyEventHandler = nullptr;
 		FHierarchyElementIdentity EventHandlerIdentity;
 		EventHandlerIdentity.Guids.Add(ScriptPropertiesItem.Script->GetUsageId());
 		EventHandlerIdentity.Guids.Add(ScriptPropertiesItem.SourceEmitterID);
-
-		auto* FoundItem = SourceRoot->GetChildrenMutable().FindByPredicate([EventHandlerIdentity](UHierarchyElement* ItemBase)
-		{
-			return ItemBase->GetPersistentIdentity() == EventHandlerIdentity;
-		});
 		
-		if(FoundItem == nullptr)		
-		{
-			HierarchyEventHandler = SourceRoot->AddChild<UNiagaraHierarchyEventHandler>();
-			HierarchyEventHandler->Initialize(ScriptPropertiesItem);
-		}
-		else
+		if(TSharedPtr<FHierarchyElementViewModel> ParentElement = SectionViewModelMap[EventsSectionName])
 		{
-			HierarchyEventHandler = CastChecked<UNiagaraHierarchyEventHandler>(*FoundItem);
+			auto* FoundItem = ParentElement->GetDataMutable()->GetChildrenMutable().FindByPredicate([EventHandlerIdentity](UHierarchyElement* ItemBase)
+			{
+				return ItemBase->GetPersistentIdentity() == EventHandlerIdentity;
+			});
+			
+			if(FoundItem == nullptr)		
+			{
+				HierarchyEventHandler = ParentElement->GetDataMutable()->AddChild<UNiagaraHierarchyEventHandler>();
+				HierarchyEventHandler->Initialize(ScriptPropertiesItem);
+			}
+			else
+			{
+				HierarchyEventHandler = CastChecked<UNiagaraHierarchyEventHandler>(*FoundItem);
+			}
+			
+			NewItems.Add(HierarchyEventHandler);
+			HierarchyEventHandlers.Add(HierarchyEventHandler);
 		}
-		
-		NewItems.Add(HierarchyEventHandler);
-		HierarchyEventHandlers.Add(HierarchyEventHandler);
 	}
 
 	// We add sim stages here
-	const TArray<UNiagaraSimulationStageBase*>& SimStages = EmitterViewModel->GetEmitter().GetEmitterData()->GetSimulationStages();
-	for(UNiagaraSimulationStageBase* SimStage : SimStages)
+	for(UNiagaraSimulationStageBase* SimStage : EmitterViewModel->GetEmitter().GetEmitterData()->GetSimulationStages())
 	{
 		UNiagaraHierarchySimStage* HierarchySimStage = nullptr;
 		FHierarchyElementIdentity SimStageID;
 		SimStageID.Guids.Add(SimStage->GetMergeId());
 
-		auto* FoundItem = SourceRoot->GetChildrenMutable().FindByPredicate([SimStageID](UHierarchyElement* ItemBase)
+		if(TSharedPtr<FHierarchyElementViewModel> ParentElement = SectionViewModelMap[SimStagesSectionName])
 		{
-			return ItemBase->GetPersistentIdentity() == SimStageID;
-		});
-		
-		if(FoundItem == nullptr)		
-		{
-			HierarchySimStage = SourceRoot->AddChild<UNiagaraHierarchySimStage>();
-			HierarchySimStage->Initialize(*SimStage);
-		}
-		else
-		{
-			HierarchySimStage = CastChecked<UNiagaraHierarchySimStage>(*FoundItem);
+			auto* FoundItem = ParentElement->GetDataMutable()->GetChildrenMutable().FindByPredicate([SimStageID](UHierarchyElement* ItemBase)
+			{
+				return ItemBase->GetPersistentIdentity() == SimStageID;
+			});
+			
+			if(FoundItem == nullptr)		
+			{
+				HierarchySimStage = ParentElement->GetDataMutable()->AddChild<UNiagaraHierarchySimStage>();
+				HierarchySimStage->Initialize(*SimStage);
+			}
+			else
+			{
+				HierarchySimStage = CastChecked<UNiagaraHierarchySimStage>(*FoundItem);
+			}
+			
+			NewItems.Add(HierarchySimStage);
+			HierarchySimStages.Add(HierarchySimStage);
 		}
-		
-		NewItems.Add(HierarchySimStage);
-		HierarchySimStages.Add(HierarchySimStage);
 	}
 	
 	// We create hierarchy renderers here
-	const TArray<UNiagaraRendererProperties*>& Renderers = EmitterViewModel->GetEmitter().GetEmitterData()->GetRenderers();
-	for(UNiagaraRendererProperties* RendererProperties : Renderers)
+	for(UNiagaraRendererProperties* RendererProperties : EmitterViewModel->GetEmitter().GetEmitterData()->GetRenderers())
 	{
 		UNiagaraHierarchyRenderer* HierarchyRenderer = nullptr;
 		FHierarchyElementIdentity RendererIdentity;
 		RendererIdentity.Guids.Add(RendererProperties->GetMergeId());
 
-		auto* FoundItem = SourceRoot->GetChildrenMutable().FindByPredicate([RendererIdentity](UHierarchyElement* ItemBase)
-		{
-			return ItemBase->GetPersistentIdentity() == RendererIdentity;
-		});
-		
-		if(FoundItem == nullptr)		
+		if(TSharedPtr<FHierarchyElementViewModel> ParentElement = SectionViewModelMap[RenderersSectionName])
 		{
-			HierarchyRenderer = SourceRoot->AddChild<UNiagaraHierarchyRenderer>();
-			HierarchyRenderer->Initialize(*RendererProperties);
-		}
-		else
-		{
-			HierarchyRenderer = CastChecked<UNiagaraHierarchyRenderer>(*FoundItem);
+			auto* FoundItem = ParentElement->GetDataMutable()->GetChildrenMutable().FindByPredicate([RendererIdentity](UHierarchyElement* ItemBase)
+			{
+				return ItemBase->GetPersistentIdentity() == RendererIdentity;
+			});
+			
+			if(FoundItem == nullptr)		
+			{
+				HierarchyRenderer = ParentElement->GetDataMutable()->AddChild<UNiagaraHierarchyRenderer>();
+				HierarchyRenderer->Initialize(*RendererProperties);
+			}
+			else
+			{
+				HierarchyRenderer = CastChecked<UNiagaraHierarchyRenderer>(*FoundItem);
+			}
+			
+			NewItems.Add(HierarchyRenderer);
+			HierarchyRenderers.Add(HierarchyRenderer);
 		}
-		
-		NewItems.Add(HierarchyRenderer);
-		HierarchyRenderers.Add(HierarchyRenderer);
 	}
-
-	SourceRoot->GetChildrenMutable().Empty();
-	SourceRoot->GetChildrenMutable().Append(NewItems);
-
-	// we force a sync so we can access the section data for the source items
-	SourceRootViewModel->SyncViewModelsToData();
 	
-	// Now we create a section for each usage case that has at least one element and link up the respective sections
-	UHierarchySection* EmitterSpawnSection = nullptr;
-	if(UsageMap.Contains(ENiagaraScriptUsage::EmitterSpawnScript))
-	{
-		FText SectionName = FText::FromString("Emitter Spawn");
-		auto* FoundHierarchySection = SourceRootViewModel->GetDataMutable<UHierarchyRoot>()->GetSectionDataMutable().FindByPredicate([SectionName](UHierarchySection* Candidate)
-		{
-			return Candidate->GetSectionNameAsText().EqualTo(SectionName);
-		});
-		
-		if(FoundHierarchySection == nullptr)
-		{
-			EmitterSpawnSection = SourceRootViewModel->GetDataMutable<UHierarchyRoot>()->AddSection(SectionName);
-		}
-		else
-		{
-			EmitterSpawnSection = *FoundHierarchySection;
-		}
-
-		NewSections.Add(EmitterSpawnSection);
-
-		for(UNiagaraHierarchyModule* Module : UsageMap[ENiagaraScriptUsage::EmitterSpawnScript])
-		{
-			TSharedPtr<FNiagaraFunctionViewModel> ModuleViewModel = StaticCastSharedPtr<FNiagaraFunctionViewModel>(SourceRootViewModel->FindViewModelForChild(Module));
-			ModuleViewModel->SetSection(*EmitterSpawnSection);
-		}
-	}
-
-	UHierarchySection* EmitterUpdateSection = nullptr;
-	if(UsageMap.Contains(ENiagaraScriptUsage::EmitterUpdateScript))
-	{
-		FText SectionName = FText::FromString("Emitter Update");
-		auto* FoundHierarchySection = SourceRootViewModel->GetDataMutable<UHierarchyRoot>()->GetSectionDataMutable().FindByPredicate([SectionName](UHierarchySection* Candidate)
-		{
-			return Candidate->GetSectionNameAsText().EqualTo(SectionName);
-		});
-		
-		if(FoundHierarchySection == nullptr)
-		{
-			EmitterUpdateSection = SourceRootViewModel->GetDataMutable<UHierarchyRoot>()->AddSection(SectionName);
-		}
-		else
-		{
-			EmitterUpdateSection = *FoundHierarchySection;
-		}
-
-		NewSections.Add(EmitterUpdateSection);
-
-		for(UNiagaraHierarchyModule* Module : UsageMap[ENiagaraScriptUsage::EmitterUpdateScript])
-		{
-			TSharedPtr<FNiagaraFunctionViewModel> ModuleViewModel = StaticCastSharedPtr<FNiagaraFunctionViewModel>(SourceRootViewModel->FindViewModelForChild(Module));
-			ModuleViewModel->SetSection(*EmitterUpdateSection);
-		}
-	}
-
-	UHierarchySection* ParticleSpawnSection = nullptr;
-	if(UsageMap.Contains(ENiagaraScriptUsage::ParticleSpawnScript))
-	{
-		FText SectionName = FText::FromString("Particle Spawn");
-		auto* FoundHierarchySection = SourceRootViewModel->GetDataMutable<UHierarchyRoot>()->GetSectionDataMutable().FindByPredicate([SectionName](UHierarchySection* Candidate)
-		{
-			return Candidate->GetSectionNameAsText().EqualTo(SectionName);
-		});
-		
-		if(FoundHierarchySection == nullptr)
-		{
-			ParticleSpawnSection = SourceRootViewModel->GetDataMutable<UHierarchyRoot>()->AddSection(SectionName);
-		}
-		else
-		{
-			ParticleSpawnSection = *FoundHierarchySection;
-		}
-
-		NewSections.Add(ParticleSpawnSection);
-
-		for(UNiagaraHierarchyModule* Module : UsageMap[ENiagaraScriptUsage::ParticleSpawnScript])
-		{
-			TSharedPtr<FNiagaraFunctionViewModel> ModuleViewModel = StaticCastSharedPtr<FNiagaraFunctionViewModel>(SourceRootViewModel->FindViewModelForChild(Module));
-			ModuleViewModel->SetSection(*ParticleSpawnSection);
-		}
-	}
-
-	UHierarchySection* ParticleUpdateSection = nullptr;
-	if(UsageMap.Contains(ENiagaraScriptUsage::ParticleUpdateScript))
-	{
-		FText SectionName = FText::FromString("Particle Update");
-		auto* FoundHierarchySection = SourceRootViewModel->GetDataMutable<UHierarchyRoot>()->GetSectionDataMutable().FindByPredicate([SectionName](UHierarchySection* Candidate)
-		{
-			return Candidate->GetSectionNameAsText().EqualTo(SectionName);
-		});
-		
-		if(FoundHierarchySection == nullptr)
-		{
-			ParticleUpdateSection = SourceRootViewModel->GetDataMutable<UHierarchyRoot>()->AddSection(SectionName);
-		}
-		else
-		{
-			ParticleUpdateSection = *FoundHierarchySection;
-		}
-
-		NewSections.Add(ParticleUpdateSection);
-
-		for(UNiagaraHierarchyModule* Module : UsageMap[ENiagaraScriptUsage::ParticleUpdateScript])
-		{
-			TSharedPtr<FNiagaraFunctionViewModel> ModuleViewModel = StaticCastSharedPtr<FNiagaraFunctionViewModel>(SourceRootViewModel->FindViewModelForChild(Module));
-			ModuleViewModel->SetSection(*ParticleUpdateSection);
-		}
-	}
-
-	UHierarchySection* EventHandlerSection = nullptr;
-	if(HierarchyEventHandlers.Num() > 0)
-	{
-		FText SectionName = FText::FromString("Events");
-		auto* FoundHierarchySection = SourceRootViewModel->GetDataMutable<UHierarchyRoot>()->GetSectionDataMutable().FindByPredicate([SectionName](UHierarchySection* Candidate)
-		{
-			return Candidate->GetSectionNameAsText().EqualTo(SectionName);
-		});
-		
-		if(FoundHierarchySection == nullptr)
-		{
-			EventHandlerSection = SourceRootViewModel->GetDataMutable<UHierarchyRoot>()->AddSection(SectionName);
-		}
-		else
-		{
-			EventHandlerSection = *FoundHierarchySection;
-		}
-
-		NewSections.Add(EventHandlerSection);
-
-		for(UNiagaraHierarchyEventHandler* EventHandler : HierarchyEventHandlers)
-		{
-			TSharedPtr<FNiagaraHierarchyEventHandlerViewModel> EventHandlerViewModel = StaticCastSharedPtr<FNiagaraHierarchyEventHandlerViewModel>(SourceRootViewModel->FindViewModelForChild(EventHandler));
-			EventHandlerViewModel->SetSection(*EventHandlerSection);
-		}
-	}
-
-	UHierarchySection* SimulationStagesSection = nullptr;
-	if(HierarchySimStages.Num() > 0)
-	{
-		FText SectionName = FText::FromString("Sim Stages");
-		auto* FoundHierarchySection = SourceRootViewModel->GetDataMutable<UHierarchyRoot>()->GetSectionDataMutable().FindByPredicate([SectionName](UHierarchySection* Candidate)
-		{
-			return Candidate->GetSectionNameAsText().EqualTo(SectionName);
-		});
-		
-		if(FoundHierarchySection == nullptr)
-		{
-			SimulationStagesSection = SourceRootViewModel->GetDataMutable<UHierarchyRoot>()->AddSection(SectionName);
-		}
-		else
-		{
-			SimulationStagesSection = *FoundHierarchySection;
-		}
-
-		NewSections.Add(SimulationStagesSection);
-
-		for(UNiagaraHierarchySimStage* SimStage : HierarchySimStages)
-		{
-			TSharedPtr<FNiagaraHierarchySimStageViewModel> SimStageViewModel = StaticCastSharedPtr<FNiagaraHierarchySimStageViewModel>(SourceRootViewModel->FindViewModelForChild(SimStage));
-			SimStageViewModel->SetSection(*SimulationStagesSection);
-		}
-	}
-
-	UHierarchySection* RenderersSection = nullptr; 
-	if(HierarchyRenderers.Num() > 0)
-	{
-		FText SectionName = FText::FromString("Renderers");
-		auto* FoundHierarchySection = SourceRootViewModel->GetDataMutable<UHierarchyRoot>()->GetSectionDataMutable().FindByPredicate([SectionName](UHierarchySection* Candidate)
-		{
-			return Candidate->GetSectionNameAsText().EqualTo(SectionName);
-		});
-		
-		if(FoundHierarchySection == nullptr)
-		{
-			RenderersSection = SourceRootViewModel->GetDataMutable<UHierarchyRoot>()->AddSection(SectionName);
-		}
-		else
-		{
-			RenderersSection = *FoundHierarchySection;
-		}
-
-		NewSections.Add(RenderersSection);
-		
-		for(UNiagaraHierarchyRenderer* Renderer : HierarchyRenderers)
-		{
-			TSharedPtr<FNiagaraHierarchyRendererViewModel> RendererViewModel = StaticCastSharedPtr<FNiagaraHierarchyRendererViewModel>(SourceRootViewModel->FindViewModelForChild(Renderer));
-			RendererViewModel->SetSection(*RenderersSection);
-		}
-	}
-
-	// this will implicitly sort the sections as well as get rid of outdated ones
-	SourceRoot->GetSectionDataMutable().Empty();
-	SourceRoot->GetSectionDataMutable().Append(NewSections);
-
-	// force a sync so we have the view models for the sections available
 	SourceRootViewModel->SyncViewModelsToData();
 
-	for(TSharedPtr<FHierarchySectionViewModel> SectionViewModel : SourceRootViewModel->GetSectionViewModels())
+	for(TSharedPtr<FHierarchySectionViewModel> SectionViewModel : SectionViewModels)
 	{
-		if(EmitterSpawnSection != nullptr && SectionViewModel->GetData() == EmitterSpawnSection)
-		{
-			SectionViewModel->SetSectionImage(FNiagaraEditorStyle::Get().GetBrush("NiagaraEditor.Emitter.SpawnIcon"));
-		}
-
-		if(EmitterUpdateSection != nullptr && SectionViewModel->GetData() == EmitterUpdateSection)
-		{
-			SectionViewModel->SetSectionImage(FNiagaraEditorStyle::Get().GetBrush("NiagaraEditor.Emitter.UpdateIcon"));
-		}
-
-		if(ParticleSpawnSection != nullptr && SectionViewModel->GetData() == ParticleSpawnSection)
-		{
-			SectionViewModel->SetSectionImage(FNiagaraEditorStyle::Get().GetBrush("NiagaraEditor.Particle.SpawnIcon"));
-		}
-
-		if(ParticleUpdateSection != nullptr && SectionViewModel->GetData() == ParticleUpdateSection)
-		{
-			SectionViewModel->SetSectionImage(FNiagaraEditorStyle::Get().GetBrush("NiagaraEditor.Particle.UpdateIcon"));
-		}
-
-		if(EventHandlerSection != nullptr && SectionViewModel->GetData() == EventHandlerSection)
-		{
-			SectionViewModel->SetSectionImage(FNiagaraEditorStyle::Get().GetBrush("NiagaraEditor.EventIcon"));
-		}
-
-		if(SimulationStagesSection != nullptr && SectionViewModel->GetData() == SimulationStagesSection)
-		{
-			SectionViewModel->SetSectionImage(FNiagaraEditorStyle::Get().GetBrush("NiagaraEditor.SimulationStageIcon"));
-		}
-
-		if(RenderersSection != nullptr && SectionViewModel->GetData() == RenderersSection)
+		if(SectionViewModel->GetChildren().Num() == 0)
 		{
-			SectionViewModel->SetSectionImage(FNiagaraEditorStyle::Get().GetBrush("NiagaraEditor.RenderIcon"));
+			SectionViewModel->Delete();
 		}
 	}
 }
diff --git a/Engine/Plugins/FX/Niagara/Source/NiagaraEditor/Private/ViewModels/Stack/NiagaraStackScriptHierarchyRoot.cpp b/Engine/Plugins/FX/Niagara/Source/NiagaraEditor/Private/ViewModels/Stack/NiagaraStackScriptHierarchyRoot.cpp
--- a/Engine/Plugins/FX/Niagara/Source/NiagaraEditor/Private/ViewModels/Stack/NiagaraStackScriptHierarchyRoot.cpp
+++ b/Engine/Plugins/FX/Niagara/Source/NiagaraEditor/Private/ViewModels/Stack/NiagaraStackScriptHierarchyRoot.cpp
@@ -54,23 +54,23 @@ void UNiagaraStackScriptHierarchyRoot::FinalizeInternal()
 	Super::FinalizeInternal();
 }
 
-const TArray<UHierarchySection*>& UNiagaraStackScriptHierarchyRoot::GetSections() const
+TConstArrayView<const TObjectPtr<UHierarchySection>> UNiagaraStackScriptHierarchyRoot::GerHierarchySectionData() const
 {
 	if(OwningFunctionCallNode->GetCalledUsage() == ENiagaraScriptUsage::Module)
 	{
 		return GetScriptParameterHierarchyRoot()->GetSectionData();
 	}
-	
-	static TArray<UHierarchySection*> Dummy;
+
+	static TArray<const TObjectPtr<UHierarchySection>> Dummy;
 	return Dummy;
 }
 
-const UHierarchySection* UNiagaraStackScriptHierarchyRoot::GetActiveSection() const
+const UHierarchySection* UNiagaraStackScriptHierarchyRoot::GetActiveHierarchySection() const
 {
 	return ActiveSection.IsValid() ? ActiveSection.Get() : nullptr;
 }
 
-void UNiagaraStackScriptHierarchyRoot::SetActiveSection(const UHierarchySection* InActiveSection)
+void UNiagaraStackScriptHierarchyRoot::SetActiveHierarchySection(const UHierarchySection* InActiveSection)
 {
 	ActiveSection = InActiveSection;
 	
@@ -210,9 +210,9 @@ void UNiagaraStackScriptHierarchyRoot::RefreshChildrenInternal(const TArray<UNia
 		}
 	}
 	
-	for(const UHierarchyElement* ChildHierarchyItem : GetScriptParameterHierarchyRoot()->GetChildren())
+	for(const UHierarchyElement* ChildElement : GetScriptParameterHierarchyRoot()->GetChildren())
 	{
-		if(const UNiagaraHierarchyScriptParameter* HierarchyParameter = Cast<UNiagaraHierarchyScriptParameter>(ChildHierarchyItem))
+		if(const UNiagaraHierarchyScriptParameter* HierarchyParameter = Cast<UNiagaraHierarchyScriptParameter>(ChildElement))
 		{
 			TOptional<FNiagaraVariable> InputVariableCandidate = HierarchyParameter->GetVariable();
 
@@ -249,11 +249,11 @@ void UNiagaraStackScriptHierarchyRoot::RefreshChildrenInternal(const TArray<UNia
 			NewChildren.AddUnique(InputChild);
 		}
 		
-		if(const UHierarchyCategory* HierarchyCategory = Cast<UHierarchyCategory>(ChildHierarchyItem))
+		if(const UHierarchyCategory* HierarchyCategory = Cast<UHierarchyCategory>(ChildElement))
 		{
 			// Try to find an already existing category to reuse
 			UNiagaraStackScriptHierarchyCategory* StackCategory = FindCurrentChildOfTypeByPredicate<UNiagaraStackScriptHierarchyCategory>(CurrentChildren,
-				[&](UNiagaraStackScriptHierarchyCategory* CurrentCategory) { return CurrentCategory->GetHierarchyCategory() == ChildHierarchyItem; });
+				[&](UNiagaraStackScriptHierarchyCategory* CurrentCategory) { return CurrentCategory->GetHierarchyCategory() == ChildElement; });
 
 			if (StackCategory == nullptr)
 			{
@@ -305,7 +305,7 @@ void UNiagaraStackScriptHierarchyRoot::RefreshChildrenInternal(const TArray<UNia
 
 const UHierarchySection* UNiagaraStackScriptHierarchyRoot::FindSectionByName(FName SectionName)
 {
-	UHierarchySection* const* FoundSection = GetSections().FindByPredicate([SectionName](UHierarchySection* SectionCandidate)
+	const TObjectPtr<UHierarchySection>* FoundSection = GerHierarchySectionData().FindByPredicate([SectionName](const UHierarchySection* SectionCandidate)
 	{
 		return SectionCandidate->GetSectionName() == SectionName;
 	});
@@ -327,7 +327,9 @@ bool UNiagaraStackScriptHierarchyRoot::FilterByActiveSection(const UNiagaraStack
 
 	if(const UNiagaraStackScriptHierarchyCategory* StackHierarchyCategory = Cast<UNiagaraStackScriptHierarchyCategory>(&Child))
 	{
-		if(StackHierarchyCategory->GetHierarchyCategory()->GetSection() == ActiveSection)
+		const UHierarchyCategory* HierarchyCategory = StackHierarchyCategory->GetHierarchyCategory();
+		const FDataHierarchyElementMetaData_SectionAssociation* SectionAssociation = HierarchyCategory->FindMetaDataOfType<FDataHierarchyElementMetaData_SectionAssociation>();
+		if(SectionAssociation && SectionAssociation->Section == ActiveSection)
 		{
 			return true;
 		}
@@ -371,7 +373,7 @@ bool UNiagaraStackScriptHierarchyRoot::GetCanExpand() const
 
 bool UNiagaraStackScriptHierarchyRoot::GetShouldShowInStack() const
 {
-	return GetSections().Num() > 0 || bShouldDisplayLabel;
+	return GerHierarchySectionData().Num() > 0 || bShouldDisplayLabel;
 }
 
 int32 UNiagaraStackScriptHierarchyRoot::GetChildIndentLevel() const
diff --git a/Engine/Plugins/FX/Niagara/Source/NiagaraEditor/Private/ViewModels/Stack/NiagaraStackSummaryViewInputCollection.cpp b/Engine/Plugins/FX/Niagara/Source/NiagaraEditor/Private/ViewModels/Stack/NiagaraStackSummaryViewInputCollection.cpp
--- a/Engine/Plugins/FX/Niagara/Source/NiagaraEditor/Private/ViewModels/Stack/NiagaraStackSummaryViewInputCollection.cpp
+++ b/Engine/Plugins/FX/Niagara/Source/NiagaraEditor/Private/ViewModels/Stack/NiagaraStackSummaryViewInputCollection.cpp
@@ -564,7 +564,7 @@ void UNiagaraStackSummaryViewCollection::RefreshChildrenInternal(const TArray<UN
 
 void UNiagaraStackSummaryViewCollection::GetSectionsInternal(TArray<FNiagaraStackSection>& OutStackSections) const
 {
-	const TArray<UHierarchySection*>& HierarchySections = GetHierarchySections();
+	TConstArrayView<const TObjectPtr<UHierarchySection>> HierarchySections = GetHierarchySections();
 
 	TMap<const UHierarchySection*, TArray<UHierarchyCategory*>> SectionToCategoryMap;
 
@@ -576,7 +576,13 @@ void UNiagaraStackSummaryViewCollection::GetSectionsInternal(TArray<FNiagaraStac
 	
 	for(UHierarchyCategory* TopLevelHierarchyCategories : HierarchyCategories)
 	{
-		SectionToCategoryMap.FindOrAdd(TopLevelHierarchyCategories->GetSection()).Add(TopLevelHierarchyCategories);
+		if(const FDataHierarchyElementMetaData_SectionAssociation* SectionAssociation = TopLevelHierarchyCategories->FindMetaDataOfType<FDataHierarchyElementMetaData_SectionAssociation>())
+		{
+			if(SectionAssociation->Section.IsValid())
+			{
+				SectionToCategoryMap.FindOrAdd(SectionAssociation->Section.Get()).Add(TopLevelHierarchyCategories);	
+			}
+		}
 	}
 	
 	for(const UHierarchySection* HierarchySection : HierarchySections)
@@ -597,14 +603,14 @@ void UNiagaraStackSummaryViewCollection::GetSectionsInternal(TArray<FNiagaraStac
 	}
 }
 
-const TArray<UHierarchySection*>& UNiagaraStackSummaryViewCollection::GetHierarchySections() const
+TConstArrayView<const TObjectPtr<UHierarchySection>> UNiagaraStackSummaryViewCollection::GetHierarchySections() const
 {
 	if (GetEmitterViewModel().IsValid())
 	{
 		return GetEmitterViewModel()->GetEditorData().GetSummarySections();
 	}
 
-	static TArray<UHierarchySection*> Dummy;
+	static TArray<const TObjectPtr<UHierarchySection>> Dummy;
 	return Dummy;
 }
 
diff --git a/Engine/Plugins/FX/Niagara/Source/NiagaraEditor/Private/Widgets/SNiagaraScriptInputPreviewPanel.cpp b/Engine/Plugins/FX/Niagara/Source/NiagaraEditor/Private/Widgets/SNiagaraScriptInputPreviewPanel.cpp
--- a/Engine/Plugins/FX/Niagara/Source/NiagaraEditor/Private/Widgets/SNiagaraScriptInputPreviewPanel.cpp
+++ b/Engine/Plugins/FX/Niagara/Source/NiagaraEditor/Private/Widgets/SNiagaraScriptInputPreviewPanel.cpp
@@ -233,7 +233,7 @@ void SNiagaraScriptInputPreviewPanel::OnSearchButtonClicked(SSearchBox::SearchDi
 
 void SNiagaraScriptInputPreviewPanel::GenerateSearchItems(UHierarchyElement* InRoot, TArray<UHierarchyElement*> ParentChain, TArray<FSearchItem>& OutSearchItems)
 {
-	const TArray<UHierarchyElement*>& Children = InRoot->GetChildren();
+	TConstArrayView<const TObjectPtr<UHierarchyElement>> Children = InRoot->GetChildren();
 	ParentChain.Add(InRoot);
 	OutSearchItems.Add(FSearchItem{ParentChain});
 	for(UHierarchyElement* Child : Children)
diff --git a/Engine/Plugins/FX/Niagara/Source/NiagaraEditor/Public/NiagaraEditorModule.h b/Engine/Plugins/FX/Niagara/Source/NiagaraEditor/Public/NiagaraEditorModule.h
--- a/Engine/Plugins/FX/Niagara/Source/NiagaraEditor/Public/NiagaraEditorModule.h
+++ b/Engine/Plugins/FX/Niagara/Source/NiagaraEditor/Public/NiagaraEditorModule.h
@@ -223,6 +223,8 @@ public:
 	/** Register/unregister niagara editor settings. */
 	void RegisterSettings();
 	void UnregisterSettings();
+
+	void InitializeTaggedAssetBrowserDetails();
 	
 	/** Gets Niagara editor type utilities for a specific type if there are any registered. */
 	TSharedPtr<INiagaraEditorTypeUtilities, ESPMode::ThreadSafe> NIAGARAEDITOR_API GetTypeUtilities(const FNiagaraTypeDefinition& Type);
diff --git a/Engine/Plugins/FX/Niagara/Source/NiagaraEditor/Public/NiagaraEditorSettings.h b/Engine/Plugins/FX/Niagara/Source/NiagaraEditor/Public/NiagaraEditorSettings.h
--- a/Engine/Plugins/FX/Niagara/Source/NiagaraEditor/Public/NiagaraEditorSettings.h
+++ b/Engine/Plugins/FX/Niagara/Source/NiagaraEditor/Public/NiagaraEditorSettings.h
@@ -361,6 +361,14 @@ public:
 	UPROPERTY(config, EditAnywhere, DisplayName="Minimal Emitter", Category = Niagara, meta=(AllowedClasses="/Script/Niagara.NiagaraEmitter"))
 	FSoftObjectPath DefaultEmptyEmitter;
 
+	/** The configuration used in the Niagara Asset Wizard when creating an new Niagara Emitter or adding an emitter to a system. */
+	UPROPERTY(config, EditAnywhere, DisplayName="Emitter Asset Wizard Configuration", Category = Niagara, meta=(AllowedClasses="/Script/UserAssetTagsEditor.TaggedAssetBrowserConfiguration"))
+	FSoftObjectPath EmitterAssetWizardConfiguration;
+
+	/** The configuration used in the Niagara Asset Wizard when creating a new Niagara System. */
+	UPROPERTY(config, EditAnywhere, DisplayName="System Asset Wizard Configuration", Category = Niagara, meta=(AllowedClasses="/Script/UserAssetTagsEditor.TaggedAssetBrowserConfiguration"))
+	FSoftObjectPath SystemAssetWizardConfiguration;
+
 	/** Default material to use with new decal renderers. */
 	UPROPERTY(config, EditAnywhere, Category = "Niagara Renderer", meta = (AllowedClasses = "/Script/Engine.MaterialInterface"))
 	FSoftObjectPath DefaultDecalRendererMaterial;
diff --git a/Engine/Plugins/FX/Niagara/Source/NiagaraEditor/Public/NiagaraEditorStyle.h b/Engine/Plugins/FX/Niagara/Source/NiagaraEditor/Public/NiagaraEditorStyle.h
--- a/Engine/Plugins/FX/Niagara/Source/NiagaraEditor/Public/NiagaraEditorStyle.h
+++ b/Engine/Plugins/FX/Niagara/Source/NiagaraEditor/Public/NiagaraEditorStyle.h
@@ -24,7 +24,6 @@ private:
 	FNiagaraEditorStyle();
 	void InitStats();
 	void InitAssetPicker();
-	void InitAssetBrowser();
 	void InitActionMenu();
 	void InitEmitterHeader();
 	void InitParameters();
diff --git a/Engine/Plugins/FX/Niagara/Source/NiagaraEditor/Public/NiagaraEditorUtilities.h b/Engine/Plugins/FX/Niagara/Source/NiagaraEditor/Public/NiagaraEditorUtilities.h
--- a/Engine/Plugins/FX/Niagara/Source/NiagaraEditor/Public/NiagaraEditorUtilities.h
+++ b/Engine/Plugins/FX/Niagara/Source/NiagaraEditor/Public/NiagaraEditorUtilities.h
@@ -513,36 +513,6 @@ namespace FNiagaraEditorUtilities
 		}
 	}
 
-	namespace AssetBrowser
-	{
-		enum EAssetTagSectionSource
-		{		
-			NiagaraInternal,
-			Project,
-			Other
-		};
-		
-		struct FStructuredAssetTagDefinitionLookupData
-		{
-			UNiagaraAssetTagDefinitions* DefinitionsAsset;
-			TArray<FNiagaraAssetTagDefinition> AssetTagDefinitions;
-			EAssetTagSectionSource Source;
-		};
-
-		
-		/** Returns a fully sorted array of tag definitions. Sorted by [Source.AssetSortOrder.(Optionally)TagNames].
-		 * If not sorted by tag names, the order defined in the asset is used. Sorting by tag names makes sense for flat lists. */
-		TArray<FStructuredAssetTagDefinitionLookupData> GetStructuredSortedAssetTagDefinitions(bool bSortTagsByName = false);
-		/** Same as above, but returns a flat list. If end result is not sorted, the order is [Source.AssetSorderOrder.OriginalTagOrder]. */
-		TArray<FNiagaraAssetTagDefinition> GetFlatSortedAssetTagDefinitions(bool bSortEndResultTagsByName = false);
-
-		const FNiagaraAssetTagDefinition& FindTagDefinitionForReference(const FNiagaraAssetTagDefinitionReference& Reference);
-
-		/* Expects a UNiagaraAssetTagDefinition asset */
-		EAssetTagSectionSource GetAssetTagDefinitionSource(const FAssetData& AssetData);
-		FText GetAssetTagSectionNameFromSource(EAssetTagSectionSource Source);
-	}
-
 	namespace HierarchyEditor
 	{		
 		namespace Scripts
diff --git a/Engine/Plugins/FX/Niagara/Source/NiagaraEditor/Public/NiagaraEmitterEditorData.h b/Engine/Plugins/FX/Niagara/Source/NiagaraEditor/Public/NiagaraEmitterEditorData.h
--- a/Engine/Plugins/FX/Niagara/Source/NiagaraEditor/Public/NiagaraEmitterEditorData.h
+++ b/Engine/Plugins/FX/Niagara/Source/NiagaraEditor/Public/NiagaraEmitterEditorData.h
@@ -45,7 +45,7 @@ public:
 	NIAGARAEDITOR_API void ToggleShowSummaryView();
 	
 	UHierarchyRoot* GetSummaryRoot() const { return SummaryViewRoot; }
-	const TArray<UHierarchySection*>& GetSummarySections() const;
+	TConstArrayView<const TObjectPtr<UHierarchySection>> GetSummarySections() const;
 
 	UTexture2D* GetThumbnail() const { return EmitterThumbnail; }
 	void SetThumbnail(UTexture2D* InThumbnail) { EmitterThumbnail = InThumbnail; OnPersistentDataChanged().Broadcast(); }
diff --git a/Engine/Plugins/FX/Niagara/Source/NiagaraEditor/Public/NiagaraEmitterFactoryNew.h b/Engine/Plugins/FX/Niagara/Source/NiagaraEditor/Public/NiagaraEmitterFactoryNew.h
--- a/Engine/Plugins/FX/Niagara/Source/NiagaraEditor/Public/NiagaraEmitterFactoryNew.h
+++ b/Engine/Plugins/FX/Niagara/Source/NiagaraEditor/Public/NiagaraEmitterFactoryNew.h
@@ -1,19 +1,20 @@
 // Copyright Epic Games, Inc. All Rights Reserved.
 
 #pragma once
-#include "NiagaraSettings.h"
+
 #include "Factories/Factory.h"
+#include "NiagaraSettings.h"
+#include "FrontendFilterBase.h"
 
 #include "NiagaraEmitterFactoryNew.generated.h"
 
-
 /** A factory for niagara emitter assets. */
 UCLASS(hidecategories = Object, MinimalAPI)
 class UNiagaraEmitterFactoryNew : public UFactory
 {
 	GENERATED_UCLASS_BODY()
 
-	UNiagaraEmitter* EmitterToCopy;
+	TWeakObjectPtr<UNiagaraEmitter> EmitterToCopy;
 	bool bUseInheritance;
 	bool bAddDefaultModulesAndRenderersToEmptyEmitter;
 
@@ -24,6 +25,12 @@ class UNiagaraEmitterFactoryNew : public UFactory
 
 public:
 	NIAGARAEDITOR_API static void InitializeEmitter(UNiagaraEmitter* NewEmitter, bool bAddDefaultModulesAndRenderers);
+
+private:
+	bool OnAdditionalShouldFilterAsset(const FAssetData& AssetData) const;
+	void OnExtendAddFilterMenu(UToolMenu* ToolMenu) const;
+	TArray<TSharedRef<FFrontendFilter>> OnGetExtraFrontendFilters() const;
+
 };
 
 
diff --git a/Engine/Plugins/FX/Niagara/Source/NiagaraEditor/Public/NiagaraSystemFactoryNew.h b/Engine/Plugins/FX/Niagara/Source/NiagaraEditor/Public/NiagaraSystemFactoryNew.h
--- a/Engine/Plugins/FX/Niagara/Source/NiagaraEditor/Public/NiagaraSystemFactoryNew.h
+++ b/Engine/Plugins/FX/Niagara/Source/NiagaraEditor/Public/NiagaraSystemFactoryNew.h
@@ -3,10 +3,10 @@
 #pragma once
 
 #include "CoreMinimal.h"
-
 #include "NiagaraTypes.h"
 #include "UObject/ObjectMacros.h"
 #include "Factories/Factory.h"
+#include "FrontendFilterBase.h"
 #include "NiagaraSystemFactoryNew.generated.h"
 
 class UNiagaraSystem;
@@ -30,6 +30,10 @@ public:
 
 private:
 	void TryAssignDefaultEffectType(UNiagaraSystem* System);
+	
+	bool OnAdditionalShouldFilterAsset(const FAssetData& AssetData) const;
+	void OnExtendAddFilterMenu(UToolMenu* ToolMenu) const;
+	TArray<TSharedRef<FFrontendFilter>> OnGetExtraFrontendFilters() const;
 };
 
 
diff --git a/Engine/Plugins/FX/Niagara/Source/NiagaraEditor/Public/NiagaraValidationRules.h b/Engine/Plugins/FX/Niagara/Source/NiagaraEditor/Public/NiagaraValidationRules.h
--- a/Engine/Plugins/FX/Niagara/Source/NiagaraEditor/Public/NiagaraValidationRules.h
+++ b/Engine/Plugins/FX/Niagara/Source/NiagaraEditor/Public/NiagaraValidationRules.h
@@ -373,8 +373,8 @@ public:
 
 /** This validation rule will check if a system uses emitters that are tagged as Deprecated using the Niagara Asset Tags.
  *  This is distinct from a Niagara Emitter version that is marked as deprecated, but might have a new, non-deprecated version. */
-UCLASS(Category = "Validation", DisplayName = "Check for Deprecated Emitters")
-class UNiagaraValidationRule_CheckDeprecatedEmitters : public UNiagaraValidationRule
+UCLASS(Category = "Validation", DisplayName = "Check for Deprecated Emitters", meta=(DeprecationMessage="This validation rule is no longer valid. Delete it from your setup"), Deprecated)
+class UDEPRECATED_NiagaraValidationRule_CheckDeprecatedEmitters : public UNiagaraValidationRule
 {
 	GENERATED_BODY()
 public:
diff --git a/Engine/Plugins/FX/Niagara/Source/NiagaraEditor/Public/ViewModels/HierarchyEditor/NiagaraScriptParametersHierarchyViewModel.h b/Engine/Plugins/FX/Niagara/Source/NiagaraEditor/Public/ViewModels/HierarchyEditor/NiagaraScriptParametersHierarchyViewModel.h
--- a/Engine/Plugins/FX/Niagara/Source/NiagaraEditor/Public/ViewModels/HierarchyEditor/NiagaraScriptParametersHierarchyViewModel.h
+++ b/Engine/Plugins/FX/Niagara/Source/NiagaraEditor/Public/ViewModels/HierarchyEditor/NiagaraScriptParametersHierarchyViewModel.h
@@ -99,7 +99,7 @@ struct FNiagaraHierarchyScriptCategoryViewModel : public FHierarchyCategoryViewM
 
 protected:
 	// We don't sort anything to allow free authoring of UI
-	virtual void SortChildrenData() const override {}
+	virtual void SortChildrenData() override {}
 	virtual FResultWithUserFeedback CanDropOnInternal(TSharedPtr<FHierarchyElementViewModel> DraggedElement, EItemDropZone ItemDropZone) override;
 };
 
@@ -129,7 +129,7 @@ struct FNiagaraHierarchyScriptRootViewModel : public FHierarchyRootViewModel
 
 protected:
 	// We don't sort anything to allow free authoring of UI
-	virtual void SortChildrenData() const override { }
+	virtual void SortChildrenData() override { }
 };
 
 #undef UE_API
diff --git a/Engine/Plugins/FX/Niagara/Source/NiagaraEditor/Public/ViewModels/HierarchyEditor/NiagaraSummaryViewViewModel.h b/Engine/Plugins/FX/Niagara/Source/NiagaraEditor/Public/ViewModels/HierarchyEditor/NiagaraSummaryViewViewModel.h
--- a/Engine/Plugins/FX/Niagara/Source/NiagaraEditor/Public/ViewModels/HierarchyEditor/NiagaraSummaryViewViewModel.h
+++ b/Engine/Plugins/FX/Niagara/Source/NiagaraEditor/Public/ViewModels/HierarchyEditor/NiagaraSummaryViewViewModel.h
@@ -192,13 +192,11 @@ struct FNiagaraFunctionViewModel : public FHierarchyItemViewModel
 	
 	bool IsFromBaseEmitter() const;
 
-	void SetSection(UHierarchySection& InSection) { Section = &InSection; }
-
 	bool IsDynamicInput() const { return bIsDynamicInput; }
 private:	
 	virtual void Initialize() override;
 	virtual void RefreshChildrenDataInternal() override;
-	void RefreshChildrenInputs(bool bClearCache = false) const;
+	void RefreshChildrenInputs(bool bClearCache = false);
 	
 	virtual FString ToString() const override;
 
@@ -206,8 +204,9 @@ private:
 	virtual FResultWithUserFeedback CanHaveChildren() const override { return bIsForHierarchy == false; }
 	virtual bool CanRenameInternal() override { return false; }
 	virtual FResultWithUserFeedback CanDropOnInternal(TSharedPtr<FHierarchyElementViewModel>, EItemDropZone ItemDropZone) override;
-	virtual const UHierarchySection* GetSectionInternal() const override;
-	
+
+	virtual bool IsExpandedByDefaultInternal() const override { return false; }
+
 	virtual bool RepresentsExternalData() const override { return true; }
 	virtual bool DoesExternalDataStillExist(const UHierarchyDataRefreshContext* Context) const override { ClearCache(); return GetFunctionCallNode().IsValid(); }
 	
@@ -217,7 +216,6 @@ private:
 private:
 	FDelegateHandle OnScriptAppliedHandle;
 	mutable TOptional<bool> IsFromBaseEmitterCache;
-	TWeakObjectPtr<UHierarchySection> Section;
 	bool bIsDynamicInput = false;
 };
 
@@ -251,7 +249,8 @@ protected:
 	void ClearCache() const;
 
 	void RefreshChildDynamicInputs(bool bClearCache = false);
-	
+
+	virtual bool IsExpandedByDefaultInternal() const override { return false; }
 	virtual FResultWithUserFeedback CanHaveChildren() const override;
 	virtual FResultWithUserFeedback CanContainInternal(TSubclassOf<UHierarchyElement> InHierarchyElementType) override;
 	virtual FResultWithUserFeedback IsEditableByUser() override;
@@ -292,6 +291,7 @@ struct FNiagaraAssignmentInputViewModel : public FHierarchyItemViewModel
 
 	void ClearCache() const;
 protected:
+	virtual bool IsExpandedByDefaultInternal() const override { return false; }
 	virtual FResultWithUserFeedback IsEditableByUser() override;
 	virtual bool RepresentsExternalData() const override { return true; }
 	virtual bool DoesExternalDataStillExist(const UHierarchyDataRefreshContext* Context) const override { ClearCache(); return GetInputData().IsSet(); }
@@ -311,6 +311,7 @@ struct FNiagaraHierarchySummaryCategoryViewModel : public FHierarchyCategoryView
 	bool IsFromBaseEmitter() const;
 	
 protected:
+	virtual bool IsExpandedByDefaultInternal() const override { return false; }
 	virtual FResultWithUserFeedback IsEditableByUser() override;
 	
 	mutable TOptional<bool> IsFromBaseEmitterCache;
@@ -343,21 +344,18 @@ struct FNiagaraHierarchyRendererViewModel : public FHierarchyItemViewModel
 
 	UNiagaraRendererProperties* GetRendererProperties() const;
 
-	void SetSection(UHierarchySection& InSection) { Section = &InSection; }
-
 	bool IsFromBaseEmitter() const;
 protected:
 	virtual void RefreshChildrenDataInternal() override;
+
+	virtual bool IsExpandedByDefaultInternal() const override { return false; }
 	
 	virtual FResultWithUserFeedback IsEditableByUser() override;
 	virtual FResultWithUserFeedback CanHaveChildren() const override { return bIsForHierarchy == false; }
 	
 	virtual bool RepresentsExternalData() const override { return true; }
 	virtual bool DoesExternalDataStillExist(const UHierarchyDataRefreshContext* Context) const override { return GetRendererProperties() != nullptr;}
-
-	virtual const UHierarchySection* GetSectionInternal() const override;
 	
-	TWeakObjectPtr<UHierarchySection> Section;
 	mutable TOptional<bool> IsFromBaseEmitterCache;
 };
 
@@ -370,19 +368,16 @@ struct FNiagaraHierarchyEmitterPropertiesViewModel : public FHierarchyItemViewMo
 	virtual FString ToString() const override;
 	virtual bool CanRenameInternal() override { return false; }
 
-	void SetSection(UHierarchySection& InSection) { Section = &InSection; }
-
 	bool IsFromBaseEmitter() const;
-protected:	
+protected:
+	virtual bool IsExpandedByDefaultInternal() const override { return false; }
+
 	virtual FResultWithUserFeedback IsEditableByUser() override;
 	virtual FResultWithUserFeedback CanHaveChildren() const override { return bIsForHierarchy == false; }
 	
 	virtual bool RepresentsExternalData() const override { return true; }
 	virtual bool DoesExternalDataStillExist(const UHierarchyDataRefreshContext* Context) const override { return true; }
-
-	virtual const UHierarchySection* GetSectionInternal() const override;
 	
-	TWeakObjectPtr<UHierarchySection> Section;
 	mutable TOptional<bool> IsFromBaseEmitterCache;
 };
 
@@ -396,8 +391,6 @@ struct FNiagaraHierarchyEventHandlerViewModel : public FHierarchyItemViewModel
 
 	FNiagaraEventScriptProperties* GetEventScriptProperties() const;
 
-	void SetSection(UHierarchySection& InSection) { Section = &InSection; }
-
 	bool IsFromBaseEmitter() const;
 protected:
 	virtual void RefreshChildrenDataInternal() override;
@@ -407,10 +400,7 @@ protected:
 	
 	virtual bool RepresentsExternalData() const override { return true; }
 	virtual bool DoesExternalDataStillExist(const UHierarchyDataRefreshContext* Context) const override { return GetEventScriptProperties() != nullptr;}
-
-	virtual const UHierarchySection* GetSectionInternal() const override;
 	
-	TWeakObjectPtr<UHierarchySection> Section;
 	mutable TOptional<bool> IsFromBaseEmitterCache;
 };
 
@@ -446,9 +436,7 @@ struct FNiagaraHierarchySimStageViewModel : public FHierarchyItemViewModel
 	virtual bool CanRenameInternal() override { return false; }
 
 	UNiagaraSimulationStageBase* GetSimStage() const;
-
-	void SetSection(UHierarchySection& InSection) { Section = &InSection; }
-
+	
 	bool IsFromBaseEmitter() const;
 protected:
 	virtual void RefreshChildrenDataInternal() override;
@@ -458,10 +446,7 @@ protected:
 	
 	virtual bool RepresentsExternalData() const override { return true; }
 	virtual bool DoesExternalDataStillExist(const UHierarchyDataRefreshContext* Context) const override { return GetSimStage() != nullptr;}
-
-	virtual const UHierarchySection* GetSectionInternal() const override;
 	
-	TWeakObjectPtr<UHierarchySection> Section;
 	mutable TOptional<bool> IsFromBaseEmitterCache;
 };
 
diff --git a/Engine/Plugins/FX/Niagara/Source/NiagaraEditor/Public/ViewModels/Stack/NiagaraStackScriptHierarchyRoot.h b/Engine/Plugins/FX/Niagara/Source/NiagaraEditor/Public/ViewModels/Stack/NiagaraStackScriptHierarchyRoot.h
--- a/Engine/Plugins/FX/Niagara/Source/NiagaraEditor/Public/ViewModels/Stack/NiagaraStackScriptHierarchyRoot.h
+++ b/Engine/Plugins/FX/Niagara/Source/NiagaraEditor/Public/ViewModels/Stack/NiagaraStackScriptHierarchyRoot.h
@@ -60,9 +60,9 @@ public:
 	UNiagaraNodeFunctionCall& GetOwningFunctionCallNode() const { return *OwningFunctionCallNode.Get(); }
 	UHierarchyRoot* GetScriptParameterHierarchyRoot() const;
 	
-	NIAGARAEDITOR_API const TArray<UHierarchySection*>& GetSections() const;
-	NIAGARAEDITOR_API const UHierarchySection* GetActiveSection() const;
-	NIAGARAEDITOR_API void SetActiveSection(const UHierarchySection* InActiveSection);
+	NIAGARAEDITOR_API TConstArrayView<const TObjectPtr<UHierarchySection>> GerHierarchySectionData() const;
+	NIAGARAEDITOR_API const UHierarchySection* GetActiveHierarchySection() const;
+	NIAGARAEDITOR_API void SetActiveHierarchySection(const UHierarchySection* InActiveSection);
 
 	void SetShouldDisplayLabel(bool bInShouldDisplayLabel);
 	NIAGARAEDITOR_API bool GetShouldDisplayLabel() const;
diff --git a/Engine/Plugins/FX/Niagara/Source/NiagaraEditor/Public/ViewModels/Stack/NiagaraStackSummaryViewInputCollection.h b/Engine/Plugins/FX/Niagara/Source/NiagaraEditor/Public/ViewModels/Stack/NiagaraStackSummaryViewInputCollection.h
--- a/Engine/Plugins/FX/Niagara/Source/NiagaraEditor/Public/ViewModels/Stack/NiagaraStackSummaryViewInputCollection.h
+++ b/Engine/Plugins/FX/Niagara/Source/NiagaraEditor/Public/ViewModels/Stack/NiagaraStackSummaryViewInputCollection.h
@@ -30,7 +30,7 @@ protected:
 	NIAGARAEDITOR_API virtual void RefreshChildrenInternal(const TArray<UNiagaraStackEntry*>& CurrentChildren, TArray<UNiagaraStackEntry*>& NewChildren, TArray<FStackIssue>& NewIssues) override;
 
 	NIAGARAEDITOR_API virtual void GetSectionsInternal(TArray<FNiagaraStackSection>& OutStackSections) const override;
-	NIAGARAEDITOR_API const TArray<UHierarchySection*>& GetHierarchySections() const;
+	NIAGARAEDITOR_API TConstArrayView<const TObjectPtr<UHierarchySection>> GetHierarchySections() const;
 
 	NIAGARAEDITOR_API virtual bool FilterByActiveSection(const UNiagaraStackEntry& Child) const override;
 	bool FilterOnlyModified(const UNiagaraStackEntry& Child) const;
diff --git a/Engine/Plugins/FX/Niagara/Source/NiagaraEditorWidgets/NiagaraEditorWidgets.Build.cs b/Engine/Plugins/FX/Niagara/Source/NiagaraEditorWidgets/NiagaraEditorWidgets.Build.cs
--- a/Engine/Plugins/FX/Niagara/Source/NiagaraEditorWidgets/NiagaraEditorWidgets.Build.cs
+++ b/Engine/Plugins/FX/Niagara/Source/NiagaraEditorWidgets/NiagaraEditorWidgets.Build.cs
@@ -36,7 +36,8 @@ public class NiagaraEditorWidgets : ModuleRules
 			"ToolWidgets",
 			"KismetWidgets",
 			"ToolMenus",
-			"DataHierarchyEditor"
+			"DataHierarchyEditor",
+			"UserAssetTagsEditor"
 		});
 
 		PrivateIncludePathModuleNames.AddRange(new string[] {
diff --git a/Engine/Plugins/FX/Niagara/Source/NiagaraEditorWidgets/Private/SNiagaraOverviewGraph.cpp b/Engine/Plugins/FX/Niagara/Source/NiagaraEditorWidgets/Private/SNiagaraOverviewGraph.cpp
--- a/Engine/Plugins/FX/Niagara/Source/NiagaraEditorWidgets/Private/SNiagaraOverviewGraph.cpp
+++ b/Engine/Plugins/FX/Niagara/Source/NiagaraEditorWidgets/Private/SNiagaraOverviewGraph.cpp
@@ -23,14 +23,17 @@
 #include "SNiagaraOverviewGraphTitleBar.h"
 #include "SNiagaraStack.h"
 #include "ScopedTransaction.h"
+#include "Assets/TaggedAssetBrowserConfiguration.h"
 #include "Interfaces/IMainFrameModule.h"
+#include "Logging/StructuredLog.h"
 #include "Stateless/NiagaraStatelessEmitter.h"
 #include "Templates/SharedPointer.h"
 #include "ViewModels/NiagaraEmitterHandleViewModel.h"
 #include "ViewModels/NiagaraOverviewGraphViewModel.h"
 #include "ViewModels/NiagaraSystemViewModel.h"
 #include "Widgets/SItemSelector.h"
-#include "Widgets/AssetBrowser/SNiagaraAssetBrowser.h"
+#include "Widgets/STaggedAssetBrowser.h"
+#include "Widgets/AssetBrowser/SNiagaraAddEmitterToSystemWindow.h"
 
 #define LOCTEXT_NAMESPACE "NiagaraOverviewGraph"
 
@@ -220,8 +223,17 @@ void SNiagaraOverviewGraph::OpenAddEmitterMenu()
 	IMainFrameModule& MainFrame = FModuleManager::LoadModuleChecked<IMainFrameModule>("MainFrame");
 	TSharedPtr<SWindow>	ParentWindow = MainFrame.GetParentWindow();
 
-	TSharedRef<SWindow> AddEmitterWindow = SNew(SNiagaraAddEmitterToSystemWindow, ViewModel->GetSystemViewModel());
-	FSlateApplication::Get().AddModalWindow(AddEmitterWindow, ParentWindow);
+	FSoftObjectPath Path = GetDefault<UNiagaraEditorSettings>()->EmitterAssetWizardConfiguration;
+	UObject* ConfigUObject = Path.TryLoad();
+	if(UTaggedAssetBrowserConfiguration* ConfigurationAsset = Cast<UTaggedAssetBrowserConfiguration>(ConfigUObject))
+	{
+		TSharedRef<SWindow> AddEmitterWindow = SNew(SNiagaraAddEmitterToSystemWindow, *ConfigurationAsset, ViewModel->GetSystemViewModel());
+		FSlateApplication::Get().AddModalWindow(AddEmitterWindow, ParentWindow);
+	}
+	else
+	{
+		UE_LOGFMT(LogNiagaraEditor, Warning, "Could not find emitter wizard configuration asset. Check Niagara Editor settings under Project Settings");
+	}
 }
 
 bool SNiagaraOverviewGraph::CanAddEmitters() const
diff --git a/Engine/Plugins/FX/Niagara/Source/NiagaraEditorWidgets/Private/SNiagaraOverviewGraphTitleBar.cpp b/Engine/Plugins/FX/Niagara/Source/NiagaraEditorWidgets/Private/SNiagaraOverviewGraphTitleBar.cpp
--- a/Engine/Plugins/FX/Niagara/Source/NiagaraEditorWidgets/Private/SNiagaraOverviewGraphTitleBar.cpp
+++ b/Engine/Plugins/FX/Niagara/Source/NiagaraEditorWidgets/Private/SNiagaraOverviewGraphTitleBar.cpp
@@ -8,6 +8,7 @@
 #include "ViewModels/NiagaraOverviewGraphViewModel.h"
 #include "SNiagaraScalabilityPreviewSettings.h"
 #include "SNiagaraSystemEffectTypeBar.h"
+#include "UserAssetTagEditorUtilities.h"
 #include "AssetRegistry/AssetRegistryModule.h"
 #include "ViewModels/NiagaraEmitterHandleViewModel.h"
 
@@ -180,12 +181,9 @@ FText SNiagaraOverviewGraphTitleBar::GetSystemSubheaderText() const
 			{
 				if (FVersionedNiagaraEmitterData* ParentData = EmitterData->GetParent().GetEmitterData())
 				{
-					bool bIsDeprecatedAsset = EmitterData->GetParent().Emitter->AssetTags.ContainsByPredicate([](const FNiagaraAssetTagDefinitionReference& AssetTagDefinitionReferenceCandidate)
-					{
-						return AssetTagDefinitionReferenceCandidate.GetTagDefinitionReferenceGuid() == INiagaraModule::Get().DeprecatedTagDefinition.TagGuid;
-					});
+					bool bHasDeprecatedParentAsset = UE::UserAssetTags::HasUserAssetTag(EmitterData->GetParent().Emitter, "Deprecated");
 					
-					if (bIsDeprecatedAsset || ParentData->bDeprecated)
+					if (bHasDeprecatedParentAsset || ParentData->bDeprecated)
 					{
 						DeprecatedEmitterNames.Add(FText::FromName(EmitterModel->GetName()));
 					}
@@ -302,10 +300,8 @@ bool SNiagaraOverviewGraphTitleBar::IsUsingDeprecatedEmitter() const
 			{
 				if(FVersionedNiagaraEmitterData* ParentData = EmitterData->GetParent().GetEmitterData())
 				{
-					return ParentData->bDeprecated || EmitterData->GetParent().Emitter->AssetTags.ContainsByPredicate([](const FNiagaraAssetTagDefinitionReference& AssetTagDefinitionReferenceCandidate)
-					{
-						return AssetTagDefinitionReferenceCandidate.GetTagDefinitionReferenceGuid() == INiagaraModule::Get().DeprecatedTagDefinition.TagGuid;
-					});
+					bool bHasDeprecatedParentAsset = UE::UserAssetTags::HasUserAssetTag(EmitterData->GetParent().Emitter, "Deprecated");
+					return ParentData->bDeprecated || bHasDeprecatedParentAsset;
 				}
 			}
 		}
diff --git a/Engine/Plugins/FX/Niagara/Source/NiagaraEditorWidgets/Private/Stack/SNiagaraStackScriptHierarchyRoot.cpp b/Engine/Plugins/FX/Niagara/Source/NiagaraEditorWidgets/Private/Stack/SNiagaraStackScriptHierarchyRoot.cpp
--- a/Engine/Plugins/FX/Niagara/Source/NiagaraEditorWidgets/Private/Stack/SNiagaraStackScriptHierarchyRoot.cpp
+++ b/Engine/Plugins/FX/Niagara/Source/NiagaraEditorWidgets/Private/Stack/SNiagaraStackScriptHierarchyRoot.cpp
@@ -81,21 +81,21 @@ EVisibility SNiagaraStackScriptHierarchyRoot::GetLabelVisibility() const
 
 ECheckBoxState SNiagaraStackScriptHierarchyRoot::IsSectionChecked(const UHierarchySection* NiagaraHierarchySection) const
 {
-	return ModuleHierarchyRoot->GetActiveSection() == NiagaraHierarchySection ? ECheckBoxState::Checked : ECheckBoxState::Unchecked;
+	return ModuleHierarchyRoot->GetActiveHierarchySection() == NiagaraHierarchySection ? ECheckBoxState::Checked : ECheckBoxState::Unchecked;
 }
 
 void SNiagaraStackScriptHierarchyRoot::OnCheckStateChanged(ECheckBoxState CheckBoxState, const UHierarchySection* NiagaraHierarchySection)
 {
 	if(CheckBoxState == ECheckBoxState::Checked)
 	{
-		ModuleHierarchyRoot->SetActiveSection(NiagaraHierarchySection);
+		ModuleHierarchyRoot->SetActiveHierarchySection(NiagaraHierarchySection);
 	}
 }
 
 void SNiagaraStackScriptHierarchyRoot::ConstructSectionButtons()
 {
 	SectionSelectorBox->ClearChildren();
-	for (const UHierarchySection* Section : ModuleHierarchyRoot->GetSections())
+	for (const UHierarchySection* Section : ModuleHierarchyRoot->GerHierarchySectionData())
 	{
 		SectionSelectorBox->AddSlot()
 		[
diff --git a/Engine/Plugins/FX/Niagara/Source/NiagaraEditorWidgets/Private/Stack/SNiagaraStackTableRow.cpp b/Engine/Plugins/FX/Niagara/Source/NiagaraEditorWidgets/Private/Stack/SNiagaraStackTableRow.cpp
--- a/Engine/Plugins/FX/Niagara/Source/NiagaraEditorWidgets/Private/Stack/SNiagaraStackTableRow.cpp
+++ b/Engine/Plugins/FX/Niagara/Source/NiagaraEditorWidgets/Private/Stack/SNiagaraStackTableRow.cpp
@@ -664,7 +664,8 @@ void SNiagaraStackTableRow::ToggleShowInSummaryView() const
 			{
 				UClass* HierarchyClass = DetermineHierarchyClassForSummaryView();
 				check(HierarchyClass);
-				StackEntry->GetEmitterViewModel()->GetSummaryHierarchyViewModel()->AddElementUnderRoot(HierarchyClass, Identity);
+				UHierarchyElement* Element = StackEntry->GetEmitterViewModel()->GetSummaryHierarchyViewModel()->GetHierarchyRoot()->AddChild<UHierarchyElement>(HierarchyClass);
+				Element->SetIdentity(Identity);
 			}
 		}
 	}
diff --git a/Engine/Source/Editor/DataHierarchyEditor/Internal/Widgets/SDataHierarchyEditor.h b/Engine/Source/Editor/DataHierarchyEditor/Internal/Widgets/SDataHierarchyEditor.h
--- a/Engine/Source/Editor/DataHierarchyEditor/Internal/Widgets/SDataHierarchyEditor.h
+++ b/Engine/Source/Editor/DataHierarchyEditor/Internal/Widgets/SDataHierarchyEditor.h
@@ -35,7 +35,7 @@ class SHierarchySection : public SCompoundWidget
 	SLATE_END_ARGS()
 
 public:
-	UE_API void Construct(const FArguments& InArgs, TSharedPtr<struct FHierarchySectionViewModel> InSection);
+	void Construct(const FArguments& InArgs, TSharedPtr<struct FHierarchySectionViewModel> InSection);
 	UE_API virtual ~SHierarchySection() override;
 
 	UE_API void TryEnterEditingMode() const;
@@ -65,6 +65,9 @@ private:
 
 	UE_API UHierarchySection* TryGetSectionData() const;
 
+	const FSlateBrush* GetImageBrush() const;
+	EVisibility GetImageVisibility() const;
+
 	UE_API FText GetText() const;
 	UE_API FText GetTooltipText() const;
 	UE_API void OnRenameSection(const FText& Text, ETextCommit::Type CommitType) const;
@@ -157,11 +160,8 @@ private:
 
 	UE_API void NavigateToMatchingHierarchyElementFromSelectedSourceElement() const;
 	UE_API bool CanNavigateToMatchingHierarchyElementFromSelectedSourceElement() const;
-
-	UE_API void DeleteActiveSection() const;
-	UE_API bool CanDeleteActiveSection() const;
-
-	UE_API void OnElementAdded(TSharedPtr<FHierarchyElementViewModel> AddedItem);
+	
+	UE_API void OnHierarchyElementChanged(TInstancedStruct<FHierarchyElementChangedPayload> Payload);
 	UE_API void OnHierarchySectionActivated(TSharedPtr<FHierarchySectionViewModel> Section);
 	UE_API void OnSourceSectionActivated(TSharedPtr<FHierarchySectionViewModel> Section);
 	UE_API void OnHierarchySectionAdded(TSharedPtr<FHierarchySectionViewModel> AddedSection);
@@ -212,6 +212,8 @@ private:
 	UE_API void SelectNextSourceSearchResult();
 	UE_API void SelectPreviousSourceSearchResult();
 	UE_API TOptional<SSearchBox::FSearchResultData> GetSearchResultData() const;
+
+	void ExpandEntriesByDefault(TSharedPtr<STreeView<TSharedPtr<FHierarchyElementViewModel>>> TreeView) const;
 	
 	FHierarchyElementViewModel::FResultWithUserFeedback CanDropOnRoot(TSharedPtr<FHierarchyElementViewModel> DraggedItem) const;
 
diff --git a/Engine/Source/Editor/DataHierarchyEditor/Private/DataHierarchyEditorStyle.cpp b/Engine/Source/Editor/DataHierarchyEditor/Private/DataHierarchyEditorStyle.cpp
--- a/Engine/Source/Editor/DataHierarchyEditor/Private/DataHierarchyEditorStyle.cpp
+++ b/Engine/Source/Editor/DataHierarchyEditor/Private/DataHierarchyEditorStyle.cpp
@@ -39,7 +39,7 @@ void FDataHierarchyEditorStyle::InitDataHierarchyEditor()
 	
 	Set("HierarchyEditor.Drop.Section.Left", new CORE_BOX_BRUSH("Common/DropZoneIndicator_Above", FMargin(10.0f / 16.0f, 10.0f / 16.0f, 0.f, 0.f), SelectionColor));
 	Set("HierarchyEditor.Drop.Section.Onto", new CORE_BOX_BRUSH("Common/DropZoneIndicator_Onto", FMargin(4.0f / 16.0f), SelectionColor));
-	Set("HierarchyEditor.Drop.Section.Below", new CORE_BOX_BRUSH("Common/DropZoneIndicator_Below", FMargin(10.0f / 16.0f, 0.f, 0.f, 10.0f / 16.0f), SelectionColor));
+	Set("HierarchyEditor.Drop.Section.Right", new CORE_BOX_BRUSH("Common/DropZoneIndicator_Below", FMargin(10.0f / 16.0f, 0.f, 0.f, 10.0f / 16.0f), SelectionColor));
 	
 	const FTextBlockStyle NormalText = FAppStyle::GetWidgetStyle<FTextBlockStyle>("NormalText");
 	FSlateFontInfo CategoryFont = FAppStyle::Get().GetFontStyle(TEXT("DetailsView.CategoryFontStyle"));
