[Backout] - CL45271379
#rnx
Original CL Desc
-----------------------------------------------------------------
User Asset Tagging: v3

Tags are simply Names stored in FMetaData, serialized to Asset Registry.
This CL contains:
- User Asset Tags, Editor, API
-- Hit Ctrl T on selected assets to summon the User Asset Tags Editor that allows you to assign tags
-- User Asset Tag Providers, with a Local Favorites and a Project Settings Tag provider implementations

- Tagged Asset Browser
-- Asset to configure filters
-- Added section support
-- Integrated Data Hierarchy Editor to organize filters, sections etc.
-- Filters are extendable; you can write new filters, assign icons etc.

- Changes to Data Hierarchy Editor
-- Reworked sections; from hardcoded sections in categories to a metadata approach. Any element that can be dropped onto a section will create a section association
-- Reworked various parts of the DHT, like removing unnecessary functions, streamlining code more, adding default User Feedback for various actions
-- New 'Add new element via UI' API - adding new elements via "Add" menu

- Niagara Changes
-- Moving Niagara Asset Tags to User Asset Tags; deprecating old related classes and functions
-- Emitters and Systems look up the display names of their old tags, and move them to the new tags contained in FMetaData during PostLoad

- UPackage Changes
-- Added "transfer metadata requested delegate" that is called from ObjectTools when renaming/moving a package. The UAT system hooks into it to transfer the UAT to the new metadata object

diff --git a/Engine/Plugins/FX/Niagara/Source/Niagara/Classes/NiagaraEmitter.h b/Engine/Plugins/FX/Niagara/Source/Niagara/Classes/NiagaraEmitter.h
--- a/Engine/Plugins/FX/Niagara/Source/Niagara/Classes/NiagaraEmitter.h
+++ b/Engine/Plugins/FX/Niagara/Source/Niagara/Classes/NiagaraEmitter.h
@@ -10,6 +10,7 @@
 #include "NiagaraMessageDataBase.h"
 #include "NiagaraMessageStore.h"
 #include "INiagaraMergeManager.h"
+#include "NiagaraAssetTagDefinitions.h"
 #include "NiagaraEffectType.h"
 #include "NiagaraEmitterBase.h"
 #include "NiagaraDataInterfaceEmitterBinding.h"
@@ -24,7 +25,6 @@
 
 #include "NiagaraEmitter.generated.h"
 
-struct FNiagaraAssetTagDefinitionReference;
 class UNiagaraScratchPadContainer;
 class UMaterial;
 class UNiagaraEmitter;
@@ -746,6 +746,9 @@ private:
 public:
 	
 #if WITH_EDITORONLY_DATA
+	UPROPERTY()
+	TArray<FNiagaraAssetTagDefinitionReference> AssetTags;
+
 	/** If an emitter is inheritable, new emitters based on an inheritable emitter, or Niagara Systems using an inheritable emitter, will automatically inherit changes made to the original emitter. */
 	UPROPERTY(EditAnywhere, Category = "Asset Options", AssetRegistrySearchable)
 	bool bIsInheritable = true;
@@ -880,9 +883,6 @@ public:
 	/** Use property in struct returned from GetEmitterData() instead */ 
 	UPROPERTY(meta = (DeprecatedProperty))
 	uint32 MaxGPUParticlesSpawnPerFrame_DEPRECATED;
-
-	UPROPERTY(meta = (DeprecatedProperty))
-	TArray<FNiagaraAssetTagDefinitionReference> AssetTags_DEPRECATED;
 #endif
 	
 	NIAGARA_API void UpdateEmitterAfterLoad();
diff --git a/Engine/Plugins/FX/Niagara/Source/Niagara/Classes/NiagaraSystem.h b/Engine/Plugins/FX/Niagara/Source/Niagara/Classes/NiagaraSystem.h
--- a/Engine/Plugins/FX/Niagara/Source/Niagara/Classes/NiagaraSystem.h
+++ b/Engine/Plugins/FX/Niagara/Source/Niagara/Classes/NiagaraSystem.h
@@ -3,6 +3,7 @@
 #pragma once
 
 #include "CoreMinimal.h"
+#include "NiagaraAssetTagDefinitions.h"
 
 #include "NiagaraDataSetCompiledData.h"
 #include "NiagaraDataSetAccessor.h"
@@ -19,7 +20,6 @@
 #include "Particles/FXBudget.h"
 #include "NiagaraSystem.generated.h"
 
-struct FNiagaraAssetTagDefinitionReference;
 class FNiagaraActiveCompilation;
 class UNiagaraEffectType;
 enum class ENiagaraCullProxyMode : uint32;
@@ -493,9 +493,6 @@ public:
 	UPROPERTY()
 	ENiagaraScriptTemplateSpecification TemplateSpecification_DEPRECATED;
 
-	UPROPERTY(meta=(DeprecatedProperty))
-	TArray<FNiagaraAssetTagDefinitionReference> AssetTags_DEPRECATED;
-	
 	UPROPERTY(EditAnywhere, AdvancedDisplay, Category = "Asset Options", AssetRegistrySearchable, DisplayName="Asset Description", meta = (SkipSystemResetOnChange = "true"))
 	FText TemplateAssetDescription;
 
@@ -503,6 +500,9 @@ public:
 	UPROPERTY(AssetRegistrySearchable, meta = (SkipSystemResetOnChange = "true"))
 	FText Category;
 
+	UPROPERTY(meta = (SkipSystemResetOnChange = "true"))
+	TArray<FNiagaraAssetTagDefinitionReference> AssetTags;
+
 	UPROPERTY(EditAnywhere, AdvancedDisplay, Category="Asset Options", AssetRegistrySearchable, meta = (SkipSystemResetOnChange = "true", AllowedClasses="/Script/MediaAssets.FileMediaSource"))
 	FSoftObjectPath PreviewMoviePath;
 	
diff --git a/Engine/Plugins/FX/Niagara/Source/Niagara/Classes/NiagaraValidationRule.h b/Engine/Plugins/FX/Niagara/Source/Niagara/Classes/NiagaraValidationRule.h
--- a/Engine/Plugins/FX/Niagara/Source/Niagara/Classes/NiagaraValidationRule.h
+++ b/Engine/Plugins/FX/Niagara/Source/Niagara/Classes/NiagaraValidationRule.h
@@ -70,7 +70,7 @@ class UNiagaraValidationRule : public UObject
 	GENERATED_BODY()
 
 public:
-	bool IsEnabled() const { return bIsConfigDisabled == false && GetClass()->HasAnyClassFlags(CLASS_Deprecated) == false; }
+	bool IsEnabled() const { return bIsConfigDisabled == false; }
 #if WITH_EDITOR
 	NIAGARA_API virtual void CheckValidity(const FNiagaraValidationContext& Context, TArray<FNiagaraValidationResult>& OutResults) const;
 #endif
diff --git a/Engine/Plugins/FX/Niagara/Source/Niagara/Private/NiagaraAssetTagDefinitions.cpp b/Engine/Plugins/FX/Niagara/Source/Niagara/Private/NiagaraAssetTagDefinitions.cpp
--- a/Engine/Plugins/FX/Niagara/Source/Niagara/Private/NiagaraAssetTagDefinitions.cpp
+++ b/Engine/Plugins/FX/Niagara/Source/Niagara/Private/NiagaraAssetTagDefinitions.cpp
@@ -94,7 +94,7 @@ void FNiagaraAssetTagDefinitionReference::AddTagToAssetRegistryTags(FAssetRegist
 	}
 }
 
-FText UDEPRECATED_NiagaraAssetTagDefinitions::GetDisplayName() const
+FText UNiagaraAssetTagDefinitions::GetDisplayName() const
 {
 	if(DisplayName.IsEmptyOrWhitespace())
 	{
@@ -104,12 +104,12 @@ FText UDEPRECATED_NiagaraAssetTagDefinitions::GetDisplayName() const
 	return DisplayName;
 }
 
-FText UDEPRECATED_NiagaraAssetTagDefinitions::GetDescription() const
+FText UNiagaraAssetTagDefinitions::GetDescription() const
 {
 	return Description;
 }
 
-bool UDEPRECATED_NiagaraAssetTagDefinitions::DoesAssetDataContainAnyTag(const FAssetData& AssetData) const
+bool UNiagaraAssetTagDefinitions::DoesAssetDataContainAnyTag(const FAssetData& AssetData) const
 {
 	for(const FNiagaraAssetTagDefinition& AssetTagDefinition : TagDefinitions)
 	{
diff --git a/Engine/Plugins/FX/Niagara/Source/Niagara/Private/NiagaraEmitter.cpp b/Engine/Plugins/FX/Niagara/Source/Niagara/Private/NiagaraEmitter.cpp
--- a/Engine/Plugins/FX/Niagara/Source/Niagara/Private/NiagaraEmitter.cpp
+++ b/Engine/Plugins/FX/Niagara/Source/Niagara/Private/NiagaraEmitter.cpp
@@ -631,6 +631,15 @@ void UNiagaraEmitter::PostLoad()
 		{
 			bIsInheritable = false;
 		}
+
+		if(CurrentTemplateSpecification == ENiagaraScriptTemplateSpecification::Template)
+		{			
+			AssetTags.AddUnique(INiagaraModule::TemplateTagDefinition);
+		}
+		else if(CurrentTemplateSpecification == ENiagaraScriptTemplateSpecification::Behavior)
+		{
+			AssetTags.AddUnique(INiagaraModule::LearningContentTagDefinition);
+		}
 	}
 	
 	if (MessageKeyToMessageMap_DEPRECATED.IsEmpty() == false)
@@ -649,12 +658,6 @@ void UNiagaraEmitter::PostLoad()
 	}
 #endif
 
-#if WITH_EDITORONLY_DATA	
-	if(NiagaraVer < FNiagaraCustomVersion::Type::UserAssetTagStorageInMetaData && AssetTags_DEPRECATED.Num() > 0)
-	{
-		INiagaraModule::Get().GetEditorOnlyDataUtilities().MoveUserAssetTagsToMetaData(this);
-	}
-#endif
 	const int32 UE5MainVer = GetLinkerCustomVersion(FUE5MainStreamObjectVersion::GUID);
 
 	for (FVersionedNiagaraEmitterData& Data : VersionData)
@@ -1046,6 +1049,11 @@ void UNiagaraEmitter::GetAssetRegistryTags(FAssetRegistryTagsContext Context) co
 	INiagaraModule& NiagaraModule = FModuleManager::GetModuleChecked<INiagaraModule>("Niagara");
 	Context.AddTag(NiagaraModule.GetEditorOnlyDataUtilities().CreateClassUsageAssetRegistryTag(this));
 
+	// Asset Library Tags
+	for(const FNiagaraAssetTagDefinitionReference& Tag : AssetTags)
+	{
+		Tag.AddTagToAssetRegistryTags(Context);
+	}
 #endif
 	Super::GetAssetRegistryTags(Context);
 }
diff --git a/Engine/Plugins/FX/Niagara/Source/Niagara/Private/NiagaraModule.cpp b/Engine/Plugins/FX/Niagara/Source/Niagara/Private/NiagaraModule.cpp
--- a/Engine/Plugins/FX/Niagara/Source/Niagara/Private/NiagaraModule.cpp
+++ b/Engine/Plugins/FX/Niagara/Source/Niagara/Private/NiagaraModule.cpp
@@ -40,7 +40,6 @@
 #include "Particles/FXBudget.h"
 #include "Engine/StaticMesh.h"
 #include "UObject/UObjectGlobals.h"
-#include "NiagaraAssetTagDefinitions.h"
 
 #include "DataInterface/NiagaraDataInterfaceDataChannelCommon.h"
 #include "DataInterface/NiagaraDataInterfaceDataChannelWrite.h"
@@ -67,7 +66,6 @@ static FAutoConsoleVariableRef CVarLogCompileIdGeneration(
 float INiagaraModule::EngineGlobalSpawnCountScale = 1.0f;
 float INiagaraModule::EngineGlobalSystemCountScale = 1.0f;
 
-PRAGMA_DISABLE_DEPRECATION_WARNINGS
 const FNiagaraAssetTagDefinition INiagaraModule::LightweightTagDefinition = {
 	LOCTEXT("LightweightDisplayName", "Lightweight")
 	, (int32) (ENiagaraAssetLibraryAssetTypes::Systems)
@@ -107,7 +105,6 @@ const FNiagaraAssetTagDefinition INiagaraModule::DeprecatedTagDefinition = {
 	, ENiagaraAssetTagDefinitionImportance::Internal
 	, FLinearColor::Black
 	, FGuid(0xF1F2F2FB, 0x50C64AF9, 0xB6BCE631, 0x32C4CC61)};
-PRAGMA_ENABLE_DEPRECATION_WARNINGS
 
 std::atomic<bool> INiagaraModule::bDataChannelRefreshRequested = false;
 
@@ -515,9 +512,7 @@ void INiagaraModule::StartupModule()
 #endif
 #endif
 
-PRAGMA_DISABLE_DEPRECATION_WARNINGS
 	RegisterInternalAssetTagDefinitions();
-PRAGMA_ENABLE_DEPRECATION_WARNINGS
 	
 	FCoreDelegates::OnPostEngineInit.AddRaw(this, &INiagaraModule::OnPostEngineInit);
 	FCoreDelegates::OnPreExit.AddRaw(this, &INiagaraModule::OnPreExit);
@@ -611,13 +606,6 @@ void INiagaraModule::RefreshDataChannels()
 	}
 }
 
-const TArray<const FNiagaraAssetTagDefinition*>& INiagaraModule::GetInternalAssetTagDefinitions() const
-{
-PRAGMA_DISABLE_DEPRECATION_WARNINGS
-	return InternalAssetTagDefinitions;
-PRAGMA_ENABLE_DEPRECATION_WARNINGS
-}
-
 void INiagaraModule::OnPostGarbageCollect()
 {
 	FNiagaraTypeHelper::TickTypeRemap();
@@ -1922,14 +1910,12 @@ FSHAHash FNiagaraEmitterParameters::GetStructHash()
 
 void INiagaraModule::RegisterInternalAssetTagDefinitions()
 {
-PRAGMA_DISABLE_DEPRECATION_WARNINGS
 	// All internal tags have to be manually assigned a unique and stable guid so that names can be changed without affecting assigned assets.
 	InternalAssetTagDefinitions.Add(&TemplateTagDefinition);
 	InternalAssetTagDefinitions.Add(&LightweightTagDefinition);
 	InternalAssetTagDefinitions.Add(&LearningContentTagDefinition);
 	InternalAssetTagDefinitions.Add(&HiddenAssetTagDefinition);
 	InternalAssetTagDefinitions.Add(&DeprecatedTagDefinition);
-PRAGMA_ENABLE_DEPRECATION_WARNINGS
 }
 
 #if NIAGARA_PERF_BASELINES
diff --git a/Engine/Plugins/FX/Niagara/Source/Niagara/Private/NiagaraSystem.cpp b/Engine/Plugins/FX/Niagara/Source/Niagara/Private/NiagaraSystem.cpp
--- a/Engine/Plugins/FX/Niagara/Source/Niagara/Private/NiagaraSystem.cpp
+++ b/Engine/Plugins/FX/Niagara/Source/Niagara/Private/NiagaraSystem.cpp
@@ -1416,9 +1416,23 @@ void UNiagaraSystem::PostLoad()
 #endif
 
 #if WITH_EDITORONLY_DATA
-	if(NiagaraVer < FNiagaraCustomVersion::UserAssetTagStorageInMetaData && AssetTags_DEPRECATED.Num() > 0)
+	// see the equivalent in NiagaraEmitter for details
+	ENiagaraScriptTemplateSpecification CurrentTemplateSpecification = TemplateSpecification_DEPRECATED;
+	if(bIsTemplateAsset_DEPRECATED)
 	{
-		INiagaraModule::Get().GetEditorOnlyDataUtilities().MoveUserAssetTagsToMetaData(this);
+		CurrentTemplateSpecification = ENiagaraScriptTemplateSpecification::Template;
+	}
+
+	if(NiagaraVer < FNiagaraCustomVersion::InheritanceUxRefactor)
+	{
+		if(CurrentTemplateSpecification == ENiagaraScriptTemplateSpecification::Template)
+		{
+			AssetTags.AddUnique(INiagaraModule::TemplateTagDefinition);
+		}
+		else if(CurrentTemplateSpecification == ENiagaraScriptTemplateSpecification::Behavior)
+		{
+			AssetTags.AddUnique(INiagaraModule::LearningContentTagDefinition);
+		}
 	}
 #endif // WITH_EDITORONLY_DATA
 
@@ -1908,6 +1922,12 @@ void UNiagaraSystem::GetAssetRegistryTags(FAssetRegistryTagsContext Context) con
 		DI->GetClass()->GetName(ClassName);
 		Context.AddTag(FAssetRegistryTag(*(TEXT("bHas")+ClassName), TEXT("True"), FAssetRegistryTag::TT_Alphabetical));
 	}*/
+
+	// Asset Library Tags
+	for(const FNiagaraAssetTagDefinitionReference& Tag : AssetTags)
+	{
+		Tag.AddTagToAssetRegistryTags(Context);
+	}
 	
 	//Context.AddTag(FAssetRegistryTag("CPUCollision", UsesCPUCollision() ? TEXT("True") : TEXT("False"), FAssetRegistryTag::TT_Alphabetical));
 	//Context.AddTag(FAssetRegistryTag("Looping", bAnyEmitterLoopsForever ? TEXT("True") : TEXT("False"), FAssetRegistryTag::TT_Alphabetical));
diff --git a/Engine/Plugins/FX/Niagara/Source/Niagara/Public/INiagaraEditorOnlyDataUtlities.h b/Engine/Plugins/FX/Niagara/Source/Niagara/Public/INiagaraEditorOnlyDataUtlities.h
--- a/Engine/Plugins/FX/Niagara/Source/Niagara/Public/INiagaraEditorOnlyDataUtlities.h
+++ b/Engine/Plugins/FX/Niagara/Source/Niagara/Public/INiagaraEditorOnlyDataUtlities.h
@@ -35,6 +35,4 @@ public:
 	virtual TOptional<FNiagaraSystemStateData> GetSystemStateData(const UNiagaraSystem& System) const = 0;
 
 	virtual bool GetResetSimulationOnParameterChange() const = 0;
-	
-	virtual void MoveUserAssetTagsToMetaData(UObject* NiagaraAsset) const = 0;
 };
diff --git a/Engine/Plugins/FX/Niagara/Source/Niagara/Public/NiagaraAssetTagDefinitions.h b/Engine/Plugins/FX/Niagara/Source/Niagara/Public/NiagaraAssetTagDefinitions.h
--- a/Engine/Plugins/FX/Niagara/Source/Niagara/Public/NiagaraAssetTagDefinitions.h
+++ b/Engine/Plugins/FX/Niagara/Source/Niagara/Public/NiagaraAssetTagDefinitions.h
@@ -8,8 +8,6 @@
 
 #define UE_API NIAGARA_API
 
-/** DEPRECATED - All classes and data in here are deprecated in favor of the generic User Asset Tags system. */
-
 UENUM(meta=(Bitflags, UseEnumValuesAsMaskValuesInEditor=true))
 enum class ENiagaraAssetLibraryAssetTypes : uint8
 {
@@ -115,15 +113,14 @@ private:
 	FGuid AssetTagDefinitionGuid;
 };
 
-/**
- * An Asset Tag Definition defines a tag that can be added to various Niagara assets for sorting & filtering purposes.
+/** An Asset Tag Definition defines a tag that can be added to various Niagara assets for sorting & filtering purposes.
  * For example, custom tags will show up in the Create Niagara System dialog to filter available emitters.
  * They can also be used to filter assets in the content browser, when used with the custom filter option.
  *
  * You can modify asset tags in the content browser by right-clicking on a Niagara asset, then use the "Manage Tags" submenu to add or remove them.
  */
-UCLASS(MinimalAPI, Deprecated)
-class UDEPRECATED_NiagaraAssetTagDefinitions : public UObject
+UCLASS(MinimalAPI)
+class UNiagaraAssetTagDefinitions : public UObject
 {
 	GENERATED_BODY()
 public:	
diff --git a/Engine/Plugins/FX/Niagara/Source/Niagara/Public/NiagaraModule.h b/Engine/Plugins/FX/Niagara/Source/Niagara/Public/NiagaraModule.h
--- a/Engine/Plugins/FX/Niagara/Source/Niagara/Public/NiagaraModule.h
+++ b/Engine/Plugins/FX/Niagara/Source/Niagara/Public/NiagaraModule.h
@@ -26,8 +26,6 @@ struct FNiagaraParameterStore;
 class FCommonViewportClient;
 class FNiagaraDebuggerClient;
 struct FNiagaraSystemAsyncCompileResults;
-struct FNiagaraAssetTagReference;
-struct FNiagaraAssetTagDefinition;
 
 extern NIAGARA_API int32 GEnableVerboseNiagaraChangeIdLogging;
 
@@ -138,9 +136,7 @@ public:
 	static void RequestRefreshDataChannels() { bDataChannelRefreshRequested = true; }
 	static NIAGARA_API void RefreshDataChannels();
 
-	UE_DEPRECATED(5.7, "Niagara Asset Tags are deprecated in favor of the generic User Asset Tags. Please update your code as this function will be removed. ")
-	NIAGARA_API const TArray<const FNiagaraAssetTagDefinition*>& GetInternalAssetTagDefinitions() const;
-	UE_DEPRECATED(5.7, "Niagara Asset Tags are deprecated in favor of the generic User Asset Tags. Please update your code as this function will be removed. ")
+	const TArray<const FNiagaraAssetTagDefinition*>& GetInternalAssetTagDefinitions() { return InternalAssetTagDefinitions; }
 	void RegisterInternalAssetTagDefinitions();
 	
 #if NIAGARA_PERF_BASELINES
@@ -300,18 +296,13 @@ public:
 	static NIAGARA_API bool bUseGlobalFXBudget;
 	static NIAGARA_API bool bDataChannelsEnabled;
 
-	UE_DEPRECATED(5.7, "Niagara Asset Tags are deprecated in favor of the generic User Asset Tags. Access the new User Asset Tags using FMetaData functions. Please update your code as this property will be removed.")
 	static const FNiagaraAssetTagDefinition LightweightTagDefinition;
-	UE_DEPRECATED(5.7, "Niagara Asset Tags are deprecated in favor of the generic User Asset Tags. Access the new User Asset Tags using FMetaData functions. Please update your code as this property will be removed.")
+
 	static const FNiagaraAssetTagDefinition TemplateTagDefinition;
-	UE_DEPRECATED(5.7, "Niagara Asset Tags are deprecated in favor of the generic User Asset Tags. Access the new User Asset Tags using FMetaData functions. Please update your code as this property will be removed.")
 	static const FNiagaraAssetTagDefinition LearningContentTagDefinition;
-	UE_DEPRECATED(5.7, "Niagara Asset Tags are deprecated in favor of the generic User Asset Tags. Access the new User Asset Tags using FMetaData functions. Please update your code as this property will be removed.")
 	static NIAGARA_API const FNiagaraAssetTagDefinition HiddenAssetTagDefinition;
-	UE_DEPRECATED(5.7, "Niagara Asset Tags are deprecated in favor of the generic User Asset Tags. Access the new User Asset Tags using FMetaData functions. Please update your code as this property will be removed.")
 	static NIAGARA_API const FNiagaraAssetTagDefinition DeprecatedTagDefinition;
 
-	UE_DEPRECATED(5.7, "Niagara Asset Tags are deprecated in favor of the generic User Asset Tags. Access the new User Asset Tags using FMetaData functions. Please update your code as this property will be removed.")
 	TArray<const FNiagaraAssetTagDefinition*> InternalAssetTagDefinitions;
 private:
 	static NIAGARA_API FNiagaraVariable Engine_WorldDeltaTime;
diff --git a/Engine/Plugins/FX/Niagara/Source/NiagaraCore/Public/NiagaraCustomVersion.h b/Engine/Plugins/FX/Niagara/Source/NiagaraCore/Public/NiagaraCustomVersion.h
--- a/Engine/Plugins/FX/Niagara/Source/NiagaraCore/Public/NiagaraCustomVersion.h
+++ b/Engine/Plugins/FX/Niagara/Source/NiagaraCore/Public/NiagaraCustomVersion.h
@@ -214,8 +214,6 @@ struct FNiagaraCustomVersion
 		SerializeUsageBitMaskToGPUFunctionInfo, // Make sure that we serialize the UsageBitMask with the GPUFunctionInfo
 		
 		AddedNewInterpolatedSpawnMode, // Changed the interpolated spawn property into an enum
-
-		UserAssetTagStorageInMetaData, // Adding user asset tag storage in UMetaData for all UPackages
 		
 		// DO NOT ADD A NEW VERSION UNLESS YOU HAVE TALKED TO THE NIAGARA LEAD. Mismanagement of these versions can lead to data loss if it is adjusted in multiple streams simultaneously.
 		// -----<new versions can be added above this line>  -------------------------------------------------
diff --git a/Engine/Plugins/FX/Niagara/Source/NiagaraEditor/Private/AssetDefinitions/AssetDefinition_NiagaraAssetTagDefinitions.h b/Engine/Plugins/FX/Niagara/Source/NiagaraEditor/Private/AssetDefinitions/AssetDefinition_NiagaraAssetTagDefinitions.h
new file mode 100644
--- /dev/null
+++ b/Engine/Plugins/FX/Niagara/Source/NiagaraEditor/Private/AssetDefinitions/AssetDefinition_NiagaraAssetTagDefinitions.h
@@ -0,0 +1,26 @@
+﻿// Copyright Epic Games, Inc. All Rights Reserved.
+
+#pragma once
+
+#include "CoreMinimal.h"
+#include "AssetDefinitionDefault.h"
+#include "NiagaraAssetTagDefinitions.h"
+#include "AssetDefinition_NiagaraAssetTagDefinitions.generated.h"
+
+#define UE_API NIAGARAEDITOR_API
+
+UCLASS(MinimalAPI)
+class UAssetDefinition_NiagaraAssetTagDefinitions : public UAssetDefinitionDefault
+{
+	GENERATED_BODY()
+
+	UE_API UAssetDefinition_NiagaraAssetTagDefinitions();
+	
+	virtual FText GetAssetDisplayName() const override { return NSLOCTEXT("AssetTypeActions", "AssetTypeActions_NiagaraAssetTagDefinitions", "Niagara Asset Tag Definitions"); }
+	virtual FLinearColor GetAssetColor() const override { return FLinearColor::White; }
+	virtual TSoftClassPtr<UObject> GetAssetClass() const override { return UNiagaraAssetTagDefinitions::StaticClass(); }
+	UE_API virtual TConstArrayView<FAssetCategoryPath> GetAssetCategories() const override;
+	UE_API virtual EAssetCommandResult PerformAssetDiff(const FAssetDiffArgs& DiffArgs) const override;
+};
+
+#undef UE_API
diff --git a/Engine/Plugins/FX/Niagara/Source/NiagaraEditor/Private/Customizations/UserParameters/NiagaraUserParametersBuilderBase.cpp b/Engine/Plugins/FX/Niagara/Source/NiagaraEditor/Private/Customizations/UserParameters/NiagaraUserParametersBuilderBase.cpp
--- a/Engine/Plugins/FX/Niagara/Source/NiagaraEditor/Private/Customizations/UserParameters/NiagaraUserParametersBuilderBase.cpp
+++ b/Engine/Plugins/FX/Niagara/Source/NiagaraEditor/Private/Customizations/UserParameters/NiagaraUserParametersBuilderBase.cpp
@@ -171,10 +171,9 @@ void FNiagaraUserParameterNodeBuilder::GenerateUserParameterRows(IDetailChildren
 	{
 		if(const UHierarchyCategory* HierarchyCategory = Cast<UHierarchyCategory>(Child))
 		{
-			const FDataHierarchyElementMetaData_SectionAssociation* SectionAssociation = HierarchyCategory->FindMetaDataOfType<FDataHierarchyElementMetaData_SectionAssociation>();
-			if(SectionAssociation && SectionAssociation->Section.IsValid())
+			if(HierarchyCategory->GetSection() != nullptr)
 			{
-				SectionCategoryMap.FindOrAdd(SectionAssociation->Section.Get()).Add(HierarchyCategory);	
+				SectionCategoryMap.FindOrAdd(HierarchyCategory->GetSection()).Add(HierarchyCategory);	
 			}
 		}
 	}
@@ -185,7 +184,7 @@ void FNiagaraUserParameterNodeBuilder::GenerateUserParameterRows(IDetailChildren
 		TSharedRef<SWrapBox> SectionsBox = SNew(SWrapBox).UseAllottedSize(true);
 
 		// maps don't guarantee order, so we iterate over the original section data instead
-		for(const TObjectPtr<UHierarchySection>& Section : Root->GetSectionData())
+		for(auto& Section : Root->GetSectionData())
 		{
 			// if the map doesn't contain at the entry, it means there are 0 categories for that section, so we skip it
 			if(!SectionCategoryMap.Contains(Section))
@@ -231,8 +230,8 @@ void FNiagaraUserParameterNodeBuilder::GenerateUserParameterRows(IDetailChildren
 							[
 								SNew(STextBlock)
 								.TextStyle(FAppStyle::Get(), "SmallText")
-								.Text_UObject(Section.Get(), &UHierarchySection::GetSectionNameAsText)
-								.ToolTipText_UObject(Section.Get(), &UHierarchySection::GetTooltip)
+								.Text_UObject(Section, &UHierarchySection::GetSectionNameAsText)
+								.ToolTipText_UObject(Section, &UHierarchySection::GetTooltip)
 							]
 						]
 					];
@@ -283,16 +282,7 @@ void FNiagaraUserParameterNodeBuilder::GenerateUserParameterRows(IDetailChildren
 	{		
 		if(const UHierarchyCategory* HierarchyCategory = Cast<UHierarchyCategory>(Child))
 		{
-			const UHierarchySection* Section = nullptr;
-			if(const FDataHierarchyElementMetaData_SectionAssociation* SectionAssociation = HierarchyCategory->FindMetaDataOfType<FDataHierarchyElementMetaData_SectionAssociation>())
-			{
-				if(SectionAssociation->Section.IsValid())
-				{
-					Section = SectionAssociation->Section.Get();
-				}
-			}
-			
-			if((ActiveSection == nullptr || Section == ActiveSection) && HierarchyCategory->DoesOneChildExist<UNiagaraHierarchyUserParameter>(true))
+			if((ActiveSection == nullptr || HierarchyCategory->GetSection() == ActiveSection) && HierarchyCategory->DoesOneChildExist<UNiagaraHierarchyUserParameter>(true))
 			{
 				GenerateCategory(ChildrenBuilder, *HierarchyCategory);
 			}
diff --git a/Engine/Plugins/FX/Niagara/Source/NiagaraEditor/Private/NiagaraEditorModule.cpp b/Engine/Plugins/FX/Niagara/Source/NiagaraEditor/Private/NiagaraEditorModule.cpp
--- a/Engine/Plugins/FX/Niagara/Source/NiagaraEditor/Private/NiagaraEditorModule.cpp
+++ b/Engine/Plugins/FX/Niagara/Source/NiagaraEditor/Private/NiagaraEditorModule.cpp
@@ -170,21 +170,24 @@
 
 #include "LevelEditor.h"
 #include "LevelEditorOutlinerSettings.h"
-#include "NiagaraAssetTagDefinitions.h"
 #include "Filters/CustomClassFilterData.h"
 
 #include "Widgets/SNiagaraDebugger.h"
+#include "Widgets/AssetBrowser/NiagaraAssetBrowserConfig.h"
+#include "Widgets/AssetBrowser/SNiagaraAssetBrowser.h"
 
 #include "NiagaraDebugVis.h"
 #include "NiagaraPerfBaseline.h"
 #include "NiagaraGraphDataCache.h"
 #include "NiagaraDecalRendererProperties.h"
+#include "NiagaraEditorMenuHelpers.h"
 #include "NiagaraLightRendererProperties.h"
 #include "NiagaraRibbonRendererProperties.h"
 #include "NiagaraSpriteRendererProperties.h"
 #include "NiagaraVolumeRendererProperties.h"
 
 #include "Engine/AssetManager.h"
+#include "Widgets/AssetBrowser/SNiagaraSelectedAssetDetails.h"
 #include "NiagaraRecentAndFavoritesManager.h"
 #include "Customizations/Stack/NiagaraStackObjectPropertyCustomization_StatelessModule_DynamicMaterialParameters.h"
 #include "Stateless/Modules/NiagaraStatelessModule_DynamicMaterialParameters.h"
@@ -193,12 +196,6 @@
 #include "TraversalCache/TraversalBuilder.h"
 #include "TraversalCache/TraversalCache.h"
 
-#include "UserAssetTagMenuHelpers.h"
-#include "Config/TaggedAssetBrowserConfig.h"
-#include "UserAssetTagEditorUtilities.h"
-#include "Widgets/STaggedAssetBrowser.h"
-#include "Widgets/SUserAssetTagBrowserSelectedAssetDetails.h"
-
 #include UE_INLINE_GENERATED_CPP_BY_NAME(NiagaraEditorModule)
 
 IMPLEMENT_MODULE( FNiagaraEditorModule, NiagaraEditor );
@@ -377,72 +374,6 @@ class FNiagaraEditorOnlyDataUtilities : public INiagaraEditorOnlyDataUtilities
 		const UNiagaraEditorSettings* NiagaraEditorSettings = GetDefault<UNiagaraEditorSettings>();
 		return NiagaraEditorSettings->GetResetSimulationOnChange();
 	}
-
-	virtual void MoveUserAssetTagsToMetaData(UObject* NiagaraAsset) const override
-	{
-		check(NiagaraAsset->IsA<UNiagaraEmitter>() || NiagaraAsset->IsA<UNiagaraSystem>());
-
-		TArray<FNiagaraAssetTagDefinitionReference>* References = nullptr;
-		if(UNiagaraEmitter* Emitter = Cast<UNiagaraEmitter>(NiagaraAsset))
-		{
-			References = &Emitter->AssetTags_DEPRECATED;
-		}
-		else if(UNiagaraSystem* System = Cast<UNiagaraSystem>(NiagaraAsset))
-		{
-			References = &System->AssetTags_DEPRECATED;
-		}
-
-		if(References && References->Num() > 0)
-		{
-			TMap<FGuid, const FNiagaraAssetTagDefinition*> AllTags;
-			FAssetRegistryModule& AssetRegistryModule = FModuleManager::Get().LoadModuleChecked<FAssetRegistryModule>("AssetRegistry");
-
-			FARFilter Filter;
-			Filter.ClassPaths.Add(UDEPRECATED_NiagaraAssetTagDefinitions::StaticClass()->GetClassPathName());
-
-PRAGMA_DISABLE_DEPRECATION_WARNINGS
-			for(const FNiagaraAssetTagDefinition* InternalTagDefinition : INiagaraModule::Get().GetInternalAssetTagDefinitions())
-			{
-				AllTags.Add(InternalTagDefinition->TagGuid, InternalTagDefinition);
-			}
-PRAGMA_ENABLE_DEPRECATION_WARNINGS
-
-			TArray<FAssetData> AssetTagDefinitionAssets;
-			AssetRegistryModule.Get().GetAssets(Filter, AssetTagDefinitionAssets);
-
-			for(const FAssetData& AssetTagDefinitionAsset : AssetTagDefinitionAssets)
-			{
-				UDEPRECATED_NiagaraAssetTagDefinitions* TagDefinitionsAsset = Cast<UDEPRECATED_NiagaraAssetTagDefinitions>(AssetTagDefinitionAsset.GetAsset());
-
-				for(const FNiagaraAssetTagDefinition& TagDefinition : TagDefinitionsAsset->GetAssetTagDefinitions())
-				{
-					AllTags.Add(TagDefinition.TagGuid, &TagDefinition);
-				}
-			}
-
-			TArray<const FNiagaraAssetTagDefinition*> FoundTagDefinitions;
-			for(const FNiagaraAssetTagDefinitionReference& AssetTagDefinitionReference : *References)
-			{
-				FGuid TagGuid = AssetTagDefinitionReference.GetTagDefinitionReferenceGuid();
-				if(AllTags.Contains(TagGuid))
-				{
-					FoundTagDefinitions.Add(AllTags[TagGuid]);
-				}
-			}
-
-			for(const FNiagaraAssetTagDefinition* Definition : FoundTagDefinitions)
-			{
-				if(Definition == nullptr)
-				{
-					continue;
-				}
-
-				UE::UserAssetTags::AddUserAssetTag(NiagaraAsset, FName(Definition->AssetTag.ToString()));
-			}
-
-			References->Empty();
-		}
-	}
 };
 
 class FNiagaraScriptGraphPanelPinFactory : public FGraphPanelPinFactory
@@ -1177,13 +1108,16 @@ void FNiagaraEditorModule::StartupModule()
 
 	RecentAndFavoritesManager = MakeShared<FNiagaraRecentAndFavoritesManager>();
 	RecentAndFavoritesManager->Initialize();
+
+	UNiagaraAssetBrowserConfig::Initialize();
+
 	// We have to call shutdown before the exit actually happens to be able to save to disk
 	FEditorDelegates::OnEditorPreExit.AddLambda([]()
 	{
 		UNiagaraFavoriteActionsConfig::Shutdown();
 	});
 	
-	InitializeTaggedAssetBrowserDetails();
+	FNiagaraAssetDetailDatabase::Init();
 
 	TempPackage = NewObject<UPackage>(nullptr, TEXT("/Temp/NiagaraEditor"), RF_Transient);
 	TempPackage->AddToRoot();
@@ -1423,6 +1357,9 @@ void FNiagaraEditorModule::StartupModule()
 
 	RegisterDefaultStackObjectCustomizations();
 	
+	NiagaraEditorMenuHelpers::RegisterToolMenus();
+	NiagaraEditorMenuHelpers::RegisterMenuExtensions();
+
 	FNiagaraEditorStyle::Register();
 	ReinitializeStyleCommand = IConsoleManager::Get().RegisterConsoleCommand(
 		TEXT("fx.NiagaraEditor.ReinitializeStyle"),
@@ -2402,151 +2339,6 @@ void FNiagaraEditorModule::UnregisterSettings()
 	}
 }
 
-void FNiagaraEditorModule::InitializeTaggedAssetBrowserDetails()
-{	
-	// Now we add information on how to display emitters & systems in the tagged asset browser
-	// Emitters
-	{
-		FAssetDetailsDisplayInfo EmitterClassInfo;
-	
-		// Inheritance
-		FDisplayedPropertyData InheritanceProperty;
-		InheritanceProperty.ShouldDisplayPropertyDelegate = FDisplayedPropertyData::FShouldDisplayProperty::CreateLambda([](const FAssetData& AssetData)
-		{
-			if(AssetData.FindTag(GET_MEMBER_NAME_CHECKED(UNiagaraEmitter, bIsInheritable)))
-			{
-				return true;
-			}
-			
-			return false; 
-		});
-		InheritanceProperty.NameWidgetDelegate = FDisplayedPropertyData::FGenerateWidget::CreateLambda([](const FAssetData& AssetData) -> TSharedRef<SWidget>
-		{
-			return SNew(STextBlock)
-			.Text(LOCTEXT("Inheritance", "Inheritance"));
-		});
-		InheritanceProperty.ValueWidgetDelegate = FDisplayedPropertyData::FGenerateWidget::CreateLambda([](const FAssetData& AssetData) -> TSharedRef<SWidget>
-		{
-			bool bUseInheritance = false;
-			if(FNiagaraEditorUtilities::GetIsInheritableFromAssetRegistryTags(AssetData, bUseInheritance))
-			{
-				return SNew(STextBlock).Text(bUseInheritance
-					? LOCTEXT("Emitter_UseInheritance_Yes", "Yes")
-					: LOCTEXT("Emitter_UseInheritance_No", "No")
-					);
-			}
-	
-			return SNullWidget::NullWidget;
-		});
-		
-		// Simulation Target
-		FName SimulationTargetTag("HasGPUEmitter");
-	
-		FDisplayedPropertyData SimulationTargetProperty;
-		SimulationTargetProperty.ShouldDisplayPropertyDelegate = FDisplayedPropertyData::FShouldDisplayProperty::CreateLambda([SimulationTargetTag](const FAssetData& Asset)
-		{
-			return Asset.FindTag(SimulationTargetTag);
-		});
-		SimulationTargetProperty.NameWidgetDelegate = FDisplayedPropertyData::FGenerateWidget::CreateLambda([](const FAssetData& AssetData) -> TSharedRef<SWidget>
-		{
-			return SNew(STextBlock)
-			.Text(LOCTEXT("SimulationTargetLabel", "Runs on"));
-		});
-		SimulationTargetProperty.ValueWidgetDelegate = FDisplayedPropertyData::FGenerateWidget::CreateLambda([SimulationTargetTag](const FAssetData& Asset)
-		{
-			return SNew(SImage)
-			.Image_Lambda([Asset, SimulationTargetTag]() -> const FSlateBrush*
-			{
-				if(Asset.FindTag(SimulationTargetTag))
-				{
-					FString Value;
-					if(Asset.GetTagValue(SimulationTargetTag, Value))
-					{
-						if(Value == "True")
-						{
-							return FNiagaraEditorStyle::Get().GetBrush("NiagaraEditor.Stack.GPUIcon");
-						}
-						else
-						{
-							return FNiagaraEditorStyle::Get().GetBrush("NiagaraEditor.Stack.CPUIcon");
-						}
-					}
-				}
-	
-				return FAppStyle::GetNoBrush();
-			});
-		});
-	
-		EmitterClassInfo.DisplayedProperties =
-		{
-			InheritanceProperty,
-			SimulationTargetProperty
-		};
-		
-		EmitterClassInfo.GetDescriptionDelegate = FAssetDetailsDisplayInfo::FGetDescription::CreateLambda([](const FAssetData& AssetData)
-		{
-			FName DescriptionTagName("TemplateAssetDescription");
-			if(AssetData.FindTag(DescriptionTagName))
-			{
-				FString Description;
-				AssetData.GetTagValue(DescriptionTagName, Description);
-				return FText::FromString(Description);
-			}
-	
-			return FText::GetEmpty();
-		});
-
-		FTaggedAssetBrowserDetailsDisplayDatabase::RegisterClass(UNiagaraEmitter::StaticClass(), EmitterClassInfo);
-	}
-	
-	// Systems
-	{
-		FAssetDetailsDisplayInfo SystemClassInfo;
-	
-		// Number of emitters
-		FDisplayedPropertyData NumEmittersProperty;
-		NumEmittersProperty.ShouldDisplayPropertyDelegate = FDisplayedPropertyData::FShouldDisplayProperty::CreateLambda([](const FAssetData& AssetData) -> bool
-		{
-			return AssetData.FindTag("NumEmitters");
-		});
-		NumEmittersProperty.NameWidgetDelegate = FDisplayedPropertyData::FGenerateWidget::CreateLambda([](const FAssetData& AssetData) -> TSharedRef<SWidget>
-		{
-			return SNew(STextBlock)
-			.Text(LOCTEXT("NumberOfEmitters", "Number of Emitters"));
-		});
-		NumEmittersProperty.ValueWidgetDelegate =  FDisplayedPropertyData::FGenerateWidget::CreateLambda([](const FAssetData& AssetData) -> TSharedRef<SWidget>
-		{
-			int32 NumEmitters = INDEX_NONE;
-			if(AssetData.GetTagValue("NumEmitters", NumEmitters))
-			{
-				return SNew(STextBlock).Text(FText::AsNumber(NumEmitters));
-			}
-	
-			return SNullWidget::NullWidget;
-		});
-		
-		SystemClassInfo.DisplayedProperties =
-		{
-			NumEmittersProperty
-		};
-		
-		SystemClassInfo.GetDescriptionDelegate = FAssetDetailsDisplayInfo::FGetDescription::CreateLambda([](const FAssetData& AssetData)
-		{
-			FName DescriptionTagName("TemplateAssetDescription");
-			if(AssetData.FindTag(DescriptionTagName))
-			{
-				FString Description;
-				AssetData.GetTagValue(DescriptionTagName, Description);
-				return FText::FromString(Description);
-			}
-		
-			return FText::GetEmpty();
-		});
-		
-		FTaggedAssetBrowserDetailsDisplayDatabase::RegisterClass(UNiagaraSystem::StaticClass(), SystemClassInfo);
-	}
-}
-
 void FNiagaraEditorModule::AddReferencedObjects(FReferenceCollector& Collector)
 {
 	if (SequencerSettings)
diff --git a/Engine/Plugins/FX/Niagara/Source/NiagaraEditor/Private/NiagaraEmitterEditorData.cpp b/Engine/Plugins/FX/Niagara/Source/NiagaraEditor/Private/NiagaraEmitterEditorData.cpp
--- a/Engine/Plugins/FX/Niagara/Source/NiagaraEditor/Private/NiagaraEmitterEditorData.cpp
+++ b/Engine/Plugins/FX/Niagara/Source/NiagaraEditor/Private/NiagaraEmitterEditorData.cpp
@@ -91,26 +91,19 @@ void UNiagaraEmitterEditorData::PostLoad_TransferSummaryDataToNewFormat()
 		for(FNiagaraStackSection& StackSection : SummarySections_DEPRECATED)
 		{
 			FText SectionDisplayName = StackSection.SectionDisplayName.IsEmptyOrWhitespace() ? FText::FromName(StackSection.SectionIdentifier) : StackSection.SectionDisplayName;
-			UHierarchySection* Section = NewObject<UHierarchySection>(SummaryViewRoot, UHierarchySection::StaticClass());
-			Section->SetSectionNameAsText(SectionDisplayName);
-			
+			UHierarchySection* Section = SummaryViewRoot->AddSection(SectionDisplayName);
 			FHierarchyElementIdentity SectionIdentity;
 			SectionIdentity.Names.Add(StackSection.SectionIdentifier);
 			Section->SetIdentity(SectionIdentity);
 			
-			SummaryViewRoot->GetSectionDataMutable().Insert(Section, 0);
-			
 			Sections.Add(StackSection.SectionIdentifier, Section);
 			
 			for(const FText& CategoryIdentity : StackSection.Categories)
 			{
 				FName CategoryName = FName(CategoryIdentity.ToString());
 				UHierarchyCategory* HierarchyCategory = SummaryViewRoot->AddChild<UHierarchyCategory>();
+				HierarchyCategory->SetSection(Section);
 				HierarchyCategory->SetCategoryName(CategoryName);
-				
-				FDataHierarchyElementMetaData_SectionAssociation* SectionAssociation = HierarchyCategory->FindOrAddMetaDataOfType<FDataHierarchyElementMetaData_SectionAssociation>();
-				SectionAssociation->Section = Section;
-				
 				IdentityCategoryMap.Add(CategoryName, HierarchyCategory);
 				SectionCategoryMapping.Add(CategoryName, Section->GetSectionName());
 			}
@@ -505,7 +498,7 @@ FSimpleMulticastDelegate& UNiagaraEmitterEditorData::OnSummaryViewStateChanged()
 	return OnSummaryViewStateChangedDelegate;
 }
 
-TConstArrayView<const TObjectPtr<UHierarchySection>> UNiagaraEmitterEditorData::GetSummarySections() const
+const TArray<UHierarchySection*>& UNiagaraEmitterEditorData::GetSummarySections() const
 {
 	return SummaryViewRoot->GetSectionData();
 }
diff --git a/Engine/Plugins/FX/Niagara/Source/NiagaraEditor/Private/NiagaraEmitterFactoryNew.cpp b/Engine/Plugins/FX/Niagara/Source/NiagaraEditor/Private/NiagaraEmitterFactoryNew.cpp
--- a/Engine/Plugins/FX/Niagara/Source/NiagaraEditor/Private/NiagaraEmitterFactoryNew.cpp
+++ b/Engine/Plugins/FX/Niagara/Source/NiagaraEditor/Private/NiagaraEmitterFactoryNew.cpp
@@ -15,15 +15,11 @@
 #include "NiagaraEmitterEditorData.h"
 #include "NiagaraNodeAssignment.h"
 #include "NiagaraRecentAndFavoritesManager.h"
-#include "UserAssetTagEditorUtilities.h"
-#include "Assets/TaggedAssetBrowserConfiguration.h"
 #include "Misc/MessageDialog.h"
 #include "Modules/ModuleManager.h"
 #include "Interfaces/IMainFrameModule.h"
 #include "Framework/Application/SlateApplication.h"
-#include "Widgets/STaggedAssetBrowser.h"
-#include "Widgets/STaggedAssetBrowserCreateAssetWindow.h"
-#include "Widgets/AssetBrowser/SNiagaraAddEmitterToSystemWindow.h"
+#include "Widgets/AssetBrowser/SNiagaraAssetBrowser.h"
 
 #include UE_INLINE_GENERATED_CPP_BY_NAME(NiagaraEmitterFactoryNew)
 
@@ -41,66 +37,37 @@ UNiagaraEmitterFactoryNew::UNiagaraEmitterFactoryNew(const FObjectInitializer& O
 }
 
 bool UNiagaraEmitterFactoryNew::ConfigureProperties()
-{	
-	TArray<FAssetData> SelectedAssetData;
+{
+	IMainFrameModule& MainFrame = FModuleManager::LoadModuleChecked<IMainFrameModule>("MainFrame");
+	TSharedPtr<SWindow>	ParentWindow = MainFrame.GetParentWindow();
+
+	SNiagaraAssetBrowser::FArguments AssetBrowserArgs;
+	AssetBrowserArgs.AvailableClasses({UNiagaraEmitter::StaticClass()})
+	.RecentAndFavoritesList(FNiagaraEditorModule::Get().GetRecentsManager()->GetRecentEmitterAndSystemsList())
+	.EmptySelectionMessage(LOCTEXT("EmptyEmitterFactorySelectionMessage", "Select an emitter as a starting point for your new emitter.\n"));
+
+	SNiagaraAssetBrowserWindow::FArguments AssetBrowserWindowArgs;
+	AssetBrowserWindowArgs.AssetBrowserArgs(AssetBrowserArgs)
+	.WindowTitle(LOCTEXT("EmitterAssetBrowserWindowTitle", "Create Niagara Emitter - Select an emitter as a base"));
 	
-	FSoftObjectPath Path = GetDefault<UNiagaraEditorSettings>()->EmitterAssetWizardConfiguration;
-    UObject* ConfigUObject = Path.TryLoad();
-	if(UTaggedAssetBrowserConfiguration* ConfigurationAsset = Cast<UTaggedAssetBrowserConfiguration>(ConfigUObject))
+	TSharedRef<SNiagaraCreateAssetWindow> CreateAssetBrowserWindow = SNew(SNiagaraCreateAssetWindow, *UNiagaraEmitter::StaticClass()).AssetBrowserWindowArgs(AssetBrowserWindowArgs);
+	
+	FSlateApplication::Get().AddModalWindow(CreateAssetBrowserWindow, ParentWindow);
+
+	if(CreateAssetBrowserWindow->ShouldProceedWithAction() == false)
 	{
-		FDefaultDetailsTabConfiguration DefaultDetailsTabConfiguration;
-		DefaultDetailsTabConfiguration.bUseDefaultDetailsTab = true;
-		DefaultDetailsTabConfiguration.EmptySelectionMessage = LOCTEXT("EmptyEmitterFactorySelectionMessage", "Select an emitter as a starting point for your new emitter.\n");
- 
-		IMainFrameModule& MainFrame = FModuleManager::LoadModuleChecked<IMainFrameModule>("MainFrame");
-		TSharedPtr<SWindow>	ParentWindow = MainFrame.GetParentWindow();
-
-		STaggedAssetBrowser::FInterfaceOverrideProfiles OverrideProfiles;
-		OverrideProfiles.AssetViewOptionsProfileName = FName("TaggedAssetBrowser");
-		OverrideProfiles.DefaultFilterMenuExpansion = EAssetTypeCategories::FX;
-		OverrideProfiles.FilterBarSaveName = FName("TaggedAssetBrowser.NiagaraEmitterFactory");
-		
-		STaggedAssetBrowser::FArguments AssetBrowserArgs;
-		AssetBrowserArgs
-			.AvailableClasses({ UNiagaraEmitter::StaticClass() })
-			.DefaultDetailsTabConfiguration(DefaultDetailsTabConfiguration)
-			.InterfaceOverrideProfiles(OverrideProfiles)
-			.RecentAndFavoritesList(FNiagaraEditorModule::Get().GetRecentsManager()->GetRecentEmitterAndSystemsList())
-			.OnAdditionalShouldFilterAsset_UObject(this, &UNiagaraEmitterFactoryNew::OnAdditionalShouldFilterAsset)
-			.OnGetExtraFrontendFilters_UObject(this, &UNiagaraEmitterFactoryNew::OnGetExtraFrontendFilters)
-			.OnExtendAddFilterMenu_UObject(this, &UNiagaraEmitterFactoryNew::OnExtendAddFilterMenu);
-
-		SWindow::FArguments WindowArgs;
-		WindowArgs.Title(LOCTEXT("EmitterAssetBrowserWindowTitle", "Create Niagara Emitter - Select an emitter as a base"));
-		WindowArgs.SupportsMaximize(false);
-		WindowArgs.SupportsMinimize(false);
-		WindowArgs.ClientSize(FVector2D(1400, 750));
-		WindowArgs.SizingRule(ESizingRule::UserSized);
-		
-		STaggedAssetBrowserWindow::FArguments AssetBrowserWindowArgs;
-		AssetBrowserWindowArgs.AssetBrowserArgs(AssetBrowserArgs);
-		AssetBrowserWindowArgs.WindowArgs(WindowArgs);
-		
-		TSharedRef<STaggedAssetBrowserCreateAssetWindow> CreateAssetBrowserWindow = SNew(STaggedAssetBrowserCreateAssetWindow, *ConfigurationAsset, *UNiagaraEmitter::StaticClass())
-			.AssetBrowserWindowArgs(AssetBrowserWindowArgs);
-		
-		FSlateApplication::Get().AddModalWindow(CreateAssetBrowserWindow, ParentWindow);
+		return false;	
+	}
 
-		if(CreateAssetBrowserWindow->ShouldProceedWithAction() == false)
-		{
-			return false;	
-		}
+	TArray<FAssetData> SelectedAssetData = CreateAssetBrowserWindow->GetSelectedAssets();
 
-		SelectedAssetData = CreateAssetBrowserWindow->GetSelectedAssets();
-	}
-	
-	// Give the minimal emitter a chance to be selected if the user chose 'Minimal Emitter'
+	// Give the default empty emitter a chance to be selected if the user chose 'Empty Emitter'
 	if(SelectedAssetData.Num() == 0)
 	{
-		FSoftObjectPath MinimalEmitter = GetDefault<UNiagaraEditorSettings>()->DefaultEmptyEmitter;
-		if(MinimalEmitter.IsValid() && MinimalEmitter.IsAsset())
+		FSoftObjectPath DefaultEmptyEmitter = GetDefault<UNiagaraEditorSettings>()->DefaultEmptyEmitter;
+		if(DefaultEmptyEmitter.IsValid() && DefaultEmptyEmitter.IsAsset())
 		{
-			SelectedAssetData.Add(FAssetData(MinimalEmitter.TryLoad()));
+			SelectedAssetData.Add(FAssetData(DefaultEmptyEmitter.TryLoad()));
 		}
 	}
 
@@ -164,21 +131,22 @@ UObject* UNiagaraEmitterFactoryNew::FactoryCreateNew(UClass* Class, UObject* InP
 
 	UNiagaraEmitter* NewEmitter;
 
-	if (EmitterToCopy.IsValid())
+	if (EmitterToCopy != nullptr)
 	{
 		if (bUseInheritance)
 		{
-			NewEmitter = UNiagaraEmitter::CreateWithParentAndOwner(FVersionedNiagaraEmitter(EmitterToCopy.Get(), EmitterToCopy->GetExposedVersion().VersionGuid), InParent, Name, Flags);
+			NewEmitter = UNiagaraEmitter::CreateWithParentAndOwner(FVersionedNiagaraEmitter(EmitterToCopy, EmitterToCopy->GetExposedVersion().VersionGuid), InParent, Name, Flags);
 		}
 		else
 		{
-			NewEmitter = Cast<UNiagaraEmitter>(StaticDuplicateObject(EmitterToCopy.Get(), InParent, Name, Flags, Class));
+			NewEmitter = Cast<UNiagaraEmitter>(StaticDuplicateObject(EmitterToCopy, InParent, Name, Flags, Class));
 			NewEmitter->SetUniqueEmitterName(Name.GetPlainNameString());
 			NewEmitter->DisableVersioning(EmitterToCopy->GetExposedVersion().VersionGuid);
 		}
 
 		NewEmitter->bIsInheritable = true;
 		NewEmitter->TemplateAssetDescription = FText();
+		NewEmitter->AssetTags.Empty();
 		NewEmitter->Category = FText();
 	}
 	else
@@ -276,53 +244,5 @@ void UNiagaraEmitterFactoryNew::InitializeEmitter(UNiagaraEmitter* NewEmitter, b
 	}
 }
 
-bool UNiagaraEmitterFactoryNew::OnAdditionalShouldFilterAsset(const FAssetData& AssetData) const
-{
-	const UNiagaraEditorSettings* Settings = GetDefault<UNiagaraEditorSettings>();
-	if(Settings->IsAllowedAssetInNiagaraAssetBrowser(AssetData) == false)
-	{
-		return true;
-	}
-	
-	if(Settings->IsAllowedAssetByClassUsage(AssetData) == false)
-	{
-		return true;
-	}
-
-	return false;
-}
-
-void UNiagaraEmitterFactoryNew::OnExtendAddFilterMenu(UToolMenu* ToolMenu) const
-{
-	for (FToolMenuSection& Section : ToolMenu->Sections)
-	{
-		Section.Blocks.RemoveAll([](FToolMenuEntry& ToolMenuEntry)
-		{
-			return ToolMenuEntry.Name == FName("Common");
-		});
-	}
-	
-	TArray<FName> SectionsToKeep { FName("FilterBarAdvanced"), FName("FilterBarResetFilters"), FName("FX Filters"), FName("AssetFilterBarFilterAdvancedAsset"), FName("Niagara Filters"), FName("Niagara Tags") };
-	ToolMenu->Sections.RemoveAll([&SectionsToKeep](FToolMenuSection& ToolMenuSection)
-	{
-		return SectionsToKeep.Contains(ToolMenuSection.Name) == false;
-	});
-}
-
-TArray<TSharedRef<FFrontendFilter>> UNiagaraEmitterFactoryNew::OnGetExtraFrontendFilters() const
-{
-	TArray<TSharedRef<FFrontendFilter>> Result;
-	
-	TSharedRef<FFrontendFilterCategory> NiagaraAdditionalFiltersCategory = MakeShared<FFrontendFilterCategory>(LOCTEXT("NiagaraPropertyFilterCategoryLabel", "Niagara Filters"), LOCTEXT("NiagaraAdditionalFiltersTooltip", "Additional filters for filtering Niagara assets"));
-	
-	Result.Add(MakeShared<FFrontendFilter_NiagaraEmitterInheritance>(true, NiagaraAdditionalFiltersCategory));
-	Result.Add(MakeShared<FFrontendFilter_NiagaraEmitterInheritance>(false, NiagaraAdditionalFiltersCategory));
-	
-	Result.Add(MakeShared<FFrontendFilter_NiagaraSystemEffectType>(true, NiagaraAdditionalFiltersCategory));
-	Result.Add(MakeShared<FFrontendFilter_NiagaraSystemEffectType>(false, NiagaraAdditionalFiltersCategory));
-
-	return Result;
-}
-
 #undef LOCTEXT_NAMESPACE
 
diff --git a/Engine/Plugins/FX/Niagara/Source/NiagaraEditor/Private/NiagaraGraph.cpp b/Engine/Plugins/FX/Niagara/Source/NiagaraEditor/Private/NiagaraGraph.cpp
--- a/Engine/Plugins/FX/Niagara/Source/NiagaraEditor/Private/NiagaraGraph.cpp
+++ b/Engine/Plugins/FX/Niagara/Source/NiagaraEditor/Private/NiagaraGraph.cpp
@@ -4366,16 +4366,13 @@ void UNiagaraGraph::MigrateParameterScriptDataToHierarchyRoot(FVersionedNiagaraS
 	for(const FNiagaraStackSection& StackSection : OwnerData.InputSections_DEPRECATED)
 	{
 		FText SectionDisplayName = StackSection.SectionDisplayName.IsEmptyOrWhitespace() ? FText::FromName(StackSection.SectionIdentifier) : StackSection.SectionDisplayName;
-		UHierarchySection* HierarchySection = NewObject<UHierarchySection>(ParameterRoot, UHierarchySection::StaticClass());
-		HierarchySection->SetSectionNameAsText(SectionDisplayName);
+		UHierarchySection* HierarchySection = ParameterRoot->AddSection(SectionDisplayName);
 		HierarchySection->SetTooltip(StackSection.Tooltip);
 
 		FHierarchyElementIdentity SectionIdentity;
 		SectionIdentity.Names.Add(StackSection.SectionIdentifier);
 		HierarchySection->SetIdentity(SectionIdentity);
 			
-		ParameterRoot->GetSectionDataMutable().Insert(HierarchySection, 0);
-		
 		HierarchySections.Add(StackSection.SectionIdentifier, HierarchySection);
 
 		for(const FText& CategoryIdentity : StackSection.Categories)
@@ -4388,11 +4385,8 @@ void UNiagaraGraph::MigrateParameterScriptDataToHierarchyRoot(FVersionedNiagaraS
 				
 			FName CategoryName = FName(CategoryIdentity.ToString());
 			UNiagaraHierarchyScriptCategory* HierarchyCategory = ParameterRoot->AddChild<UNiagaraHierarchyScriptCategory>();
+			HierarchyCategory->SetSection(HierarchySection);
 			HierarchyCategory->SetCategoryName(CategoryName);
-			
-			FDataHierarchyElementMetaData_SectionAssociation* SectionAssociation = HierarchyCategory->FindOrAddMetaDataOfType<FDataHierarchyElementMetaData_SectionAssociation>();
-			SectionAssociation->Section = HierarchySection;
-			
 			IdentityCategoryMap.Add(CategoryName, HierarchyCategory);
 			SectionCategoryMapping.Add(CategoryName, HierarchySection->GetSectionName());
 		}
diff --git a/Engine/Plugins/FX/Niagara/Source/NiagaraEditor/Private/NiagaraScriptMergeManager.cpp b/Engine/Plugins/FX/Niagara/Source/NiagaraEditor/Private/NiagaraScriptMergeManager.cpp
--- a/Engine/Plugins/FX/Niagara/Source/NiagaraEditor/Private/NiagaraScriptMergeManager.cpp
+++ b/Engine/Plugins/FX/Niagara/Source/NiagaraEditor/Private/NiagaraScriptMergeManager.cpp
@@ -2474,9 +2474,9 @@ void FNiagaraScriptMergeManager::DiffEmitterSummary(const UNiagaraEmitterEditorD
 
 	DiffResults.NewShouldShowSummaryViewValue = OtherEditorData->ShouldShowSummaryView();
 	
-	TArray<const UHierarchyElement*> AddedItemsInOther;
-	TArray<const UHierarchyElement*> RemovedItemsInBase;
-	TArray<const UHierarchySection*> AddedSections;
+	TArray<UHierarchyElement*> AddedItemsInOther;
+	TArray<UHierarchyElement*> RemovedItemsInBase;
+	TArray<UHierarchySection*> AddedSections;
 	
 	UHierarchyRoot* BaseRoot = BaseEditorData->GetSummaryRoot();
 	UHierarchyRoot* OtherRoot = OtherEditorData->GetSummaryRoot();	
@@ -2494,13 +2494,13 @@ void FNiagaraScriptMergeManager::DiffEmitterSummary(const UNiagaraEmitterEditorD
 	
 	DiffResults.AddedSummaryEntriesInOther.Append(SummaryItemDiff.AddedOtherValues);
 	
-	TConstArrayView<const TObjectPtr<UHierarchySection>> BaseSections = BaseRoot->GetSectionData();
-	TConstArrayView<const TObjectPtr<UHierarchySection>> OtherSections = OtherRoot->GetSectionData();
+	TArray<UHierarchySection*> BaseSections = BaseRoot->GetSectionData();
+	TArray<UHierarchySection*> OtherSections = OtherRoot->GetSectionData();
 
-	for (const UHierarchySection* OtherSection : OtherSections)
+	for (UHierarchySection* OtherSection : OtherSections)
 	{
 		FHierarchyElementIdentity OtherSectionIdentity = OtherSection->GetPersistentIdentity();
-		if(!BaseSections.ContainsByPredicate([OtherSectionIdentity](const UHierarchySection* BaseSection)
+		if(!BaseSections.ContainsByPredicate([OtherSectionIdentity](UHierarchySection* BaseSection)
 		{
 			return BaseSection->GetPersistentIdentity() == OtherSectionIdentity;
 		}))
@@ -3736,7 +3736,7 @@ FNiagaraScriptMergeManager::FApplyDiffResults FNiagaraScriptMergeManager::ApplyE
 		BaseRoot->DuplicateSectionFromOtherRoot(*AddedOtherSection);
 	}
 
-	for (const UHierarchyElement* AddedInput : DiffResults.AddedSummaryEntriesInOther)
+	for (UHierarchyElement* AddedInput : DiffResults.AddedSummaryEntriesInOther)
 	{
 		if (UHierarchyElement* SummaryItemParent = AddedInput->GetTypedOuter<UHierarchyElement>())
 		{
@@ -3757,21 +3757,10 @@ FNiagaraScriptMergeManager::FApplyDiffResults FNiagaraScriptMergeManager::ApplyE
 					{
 						UHierarchyElement* NewChild = BaseRoot->CopyAndAddItemAsChild(*AddedInput);
 						NewChild->GetChildrenMutable().Empty();
-
-						if(FDataHierarchyElementMetaData_SectionAssociation* SectionAssociation = NewChild->FindOrAddMetaDataOfType<FDataHierarchyElementMetaData_SectionAssociation>())
+						// for root level category, we have to fixup the section objects they point to as after duplicating them they will still point to the original section
+						if(UHierarchyCategory* AsCategory = Cast<UHierarchyCategory>(NewChild))
 						{
-							if(SectionAssociation->Section.IsValid())
-							{
-								const TObjectPtr<UHierarchySection>* NewSection = BaseRoot->GetSectionData().FindByPredicate([OriginalSection = SectionAssociation->Section](UHierarchySection* SectionCandidate)
-									{
-										return SectionCandidate->GetSectionName().IsEqual(OriginalSection->GetSectionName());
-									});
-
-								if(NewSection)
-								{
-									SectionAssociation->Section = (*NewSection);
-								}
-							}
+							AsCategory->FixupSectionLinkage();
 						}
 					}				
 				}
diff --git a/Engine/Plugins/FX/Niagara/Source/NiagaraEditor/Private/NiagaraScriptMergeManager.h b/Engine/Plugins/FX/Niagara/Source/NiagaraEditor/Private/NiagaraScriptMergeManager.h
--- a/Engine/Plugins/FX/Niagara/Source/NiagaraEditor/Private/NiagaraScriptMergeManager.h
+++ b/Engine/Plugins/FX/Niagara/Source/NiagaraEditor/Private/NiagaraScriptMergeManager.h
@@ -404,8 +404,8 @@ struct FNiagaraEmitterDiffResults
 	TArray<TSharedRef<FNiagaraRendererMergeAdapter>> ModifiedBaseRenderers;
 	TArray<TSharedRef<FNiagaraRendererMergeAdapter>> ModifiedOtherRenderers;
 
-	TArray<const UHierarchyElement*> AddedSummaryEntriesInOther;	
-	TArray<const UHierarchySection*> AddedSummarySectionsInOther;
+	TArray<UHierarchyElement*> AddedSummaryEntriesInOther;	
+	TArray<UHierarchySection*> AddedSummarySectionsInOther;
 
 	TOptional<bool> NewShouldShowSummaryViewValue;
 	
diff --git a/Engine/Plugins/FX/Niagara/Source/NiagaraEditor/Private/NiagaraSystemFactoryNew.cpp b/Engine/Plugins/FX/Niagara/Source/NiagaraEditor/Private/NiagaraSystemFactoryNew.cpp
--- a/Engine/Plugins/FX/Niagara/Source/NiagaraEditor/Private/NiagaraSystemFactoryNew.cpp
+++ b/Engine/Plugins/FX/Niagara/Source/NiagaraEditor/Private/NiagaraSystemFactoryNew.cpp
@@ -19,11 +19,7 @@
 #include "NiagaraRecentAndFavoritesManager.h"
 #include "NiagaraSettings.h"
 #include "ObjectTools.h"
-#include "UserAssetTagEditorUtilities.h"
-#include "Assets/TaggedAssetBrowserConfiguration.h"
-#include "Widgets/STaggedAssetBrowser.h"
-#include "Widgets/STaggedAssetBrowserCreateAssetWindow.h"
-#include "Widgets/AssetBrowser/SNiagaraAddEmitterToSystemWindow.h"
+#include "Widgets/AssetBrowser/SNiagaraAssetBrowser.h"
 
 #include UE_INLINE_GENERATED_CPP_BY_NAME(NiagaraSystemFactoryNew)
 
@@ -39,56 +35,28 @@ UNiagaraSystemFactoryNew::UNiagaraSystemFactoryNew(const FObjectInitializer& Obj
 
 bool UNiagaraSystemFactoryNew::ConfigureProperties()
 {
-	TArray<FAssetData> SelectedAssetData;
-	
-	FSoftObjectPath Path = GetDefault<UNiagaraEditorSettings>()->SystemAssetWizardConfiguration;
-	UObject* ConfigUObject = Path.TryLoad();
-	if(UTaggedAssetBrowserConfiguration* ConfigurationAsset = Cast<UTaggedAssetBrowserConfiguration>(ConfigUObject))
-	{
-		IMainFrameModule& MainFrame = FModuleManager::LoadModuleChecked<IMainFrameModule>("MainFrame");
-		TSharedPtr<SWindow>	ParentWindow = MainFrame.GetParentWindow();
-		
-		STaggedAssetBrowser::FInterfaceOverrideProfiles OverrideProfiles;
-		OverrideProfiles.AssetViewOptionsProfileName = FName("TaggedAssetBrowser");
-		OverrideProfiles.DefaultFilterMenuExpansion = EAssetTypeCategories::FX;
-		OverrideProfiles.FilterBarSaveName = FName("TaggedAssetBrowser.NiagaraSystemFactory");
-		
-		FDefaultDetailsTabConfiguration DefaultDetailsTabConfiguration;
-		DefaultDetailsTabConfiguration.bUseDefaultDetailsTab = true;
-		DefaultDetailsTabConfiguration.EmptySelectionMessage = LOCTEXT("EmptySystemFactorySelectionMessage", "Select an emitter or system as a starting point for your new system.\nA system consists of one or more emitters.");
-		
-		STaggedAssetBrowser::FArguments AssetBrowserArgs;
-		AssetBrowserArgs
-			.AvailableClasses({UNiagaraEmitter::StaticClass(), UNiagaraSystem::StaticClass()})
-			.DefaultDetailsTabConfiguration(DefaultDetailsTabConfiguration)
-			.InterfaceOverrideProfiles(OverrideProfiles)
-			.RecentAndFavoritesList(FNiagaraEditorModule::Get().GetRecentsManager()->GetRecentEmitterAndSystemsList())
-			.OnGetExtraFrontendFilters_UObject(this, &UNiagaraSystemFactoryNew::OnGetExtraFrontendFilters)
-			.OnExtendAddFilterMenu_UObject(this, &UNiagaraSystemFactoryNew::OnExtendAddFilterMenu);
+	IMainFrameModule& MainFrame = FModuleManager::LoadModuleChecked<IMainFrameModule>("MainFrame");
+	TSharedPtr<SWindow>	ParentWindow = MainFrame.GetParentWindow();
 
-		SWindow::FArguments WindowArgs;
-		WindowArgs.Title(LOCTEXT("SystemAssetBrowserWindowTitle", "Create Niagara System - Select an emitter or system as a base"));
-		WindowArgs.SupportsMaximize(false);
-		WindowArgs.SupportsMinimize(false);
-		WindowArgs.ClientSize(FVector2D(1400, 750));
-		WindowArgs.SizingRule(ESizingRule::UserSized);
-		
-		STaggedAssetBrowserWindow::FArguments AssetBrowserWindowArgs;
-		AssetBrowserWindowArgs.AssetBrowserArgs(AssetBrowserArgs);
-		AssetBrowserWindowArgs.WindowArgs(WindowArgs);
-		
-		TSharedRef<STaggedAssetBrowserCreateAssetWindow> CreateAssetBrowserWindow = SNew(STaggedAssetBrowserCreateAssetWindow, *ConfigurationAsset, *UNiagaraSystem::StaticClass())
-			.AssetBrowserWindowArgs(AssetBrowserWindowArgs);
+	SNiagaraAssetBrowser::FArguments AssetBrowserArgs;
+	AssetBrowserArgs.AvailableClasses({UNiagaraEmitter::StaticClass(), UNiagaraSystem::StaticClass()})
+	.RecentAndFavoritesList(FNiagaraEditorModule::Get().GetRecentsManager()->GetRecentEmitterAndSystemsList())
+	.EmptySelectionMessage(LOCTEXT("EmptySystemFactorySelectionMessage", "Select an emitter or system as a starting point for your new system.\nA system consists of one or more emitters."));
 
-		FSlateApplication::Get().AddModalWindow(CreateAssetBrowserWindow, ParentWindow);
+	SNiagaraAssetBrowserWindow::FArguments AssetBrowserWindowArgs;
+	AssetBrowserWindowArgs.AssetBrowserArgs(AssetBrowserArgs)
+	.WindowTitle(LOCTEXT("SystemAssetBrowserWindowTitle", "Create Niagara System - Select an emitter or system as a base"));
+	
+	TSharedRef<SNiagaraCreateAssetWindow> CreateAssetBrowserWindow = SNew(SNiagaraCreateAssetWindow, *UNiagaraSystem::StaticClass()).AssetBrowserWindowArgs(AssetBrowserWindowArgs);
 
-		if(CreateAssetBrowserWindow->ShouldProceedWithAction() == false)
-		{
-			return false;
-		}
+	FSlateApplication::Get().AddModalWindow(CreateAssetBrowserWindow, ParentWindow);
 
-		SelectedAssetData = CreateAssetBrowserWindow->GetSelectedAssets();
+	if(CreateAssetBrowserWindow->ShouldProceedWithAction() == false)
+	{
+		return false;
 	}
+
+	TArray<FAssetData> SelectedAssetData = CreateAssetBrowserWindow->GetSelectedAssets();
 	
 	if(SelectedAssetData.Num() == 1)
 	{
@@ -130,7 +98,7 @@ UObject* UNiagaraSystemFactoryNew::FactoryCreateNew(UClass* Class, UObject* InPa
 		NewSystem = Cast<UNiagaraSystem>(StaticDuplicateObject(SystemToCopy, InParent, Name, Flags, Class));
 		NewSystem->TemplateAssetDescription = FText();
 		NewSystem->Category = FText();
-		//NewSystem->AssetTags.Empty();
+		NewSystem->AssetTags.Empty();
 
 		// if the new system doesn't have a thumbnail image, check the thumbnail map of the original asset's upackage
 		if(NewSystem->ThumbnailImage == nullptr)
@@ -225,53 +193,5 @@ void UNiagaraSystemFactoryNew::TryAssignDefaultEffectType(UNiagaraSystem* System
 	}
 }
 
-bool UNiagaraSystemFactoryNew::OnAdditionalShouldFilterAsset(const FAssetData& AssetData) const
-{
-	const UNiagaraEditorSettings* Settings = GetDefault<UNiagaraEditorSettings>();
-	if(Settings->IsAllowedAssetInNiagaraAssetBrowser(AssetData) == false)
-	{
-		return true;
-	}
-	
-	if(Settings->IsAllowedAssetByClassUsage(AssetData) == false)
-	{
-		return true;
-	}
-
-	return false;
-}
-
-void UNiagaraSystemFactoryNew::OnExtendAddFilterMenu(UToolMenu* ToolMenu) const
-{
-	for (FToolMenuSection& Section : ToolMenu->Sections)
-	{
-		Section.Blocks.RemoveAll([](FToolMenuEntry& ToolMenuEntry)
-		{
-			return ToolMenuEntry.Name == FName("Common");
-		});
-	}
-	
-	TArray<FName> SectionsToKeep { FName("FilterBarAdvanced"), FName("FilterBarResetFilters"), FName("FX Filters"), FName("AssetFilterBarFilterAdvancedAsset"), FName("Niagara Filters"), FName("Niagara Tags") };
-	ToolMenu->Sections.RemoveAll([&SectionsToKeep](FToolMenuSection& ToolMenuSection)
-	{
-		return SectionsToKeep.Contains(ToolMenuSection.Name) == false;
-	});
-}
-
-TArray<TSharedRef<FFrontendFilter>> UNiagaraSystemFactoryNew::OnGetExtraFrontendFilters() const
-{
-	TArray<TSharedRef<FFrontendFilter>> Result;
-	
-	TSharedRef<FFrontendFilterCategory> NiagaraAdditionalFiltersCategory = MakeShared<FFrontendFilterCategory>(LOCTEXT("NiagaraPropertyFilterCategoryLabel", "Niagara Filters"), LOCTEXT("NiagaraAdditionalFiltersTooltip", "Additional filters for filtering Niagara assets"));
-	
-	Result.Add(MakeShared<FFrontendFilter_NiagaraEmitterInheritance>(true, NiagaraAdditionalFiltersCategory));
-	Result.Add(MakeShared<FFrontendFilter_NiagaraEmitterInheritance>(false, NiagaraAdditionalFiltersCategory));
-	
-	Result.Add(MakeShared<FFrontendFilter_NiagaraSystemEffectType>(true, NiagaraAdditionalFiltersCategory));
-	Result.Add(MakeShared<FFrontendFilter_NiagaraSystemEffectType>(false, NiagaraAdditionalFiltersCategory));
-
-	return Result;
-}
-
 #undef LOCTEXT_NAMESPACE
 
diff --git a/Engine/Plugins/FX/Niagara/Source/NiagaraEditor/Private/Toolkits/NiagaraSystemToolkit.cpp b/Engine/Plugins/FX/Niagara/Source/NiagaraEditor/Private/Toolkits/NiagaraSystemToolkit.cpp
--- a/Engine/Plugins/FX/Niagara/Source/NiagaraEditor/Private/Toolkits/NiagaraSystemToolkit.cpp
+++ b/Engine/Plugins/FX/Niagara/Source/NiagaraEditor/Private/Toolkits/NiagaraSystemToolkit.cpp
@@ -53,9 +53,7 @@
 #include "Modules/ModuleManager.h"
 #include "ViewModels/NiagaraMessageLogViewModel.h"
 #include "NiagaraVersionMetaData.h"
-#include "Assets/TaggedAssetBrowserConfiguration.h"
 #include "Interfaces/IMainFrameModule.h"
-#include "Logging/StructuredLog.h"
 #include "Toolkits/SystemToolkitModes/NiagaraSystemToolkitModeBase.h"
 #include "SystemToolkitModes/NiagaraSystemToolkitMode_Default.h"
 #include "SystemToolkitModes/NiagaraSystemToolkitMode_Scalability.h"
@@ -65,7 +63,7 @@
 #include "Widgets/Layout/SBox.h"
 #include "Widgets/SBoxPanel.h"
 #include "Widgets/SNiagaraEmitterVersionWidget.h"
-#include "Widgets/AssetBrowser/SNiagaraAddEmitterToSystemWindow.h"
+#include "Widgets/AssetBrowser/SNiagaraAssetBrowser.h"
 
 #define LOCTEXT_NAMESPACE "NiagaraSystemEditor"
 
@@ -608,17 +606,8 @@ void FNiagaraSystemToolkit::OpenAddEmitterMenu()
 	IMainFrameModule& MainFrame = FModuleManager::LoadModuleChecked<IMainFrameModule>("MainFrame");
 	TSharedPtr<SWindow>	ParentWindow = MainFrame.GetParentWindow();
 
-	FSoftObjectPath Path = GetDefault<UNiagaraEditorSettings>()->EmitterAssetWizardConfiguration;
-	UObject* ConfigUObject = Path.TryLoad();
-	if(UTaggedAssetBrowserConfiguration* ConfigurationAsset = Cast<UTaggedAssetBrowserConfiguration>(ConfigUObject))
-	{
-		TSharedRef<SWindow> AddEmitterWindow = SNew(SNiagaraAddEmitterToSystemWindow, *ConfigurationAsset, GetSystemViewModel().ToSharedRef());
-		FSlateApplication::Get().AddModalWindow(AddEmitterWindow, ParentWindow);
-	}
-	else
-	{
-		UE_LOGFMT(LogNiagaraEditor, Warning, "Could not find emitter wizard configuration asset. Check Niagara Editor settings under Project Settings");
-	}
+	TSharedRef<SWindow> AddEmitterWindow = SNew(SNiagaraAddEmitterToSystemWindow, GetSystemViewModel().ToSharedRef());
+	FSlateApplication::Get().AddModalWindow(AddEmitterWindow, ParentWindow);
 }
 
 bool FNiagaraSystemToolkit::CanAddEmitters() const
diff --git a/Engine/Plugins/FX/Niagara/Source/NiagaraEditor/Private/ViewModels/HierarchyEditor/NiagaraScriptParametersHierarchyViewModel.cpp b/Engine/Plugins/FX/Niagara/Source/NiagaraEditor/Private/ViewModels/HierarchyEditor/NiagaraScriptParametersHierarchyViewModel.cpp
--- a/Engine/Plugins/FX/Niagara/Source/NiagaraEditor/Private/ViewModels/HierarchyEditor/NiagaraScriptParametersHierarchyViewModel.cpp
+++ b/Engine/Plugins/FX/Niagara/Source/NiagaraEditor/Private/ViewModels/HierarchyEditor/NiagaraScriptParametersHierarchyViewModel.cpp
@@ -279,7 +279,7 @@ FHierarchyElementViewModel::FResultWithUserFeedback FNiagaraHierarchyScriptCateg
 bool FNiagaraHierarchyScriptParameterViewModel::DoesExternalDataStillExist(const UHierarchyDataRefreshContext* Context) const
 {
 	// GetScriptVariable checks the owning graph if it still exists
-	const UNiagaraScriptVariable* ScriptVariable = Cast<UNiagaraHierarchyScriptParameter>(GetData())->GetScriptVariable();
+	const UNiagaraScriptVariable* ScriptVariable = Cast<UNiagaraHierarchyScriptParameter>(GetDataMutable())->GetScriptVariable();
 	
 	if(ScriptVariable == nullptr)
 	{
diff --git a/Engine/Plugins/FX/Niagara/Source/NiagaraEditor/Private/ViewModels/HierarchyEditor/NiagaraSummaryViewViewModel.cpp b/Engine/Plugins/FX/Niagara/Source/NiagaraEditor/Private/ViewModels/HierarchyEditor/NiagaraSummaryViewViewModel.cpp
--- a/Engine/Plugins/FX/Niagara/Source/NiagaraEditor/Private/ViewModels/HierarchyEditor/NiagaraSummaryViewViewModel.cpp
+++ b/Engine/Plugins/FX/Niagara/Source/NiagaraEditor/Private/ViewModels/HierarchyEditor/NiagaraSummaryViewViewModel.cpp
@@ -243,7 +243,7 @@ void FNiagaraFunctionViewModel::RefreshChildrenDataInternal()
 	RefreshChildrenInputs(false);
 }
 
-void FNiagaraFunctionViewModel::RefreshChildrenInputs(bool bClearCache)
+void FNiagaraFunctionViewModel::RefreshChildrenInputs(bool bClearCache) const
 {
 	TWeakObjectPtr<UNiagaraNodeFunctionCall> FunctionNodeWeak = GetFunctionCallNode();
 	if(FunctionNodeWeak.IsValid())
@@ -322,7 +322,7 @@ void FNiagaraFunctionViewModel::RefreshChildrenInputs(bool bClearCache)
 
 				for(const FNiagaraVariable& Variable : AssignmentNode->GetAssignmentTargets())
 				{
-					const bool bChildExists = GetData()->GetChildren().ContainsByPredicate([Variable, bClearCache, this](UHierarchyElement* Candidate)
+					const bool bChildExists = GetDataMutable()->GetChildrenMutable().ContainsByPredicate([Variable, bClearCache, this](UHierarchyElement* Candidate)
 					{
 						if(UNiagaraHierarchyAssignmentInput* AssignmentInput = Cast<UNiagaraHierarchyAssignmentInput>(Candidate))
 						{
@@ -361,6 +361,16 @@ FHierarchyElementViewModel::FResultWithUserFeedback FNiagaraFunctionViewModel::C
 	return FHierarchyItemViewModel::CanDropOnInternal(NiagaraHierarchyItemViewModelBase, ItemDropZone);
 }
 
+const UHierarchySection* FNiagaraFunctionViewModel::GetSectionInternal() const
+{
+	if(bIsForHierarchy)
+	{
+		return nullptr;
+	}
+
+	return Section.IsValid() ? Section.Get() : nullptr;
+}
+
 FHierarchyElementViewModel::FResultWithUserFeedback FNiagaraModuleInputViewModel::CanDropOnInternal(TSharedPtr<FHierarchyElementViewModel> DraggedItem, EItemDropZone ItemDropZone)
 {
 	// if the input isn't editable, we don't allow any drops on/above/below the item.
@@ -894,6 +904,11 @@ FHierarchyElementViewModel::FResultWithUserFeedback FNiagaraHierarchyRendererVie
 	return CanEditResults;
 }
 
+const UHierarchySection* FNiagaraHierarchyRendererViewModel::GetSectionInternal() const
+{
+	return Section.IsValid() ? Section.Get() : nullptr;
+}
+
 FString FNiagaraHierarchyEmitterPropertiesViewModel::ToString() const
 {
 	return TEXT("Emitter Properties");
@@ -917,6 +932,11 @@ FHierarchyElementViewModel::FResultWithUserFeedback FNiagaraHierarchyEmitterProp
 	return CanEditResults;
 }
 
+const UHierarchySection* FNiagaraHierarchyEmitterPropertiesViewModel::GetSectionInternal() const
+{
+	return Section.IsValid() ? Section.Get() : nullptr;
+}
+
 FString FNiagaraHierarchyEventHandlerViewModel::ToString() const
 {
 	FNiagaraEventScriptProperties* ScriptProperties = GetEventScriptProperties();
@@ -1019,6 +1039,16 @@ FHierarchyElementViewModel::FResultWithUserFeedback FNiagaraHierarchyEventHandle
 	return CanEditResults;
 }
 
+const UHierarchySection* FNiagaraHierarchyEventHandlerViewModel::GetSectionInternal() const
+{
+	if(Section.IsValid())
+	{
+		return Section.Get();
+	}
+
+	return nullptr;
+}
+
 UNiagaraRendererProperties* FNiagaraHierarchyRendererViewModel::GetRendererProperties() const
 {
 	TWeakObjectPtr<UDataHierarchyViewModelBase> ViewModel = GetHierarchyViewModel();
@@ -1224,6 +1254,11 @@ FHierarchyElementViewModel::FResultWithUserFeedback FNiagaraHierarchySimStageVie
 	return CanEditResults;
 }
 
+const UHierarchySection* FNiagaraHierarchySimStageViewModel::GetSectionInternal() const
+{
+	return Section.IsValid() ? Section.Get() : nullptr;
+}
+
 FString FNiagaraHierarchySimStagePropertiesViewModel::ToString() const
 {
 	return GetSimStage()->SimulationStageName.ToString().Append(" Properties");
@@ -1380,86 +1415,12 @@ void UNiagaraSummaryViewViewModel::PrepareSourceItems(UHierarchyRoot* SourceRoot
 	TSharedPtr<FNiagaraEmitterViewModel> EmitterViewModel = EmitterViewModelWeak.Pin();
 
 	TArray<UHierarchyElement*> NewItems;
-	TArray<UHierarchySection*> Sections;
+	TArray<UHierarchySection*> NewSections;
 	
 	FunctionCallCache.Empty();
-
-	const FName EmitterSpawnSectionName = "Emitter Spawn";
-	const FName EmitterUpdateSectionName = "Emitter Update";
-	const FName ParticleSpawnSectionName = "Particle Spawn";
-	const FName ParticleUpdateSectionName = "Particle Update";
-	const FName EventsSectionName = "Events";
-	const FName SimStagesSectionName = "Sim Stages";
-	const FName RenderersSectionName = "Renderers";
-	
-	TMap<ENiagaraScriptUsage, FName> ScriptUsageSectionNameMap;
-	ScriptUsageSectionNameMap.Add(ENiagaraScriptUsage::EmitterSpawnScript, EmitterSpawnSectionName);
-	ScriptUsageSectionNameMap.Add(ENiagaraScriptUsage::EmitterUpdateScript, EmitterUpdateSectionName);
-	ScriptUsageSectionNameMap.Add(ENiagaraScriptUsage::ParticleSpawnScript, ParticleSpawnSectionName);
-	ScriptUsageSectionNameMap.Add(ENiagaraScriptUsage::ParticleUpdateScript, ParticleUpdateSectionName);
-
-	TMap<FName, FName> ScriptUsageSectionIconMap;
-	ScriptUsageSectionIconMap.Add(ScriptUsageSectionNameMap[ENiagaraScriptUsage::EmitterSpawnScript], "NiagaraEditor.Emitter.SpawnIcon");
-	ScriptUsageSectionIconMap.Add(ScriptUsageSectionNameMap[ENiagaraScriptUsage::EmitterUpdateScript], "NiagaraEditor.Emitter.UpdateIcon");
-	ScriptUsageSectionIconMap.Add(ScriptUsageSectionNameMap[ENiagaraScriptUsage::ParticleSpawnScript], "NiagaraEditor.Particle.SpawnIcon");
-	ScriptUsageSectionIconMap.Add(ScriptUsageSectionNameMap[ENiagaraScriptUsage::ParticleUpdateScript], "NiagaraEditor.Particle.UpdateIcon");
-	ScriptUsageSectionIconMap.Add(EventsSectionName, "NiagaraEditor.EventIcon");
-	ScriptUsageSectionIconMap.Add(SimStagesSectionName, "NiagaraEditor.SimulationStageIcon");
-	ScriptUsageSectionIconMap.Add(RenderersSectionName, "NiagaraEditor.RenderIcon");
-
-	TMap<FName, TSharedPtr<FHierarchySectionViewModel>> SectionViewModelMap;
-
-	auto CreateAndInitializeSectionData = [this, SourceRoot, &Sections](FName SectionName) -> UHierarchySection*
-	{
-		UHierarchySection* Section = nullptr;
-
-		auto* FoundHierarchySection = SourceRoot->GetSectionDataMutable().FindByPredicate([SectionName](UHierarchySection* Candidate)
-		{
-			return Candidate->GetSectionName().IsEqual(SectionName);
-		});
-
-		if(FoundHierarchySection == nullptr)
-		{
-			Section = SourceRoot->AddChild<UHierarchySection>(GetSectionDataClass(), INDEX_NONE);
-			Section->SetSectionName(SectionName);
-		}
-		else
-		{
-			Section = *FoundHierarchySection;
-		}
-
-		ensure(Section);
-		Sections.Add(Section);
-		return Section;
-	};
-
-	TMap<ENiagaraScriptUsage, UHierarchySection*> ScriptUsageSectionMap;
-	for(auto& SectionNamePair : ScriptUsageSectionNameMap)
-	{
-		ScriptUsageSectionMap.Add(SectionNamePair.Key, CreateAndInitializeSectionData(SectionNamePair.Value));
-	}
-	
-	CreateAndInitializeSectionData(EventsSectionName);
-	CreateAndInitializeSectionData(SimStagesSectionName);
-	CreateAndInitializeSectionData(RenderersSectionName);
-
-	// We create the view models now
-	SourceRootViewModel->SyncViewModelsToData();
-	
-	TArray<TSharedPtr<FHierarchySectionViewModel>> SectionViewModels;
-	SourceRootViewModel->GetChildrenViewModelsForType<UHierarchySection, FHierarchySectionViewModel>(SectionViewModels);
-
-	for(TSharedPtr<FHierarchySectionViewModel> SectionViewModel : SectionViewModels)
-	{
-		SectionViewModelMap.Add(SectionViewModel->GetSectionName(), SectionViewModel);
-	}
-
-	// Now we set their images
-	for(const auto& SectionViewModelMapEntry : SectionViewModelMap)
-	{
-		SectionViewModelMapEntry.Value->SetSectionImage(FNiagaraEditorStyle::Get().GetBrush(ScriptUsageSectionIconMap[SectionViewModelMapEntry.Key]));
-	}
-	
+		
+	// we keep track of all category items (modules etc.) belonging to a certain usage so we can set the sections later 
+	TMap<ENiagaraScriptUsage, TArray<UNiagaraHierarchyModule*>> UsageMap;
 	// we keep track of renderers & sim stages here for the same reasons
 	TArray<UNiagaraHierarchyRenderer*> HierarchyRenderers;
 	TArray<UNiagaraHierarchySimStage*> HierarchySimStages;
@@ -1479,7 +1440,8 @@ void UNiagaraSummaryViewViewModel::PrepareSourceItems(UHierarchyRoot* SourceRoot
 	NewItems.Add(EmitterProperties);
 	
 	// We create hierarchy modules here. We attempt to maintain as many previous elements as possible in order to maintain UI state
-	for(UNiagaraNodeFunctionCall* ModuleNode : FNiagaraStackGraphUtilities::GetAllModuleNodes(EmitterViewModel.ToSharedRef()))
+	TArray<UNiagaraNodeFunctionCall*> ModuleNodes = FNiagaraStackGraphUtilities::GetAllModuleNodes(EmitterViewModel.ToSharedRef());
+	for(UNiagaraNodeFunctionCall* ModuleNode : ModuleNodes)
 	{
 		// we skip over sim stage modules here as we want to add them to their respective sim stage group items instead
 		if(SimStageModules.Contains(ModuleNode) || EventHandlerModules.Contains(ModuleNode))
@@ -1487,111 +1449,332 @@ void UNiagaraSummaryViewViewModel::PrepareSourceItems(UHierarchyRoot* SourceRoot
 			continue;
 		}
 		
-		ENiagaraScriptUsage ScriptUsage = FNiagaraStackGraphUtilities::GetOutputNodeUsage(*ModuleNode);
+		UNiagaraHierarchyModule* HierarchyModule = SourceRoot->AddChild<UNiagaraHierarchyModule>();		
+		HierarchyModule->Initialize(*ModuleNode);		
 		
-		if(TSharedPtr<FHierarchyElementViewModel> ParentElement = SectionViewModelMap[ScriptUsageSectionNameMap[ScriptUsage]])
-		{
-			UNiagaraHierarchyModule* HierarchyModule = ScriptUsageSectionMap[ScriptUsage]->AddChild<UNiagaraHierarchyModule>();		
-			HierarchyModule->Initialize(*ModuleNode);		
-		
-			NewItems.Add(HierarchyModule);
-		}
+		NewItems.Add(HierarchyModule);
+		UsageMap.FindOrAdd(FNiagaraStackGraphUtilities::GetOutputNodeUsage(*ModuleNode)).Add(HierarchyModule);
 	}
-	
-	for(const FNiagaraEventScriptProperties& ScriptPropertiesItem : EmitterViewModel->GetEmitter().GetEmitterData()->GetEventHandlers())
+
+	const TArray<FNiagaraEventScriptProperties> ScriptProperties = EmitterViewModel->GetEmitter().GetEmitterData()->GetEventHandlers();
+	for(const FNiagaraEventScriptProperties& ScriptPropertiesItem : ScriptProperties)
 	{
 		UNiagaraHierarchyEventHandler* HierarchyEventHandler = nullptr;
 		FHierarchyElementIdentity EventHandlerIdentity;
 		EventHandlerIdentity.Guids.Add(ScriptPropertiesItem.Script->GetUsageId());
 		EventHandlerIdentity.Guids.Add(ScriptPropertiesItem.SourceEmitterID);
+
+		auto* FoundItem = SourceRoot->GetChildrenMutable().FindByPredicate([EventHandlerIdentity](UHierarchyElement* ItemBase)
+		{
+			return ItemBase->GetPersistentIdentity() == EventHandlerIdentity;
+		});
 		
-		if(TSharedPtr<FHierarchyElementViewModel> ParentElement = SectionViewModelMap[EventsSectionName])
+		if(FoundItem == nullptr)		
 		{
-			auto* FoundItem = ParentElement->GetDataMutable()->GetChildrenMutable().FindByPredicate([EventHandlerIdentity](UHierarchyElement* ItemBase)
-			{
-				return ItemBase->GetPersistentIdentity() == EventHandlerIdentity;
-			});
-			
-			if(FoundItem == nullptr)		
-			{
-				HierarchyEventHandler = ParentElement->GetDataMutable()->AddChild<UNiagaraHierarchyEventHandler>();
-				HierarchyEventHandler->Initialize(ScriptPropertiesItem);
-			}
-			else
-			{
-				HierarchyEventHandler = CastChecked<UNiagaraHierarchyEventHandler>(*FoundItem);
-			}
-			
-			NewItems.Add(HierarchyEventHandler);
-			HierarchyEventHandlers.Add(HierarchyEventHandler);
+			HierarchyEventHandler = SourceRoot->AddChild<UNiagaraHierarchyEventHandler>();
+			HierarchyEventHandler->Initialize(ScriptPropertiesItem);
 		}
+		else
+		{
+			HierarchyEventHandler = CastChecked<UNiagaraHierarchyEventHandler>(*FoundItem);
+		}
+		
+		NewItems.Add(HierarchyEventHandler);
+		HierarchyEventHandlers.Add(HierarchyEventHandler);
 	}
 
 	// We add sim stages here
-	for(UNiagaraSimulationStageBase* SimStage : EmitterViewModel->GetEmitter().GetEmitterData()->GetSimulationStages())
+	const TArray<UNiagaraSimulationStageBase*>& SimStages = EmitterViewModel->GetEmitter().GetEmitterData()->GetSimulationStages();
+	for(UNiagaraSimulationStageBase* SimStage : SimStages)
 	{
 		UNiagaraHierarchySimStage* HierarchySimStage = nullptr;
 		FHierarchyElementIdentity SimStageID;
 		SimStageID.Guids.Add(SimStage->GetMergeId());
 
-		if(TSharedPtr<FHierarchyElementViewModel> ParentElement = SectionViewModelMap[SimStagesSectionName])
+		auto* FoundItem = SourceRoot->GetChildrenMutable().FindByPredicate([SimStageID](UHierarchyElement* ItemBase)
 		{
-			auto* FoundItem = ParentElement->GetDataMutable()->GetChildrenMutable().FindByPredicate([SimStageID](UHierarchyElement* ItemBase)
-			{
-				return ItemBase->GetPersistentIdentity() == SimStageID;
-			});
-			
-			if(FoundItem == nullptr)		
-			{
-				HierarchySimStage = ParentElement->GetDataMutable()->AddChild<UNiagaraHierarchySimStage>();
-				HierarchySimStage->Initialize(*SimStage);
-			}
-			else
-			{
-				HierarchySimStage = CastChecked<UNiagaraHierarchySimStage>(*FoundItem);
-			}
-			
-			NewItems.Add(HierarchySimStage);
-			HierarchySimStages.Add(HierarchySimStage);
+			return ItemBase->GetPersistentIdentity() == SimStageID;
+		});
+		
+		if(FoundItem == nullptr)		
+		{
+			HierarchySimStage = SourceRoot->AddChild<UNiagaraHierarchySimStage>();
+			HierarchySimStage->Initialize(*SimStage);
+		}
+		else
+		{
+			HierarchySimStage = CastChecked<UNiagaraHierarchySimStage>(*FoundItem);
 		}
+		
+		NewItems.Add(HierarchySimStage);
+		HierarchySimStages.Add(HierarchySimStage);
 	}
 	
 	// We create hierarchy renderers here
-	for(UNiagaraRendererProperties* RendererProperties : EmitterViewModel->GetEmitter().GetEmitterData()->GetRenderers())
+	const TArray<UNiagaraRendererProperties*>& Renderers = EmitterViewModel->GetEmitter().GetEmitterData()->GetRenderers();
+	for(UNiagaraRendererProperties* RendererProperties : Renderers)
 	{
 		UNiagaraHierarchyRenderer* HierarchyRenderer = nullptr;
 		FHierarchyElementIdentity RendererIdentity;
 		RendererIdentity.Guids.Add(RendererProperties->GetMergeId());
 
-		if(TSharedPtr<FHierarchyElementViewModel> ParentElement = SectionViewModelMap[RenderersSectionName])
+		auto* FoundItem = SourceRoot->GetChildrenMutable().FindByPredicate([RendererIdentity](UHierarchyElement* ItemBase)
 		{
-			auto* FoundItem = ParentElement->GetDataMutable()->GetChildrenMutable().FindByPredicate([RendererIdentity](UHierarchyElement* ItemBase)
-			{
-				return ItemBase->GetPersistentIdentity() == RendererIdentity;
-			});
-			
-			if(FoundItem == nullptr)		
-			{
-				HierarchyRenderer = ParentElement->GetDataMutable()->AddChild<UNiagaraHierarchyRenderer>();
-				HierarchyRenderer->Initialize(*RendererProperties);
-			}
-			else
-			{
-				HierarchyRenderer = CastChecked<UNiagaraHierarchyRenderer>(*FoundItem);
-			}
-			
-			NewItems.Add(HierarchyRenderer);
-			HierarchyRenderers.Add(HierarchyRenderer);
+			return ItemBase->GetPersistentIdentity() == RendererIdentity;
+		});
+		
+		if(FoundItem == nullptr)		
+		{
+			HierarchyRenderer = SourceRoot->AddChild<UNiagaraHierarchyRenderer>();
+			HierarchyRenderer->Initialize(*RendererProperties);
+		}
+		else
+		{
+			HierarchyRenderer = CastChecked<UNiagaraHierarchyRenderer>(*FoundItem);
 		}
+		
+		NewItems.Add(HierarchyRenderer);
+		HierarchyRenderers.Add(HierarchyRenderer);
 	}
+
+	SourceRoot->GetChildrenMutable().Empty();
+	SourceRoot->GetChildrenMutable().Append(NewItems);
+
+	// we force a sync so we can access the section data for the source items
+	SourceRootViewModel->SyncViewModelsToData();
 	
+	// Now we create a section for each usage case that has at least one element and link up the respective sections
+	UHierarchySection* EmitterSpawnSection = nullptr;
+	if(UsageMap.Contains(ENiagaraScriptUsage::EmitterSpawnScript))
+	{
+		FText SectionName = FText::FromString("Emitter Spawn");
+		auto* FoundHierarchySection = SourceRootViewModel->GetDataMutable<UHierarchyRoot>()->GetSectionDataMutable().FindByPredicate([SectionName](UHierarchySection* Candidate)
+		{
+			return Candidate->GetSectionNameAsText().EqualTo(SectionName);
+		});
+		
+		if(FoundHierarchySection == nullptr)
+		{
+			EmitterSpawnSection = SourceRootViewModel->GetDataMutable<UHierarchyRoot>()->AddSection(SectionName);
+		}
+		else
+		{
+			EmitterSpawnSection = *FoundHierarchySection;
+		}
+
+		NewSections.Add(EmitterSpawnSection);
+
+		for(UNiagaraHierarchyModule* Module : UsageMap[ENiagaraScriptUsage::EmitterSpawnScript])
+		{
+			TSharedPtr<FNiagaraFunctionViewModel> ModuleViewModel = StaticCastSharedPtr<FNiagaraFunctionViewModel>(SourceRootViewModel->FindViewModelForChild(Module));
+			ModuleViewModel->SetSection(*EmitterSpawnSection);
+		}
+	}
+
+	UHierarchySection* EmitterUpdateSection = nullptr;
+	if(UsageMap.Contains(ENiagaraScriptUsage::EmitterUpdateScript))
+	{
+		FText SectionName = FText::FromString("Emitter Update");
+		auto* FoundHierarchySection = SourceRootViewModel->GetDataMutable<UHierarchyRoot>()->GetSectionDataMutable().FindByPredicate([SectionName](UHierarchySection* Candidate)
+		{
+			return Candidate->GetSectionNameAsText().EqualTo(SectionName);
+		});
+		
+		if(FoundHierarchySection == nullptr)
+		{
+			EmitterUpdateSection = SourceRootViewModel->GetDataMutable<UHierarchyRoot>()->AddSection(SectionName);
+		}
+		else
+		{
+			EmitterUpdateSection = *FoundHierarchySection;
+		}
+
+		NewSections.Add(EmitterUpdateSection);
+
+		for(UNiagaraHierarchyModule* Module : UsageMap[ENiagaraScriptUsage::EmitterUpdateScript])
+		{
+			TSharedPtr<FNiagaraFunctionViewModel> ModuleViewModel = StaticCastSharedPtr<FNiagaraFunctionViewModel>(SourceRootViewModel->FindViewModelForChild(Module));
+			ModuleViewModel->SetSection(*EmitterUpdateSection);
+		}
+	}
+
+	UHierarchySection* ParticleSpawnSection = nullptr;
+	if(UsageMap.Contains(ENiagaraScriptUsage::ParticleSpawnScript))
+	{
+		FText SectionName = FText::FromString("Particle Spawn");
+		auto* FoundHierarchySection = SourceRootViewModel->GetDataMutable<UHierarchyRoot>()->GetSectionDataMutable().FindByPredicate([SectionName](UHierarchySection* Candidate)
+		{
+			return Candidate->GetSectionNameAsText().EqualTo(SectionName);
+		});
+		
+		if(FoundHierarchySection == nullptr)
+		{
+			ParticleSpawnSection = SourceRootViewModel->GetDataMutable<UHierarchyRoot>()->AddSection(SectionName);
+		}
+		else
+		{
+			ParticleSpawnSection = *FoundHierarchySection;
+		}
+
+		NewSections.Add(ParticleSpawnSection);
+
+		for(UNiagaraHierarchyModule* Module : UsageMap[ENiagaraScriptUsage::ParticleSpawnScript])
+		{
+			TSharedPtr<FNiagaraFunctionViewModel> ModuleViewModel = StaticCastSharedPtr<FNiagaraFunctionViewModel>(SourceRootViewModel->FindViewModelForChild(Module));
+			ModuleViewModel->SetSection(*ParticleSpawnSection);
+		}
+	}
+
+	UHierarchySection* ParticleUpdateSection = nullptr;
+	if(UsageMap.Contains(ENiagaraScriptUsage::ParticleUpdateScript))
+	{
+		FText SectionName = FText::FromString("Particle Update");
+		auto* FoundHierarchySection = SourceRootViewModel->GetDataMutable<UHierarchyRoot>()->GetSectionDataMutable().FindByPredicate([SectionName](UHierarchySection* Candidate)
+		{
+			return Candidate->GetSectionNameAsText().EqualTo(SectionName);
+		});
+		
+		if(FoundHierarchySection == nullptr)
+		{
+			ParticleUpdateSection = SourceRootViewModel->GetDataMutable<UHierarchyRoot>()->AddSection(SectionName);
+		}
+		else
+		{
+			ParticleUpdateSection = *FoundHierarchySection;
+		}
+
+		NewSections.Add(ParticleUpdateSection);
+
+		for(UNiagaraHierarchyModule* Module : UsageMap[ENiagaraScriptUsage::ParticleUpdateScript])
+		{
+			TSharedPtr<FNiagaraFunctionViewModel> ModuleViewModel = StaticCastSharedPtr<FNiagaraFunctionViewModel>(SourceRootViewModel->FindViewModelForChild(Module));
+			ModuleViewModel->SetSection(*ParticleUpdateSection);
+		}
+	}
+
+	UHierarchySection* EventHandlerSection = nullptr;
+	if(HierarchyEventHandlers.Num() > 0)
+	{
+		FText SectionName = FText::FromString("Events");
+		auto* FoundHierarchySection = SourceRootViewModel->GetDataMutable<UHierarchyRoot>()->GetSectionDataMutable().FindByPredicate([SectionName](UHierarchySection* Candidate)
+		{
+			return Candidate->GetSectionNameAsText().EqualTo(SectionName);
+		});
+		
+		if(FoundHierarchySection == nullptr)
+		{
+			EventHandlerSection = SourceRootViewModel->GetDataMutable<UHierarchyRoot>()->AddSection(SectionName);
+		}
+		else
+		{
+			EventHandlerSection = *FoundHierarchySection;
+		}
+
+		NewSections.Add(EventHandlerSection);
+
+		for(UNiagaraHierarchyEventHandler* EventHandler : HierarchyEventHandlers)
+		{
+			TSharedPtr<FNiagaraHierarchyEventHandlerViewModel> EventHandlerViewModel = StaticCastSharedPtr<FNiagaraHierarchyEventHandlerViewModel>(SourceRootViewModel->FindViewModelForChild(EventHandler));
+			EventHandlerViewModel->SetSection(*EventHandlerSection);
+		}
+	}
+
+	UHierarchySection* SimulationStagesSection = nullptr;
+	if(HierarchySimStages.Num() > 0)
+	{
+		FText SectionName = FText::FromString("Sim Stages");
+		auto* FoundHierarchySection = SourceRootViewModel->GetDataMutable<UHierarchyRoot>()->GetSectionDataMutable().FindByPredicate([SectionName](UHierarchySection* Candidate)
+		{
+			return Candidate->GetSectionNameAsText().EqualTo(SectionName);
+		});
+		
+		if(FoundHierarchySection == nullptr)
+		{
+			SimulationStagesSection = SourceRootViewModel->GetDataMutable<UHierarchyRoot>()->AddSection(SectionName);
+		}
+		else
+		{
+			SimulationStagesSection = *FoundHierarchySection;
+		}
+
+		NewSections.Add(SimulationStagesSection);
+
+		for(UNiagaraHierarchySimStage* SimStage : HierarchySimStages)
+		{
+			TSharedPtr<FNiagaraHierarchySimStageViewModel> SimStageViewModel = StaticCastSharedPtr<FNiagaraHierarchySimStageViewModel>(SourceRootViewModel->FindViewModelForChild(SimStage));
+			SimStageViewModel->SetSection(*SimulationStagesSection);
+		}
+	}
+
+	UHierarchySection* RenderersSection = nullptr; 
+	if(HierarchyRenderers.Num() > 0)
+	{
+		FText SectionName = FText::FromString("Renderers");
+		auto* FoundHierarchySection = SourceRootViewModel->GetDataMutable<UHierarchyRoot>()->GetSectionDataMutable().FindByPredicate([SectionName](UHierarchySection* Candidate)
+		{
+			return Candidate->GetSectionNameAsText().EqualTo(SectionName);
+		});
+		
+		if(FoundHierarchySection == nullptr)
+		{
+			RenderersSection = SourceRootViewModel->GetDataMutable<UHierarchyRoot>()->AddSection(SectionName);
+		}
+		else
+		{
+			RenderersSection = *FoundHierarchySection;
+		}
+
+		NewSections.Add(RenderersSection);
+		
+		for(UNiagaraHierarchyRenderer* Renderer : HierarchyRenderers)
+		{
+			TSharedPtr<FNiagaraHierarchyRendererViewModel> RendererViewModel = StaticCastSharedPtr<FNiagaraHierarchyRendererViewModel>(SourceRootViewModel->FindViewModelForChild(Renderer));
+			RendererViewModel->SetSection(*RenderersSection);
+		}
+	}
+
+	// this will implicitly sort the sections as well as get rid of outdated ones
+	SourceRoot->GetSectionDataMutable().Empty();
+	SourceRoot->GetSectionDataMutable().Append(NewSections);
+
+	// force a sync so we have the view models for the sections available
 	SourceRootViewModel->SyncViewModelsToData();
 
-	for(TSharedPtr<FHierarchySectionViewModel> SectionViewModel : SectionViewModels)
+	for(TSharedPtr<FHierarchySectionViewModel> SectionViewModel : SourceRootViewModel->GetSectionViewModels())
 	{
-		if(SectionViewModel->GetChildren().Num() == 0)
+		if(EmitterSpawnSection != nullptr && SectionViewModel->GetData() == EmitterSpawnSection)
+		{
+			SectionViewModel->SetSectionImage(FNiagaraEditorStyle::Get().GetBrush("NiagaraEditor.Emitter.SpawnIcon"));
+		}
+
+		if(EmitterUpdateSection != nullptr && SectionViewModel->GetData() == EmitterUpdateSection)
+		{
+			SectionViewModel->SetSectionImage(FNiagaraEditorStyle::Get().GetBrush("NiagaraEditor.Emitter.UpdateIcon"));
+		}
+
+		if(ParticleSpawnSection != nullptr && SectionViewModel->GetData() == ParticleSpawnSection)
+		{
+			SectionViewModel->SetSectionImage(FNiagaraEditorStyle::Get().GetBrush("NiagaraEditor.Particle.SpawnIcon"));
+		}
+
+		if(ParticleUpdateSection != nullptr && SectionViewModel->GetData() == ParticleUpdateSection)
+		{
+			SectionViewModel->SetSectionImage(FNiagaraEditorStyle::Get().GetBrush("NiagaraEditor.Particle.UpdateIcon"));
+		}
+
+		if(EventHandlerSection != nullptr && SectionViewModel->GetData() == EventHandlerSection)
+		{
+			SectionViewModel->SetSectionImage(FNiagaraEditorStyle::Get().GetBrush("NiagaraEditor.EventIcon"));
+		}
+
+		if(SimulationStagesSection != nullptr && SectionViewModel->GetData() == SimulationStagesSection)
+		{
+			SectionViewModel->SetSectionImage(FNiagaraEditorStyle::Get().GetBrush("NiagaraEditor.SimulationStageIcon"));
+		}
+
+		if(RenderersSection != nullptr && SectionViewModel->GetData() == RenderersSection)
 		{
-			SectionViewModel->Delete();
+			SectionViewModel->SetSectionImage(FNiagaraEditorStyle::Get().GetBrush("NiagaraEditor.RenderIcon"));
 		}
 	}
 }
diff --git a/Engine/Plugins/FX/Niagara/Source/NiagaraEditor/Private/ViewModels/Stack/NiagaraStackScriptHierarchyRoot.cpp b/Engine/Plugins/FX/Niagara/Source/NiagaraEditor/Private/ViewModels/Stack/NiagaraStackScriptHierarchyRoot.cpp
--- a/Engine/Plugins/FX/Niagara/Source/NiagaraEditor/Private/ViewModels/Stack/NiagaraStackScriptHierarchyRoot.cpp
+++ b/Engine/Plugins/FX/Niagara/Source/NiagaraEditor/Private/ViewModels/Stack/NiagaraStackScriptHierarchyRoot.cpp
@@ -54,23 +54,23 @@ void UNiagaraStackScriptHierarchyRoot::FinalizeInternal()
 	Super::FinalizeInternal();
 }
 
-TConstArrayView<const TObjectPtr<UHierarchySection>> UNiagaraStackScriptHierarchyRoot::GerHierarchySectionData() const
+const TArray<UHierarchySection*>& UNiagaraStackScriptHierarchyRoot::GetSections() const
 {
 	if(OwningFunctionCallNode->GetCalledUsage() == ENiagaraScriptUsage::Module)
 	{
 		return GetScriptParameterHierarchyRoot()->GetSectionData();
 	}
-
-	static TArray<const TObjectPtr<UHierarchySection>> Dummy;
+	
+	static TArray<UHierarchySection*> Dummy;
 	return Dummy;
 }
 
-const UHierarchySection* UNiagaraStackScriptHierarchyRoot::GetActiveHierarchySection() const
+const UHierarchySection* UNiagaraStackScriptHierarchyRoot::GetActiveSection() const
 {
 	return ActiveSection.IsValid() ? ActiveSection.Get() : nullptr;
 }
 
-void UNiagaraStackScriptHierarchyRoot::SetActiveHierarchySection(const UHierarchySection* InActiveSection)
+void UNiagaraStackScriptHierarchyRoot::SetActiveSection(const UHierarchySection* InActiveSection)
 {
 	ActiveSection = InActiveSection;
 	
@@ -210,9 +210,9 @@ void UNiagaraStackScriptHierarchyRoot::RefreshChildrenInternal(const TArray<UNia
 		}
 	}
 	
-	for(const UHierarchyElement* ChildElement : GetScriptParameterHierarchyRoot()->GetChildren())
+	for(const UHierarchyElement* ChildHierarchyItem : GetScriptParameterHierarchyRoot()->GetChildren())
 	{
-		if(const UNiagaraHierarchyScriptParameter* HierarchyParameter = Cast<UNiagaraHierarchyScriptParameter>(ChildElement))
+		if(const UNiagaraHierarchyScriptParameter* HierarchyParameter = Cast<UNiagaraHierarchyScriptParameter>(ChildHierarchyItem))
 		{
 			TOptional<FNiagaraVariable> InputVariableCandidate = HierarchyParameter->GetVariable();
 
@@ -249,11 +249,11 @@ void UNiagaraStackScriptHierarchyRoot::RefreshChildrenInternal(const TArray<UNia
 			NewChildren.AddUnique(InputChild);
 		}
 		
-		if(const UHierarchyCategory* HierarchyCategory = Cast<UHierarchyCategory>(ChildElement))
+		if(const UHierarchyCategory* HierarchyCategory = Cast<UHierarchyCategory>(ChildHierarchyItem))
 		{
 			// Try to find an already existing category to reuse
 			UNiagaraStackScriptHierarchyCategory* StackCategory = FindCurrentChildOfTypeByPredicate<UNiagaraStackScriptHierarchyCategory>(CurrentChildren,
-				[&](UNiagaraStackScriptHierarchyCategory* CurrentCategory) { return CurrentCategory->GetHierarchyCategory() == ChildElement; });
+				[&](UNiagaraStackScriptHierarchyCategory* CurrentCategory) { return CurrentCategory->GetHierarchyCategory() == ChildHierarchyItem; });
 
 			if (StackCategory == nullptr)
 			{
@@ -305,7 +305,7 @@ void UNiagaraStackScriptHierarchyRoot::RefreshChildrenInternal(const TArray<UNia
 
 const UHierarchySection* UNiagaraStackScriptHierarchyRoot::FindSectionByName(FName SectionName)
 {
-	const TObjectPtr<UHierarchySection>* FoundSection = GerHierarchySectionData().FindByPredicate([SectionName](const UHierarchySection* SectionCandidate)
+	UHierarchySection* const* FoundSection = GetSections().FindByPredicate([SectionName](UHierarchySection* SectionCandidate)
 	{
 		return SectionCandidate->GetSectionName() == SectionName;
 	});
@@ -327,9 +327,7 @@ bool UNiagaraStackScriptHierarchyRoot::FilterByActiveSection(const UNiagaraStack
 
 	if(const UNiagaraStackScriptHierarchyCategory* StackHierarchyCategory = Cast<UNiagaraStackScriptHierarchyCategory>(&Child))
 	{
-		const UHierarchyCategory* HierarchyCategory = StackHierarchyCategory->GetHierarchyCategory();
-		const FDataHierarchyElementMetaData_SectionAssociation* SectionAssociation = HierarchyCategory->FindMetaDataOfType<FDataHierarchyElementMetaData_SectionAssociation>();
-		if(SectionAssociation && SectionAssociation->Section == ActiveSection)
+		if(StackHierarchyCategory->GetHierarchyCategory()->GetSection() == ActiveSection)
 		{
 			return true;
 		}
@@ -373,7 +371,7 @@ bool UNiagaraStackScriptHierarchyRoot::GetCanExpand() const
 
 bool UNiagaraStackScriptHierarchyRoot::GetShouldShowInStack() const
 {
-	return GerHierarchySectionData().Num() > 0 || bShouldDisplayLabel;
+	return GetSections().Num() > 0 || bShouldDisplayLabel;
 }
 
 int32 UNiagaraStackScriptHierarchyRoot::GetChildIndentLevel() const
diff --git a/Engine/Plugins/FX/Niagara/Source/NiagaraEditor/Private/ViewModels/Stack/NiagaraStackSummaryViewInputCollection.cpp b/Engine/Plugins/FX/Niagara/Source/NiagaraEditor/Private/ViewModels/Stack/NiagaraStackSummaryViewInputCollection.cpp
--- a/Engine/Plugins/FX/Niagara/Source/NiagaraEditor/Private/ViewModels/Stack/NiagaraStackSummaryViewInputCollection.cpp
+++ b/Engine/Plugins/FX/Niagara/Source/NiagaraEditor/Private/ViewModels/Stack/NiagaraStackSummaryViewInputCollection.cpp
@@ -564,7 +564,7 @@ void UNiagaraStackSummaryViewCollection::RefreshChildrenInternal(const TArray<UN
 
 void UNiagaraStackSummaryViewCollection::GetSectionsInternal(TArray<FNiagaraStackSection>& OutStackSections) const
 {
-	TConstArrayView<const TObjectPtr<UHierarchySection>> HierarchySections = GetHierarchySections();
+	const TArray<UHierarchySection*>& HierarchySections = GetHierarchySections();
 
 	TMap<const UHierarchySection*, TArray<UHierarchyCategory*>> SectionToCategoryMap;
 
@@ -576,13 +576,7 @@ void UNiagaraStackSummaryViewCollection::GetSectionsInternal(TArray<FNiagaraStac
 	
 	for(UHierarchyCategory* TopLevelHierarchyCategories : HierarchyCategories)
 	{
-		if(const FDataHierarchyElementMetaData_SectionAssociation* SectionAssociation = TopLevelHierarchyCategories->FindMetaDataOfType<FDataHierarchyElementMetaData_SectionAssociation>())
-		{
-			if(SectionAssociation->Section.IsValid())
-			{
-				SectionToCategoryMap.FindOrAdd(SectionAssociation->Section.Get()).Add(TopLevelHierarchyCategories);	
-			}
-		}
+		SectionToCategoryMap.FindOrAdd(TopLevelHierarchyCategories->GetSection()).Add(TopLevelHierarchyCategories);
 	}
 	
 	for(const UHierarchySection* HierarchySection : HierarchySections)
@@ -603,14 +597,14 @@ void UNiagaraStackSummaryViewCollection::GetSectionsInternal(TArray<FNiagaraStac
 	}
 }
 
-TConstArrayView<const TObjectPtr<UHierarchySection>> UNiagaraStackSummaryViewCollection::GetHierarchySections() const
+const TArray<UHierarchySection*>& UNiagaraStackSummaryViewCollection::GetHierarchySections() const
 {
 	if (GetEmitterViewModel().IsValid())
 	{
 		return GetEmitterViewModel()->GetEditorData().GetSummarySections();
 	}
 
-	static TArray<const TObjectPtr<UHierarchySection>> Dummy;
+	static TArray<UHierarchySection*> Dummy;
 	return Dummy;
 }
 
diff --git a/Engine/Plugins/FX/Niagara/Source/NiagaraEditor/Private/Widgets/SNiagaraScriptInputPreviewPanel.cpp b/Engine/Plugins/FX/Niagara/Source/NiagaraEditor/Private/Widgets/SNiagaraScriptInputPreviewPanel.cpp
--- a/Engine/Plugins/FX/Niagara/Source/NiagaraEditor/Private/Widgets/SNiagaraScriptInputPreviewPanel.cpp
+++ b/Engine/Plugins/FX/Niagara/Source/NiagaraEditor/Private/Widgets/SNiagaraScriptInputPreviewPanel.cpp
@@ -233,7 +233,7 @@ void SNiagaraScriptInputPreviewPanel::OnSearchButtonClicked(SSearchBox::SearchDi
 
 void SNiagaraScriptInputPreviewPanel::GenerateSearchItems(UHierarchyElement* InRoot, TArray<UHierarchyElement*> ParentChain, TArray<FSearchItem>& OutSearchItems)
 {
-	TConstArrayView<const TObjectPtr<UHierarchyElement>> Children = InRoot->GetChildren();
+	const TArray<UHierarchyElement*>& Children = InRoot->GetChildren();
 	ParentChain.Add(InRoot);
 	OutSearchItems.Add(FSearchItem{ParentChain});
 	for(UHierarchyElement* Child : Children)
diff --git a/Engine/Plugins/FX/Niagara/Source/NiagaraEditor/Public/NiagaraAssetTagDefinitionsFactory.h b/Engine/Plugins/FX/Niagara/Source/NiagaraEditor/Public/NiagaraAssetTagDefinitionsFactory.h
new file mode 100644
--- /dev/null
+++ b/Engine/Plugins/FX/Niagara/Source/NiagaraEditor/Public/NiagaraAssetTagDefinitionsFactory.h
@@ -0,0 +1,21 @@
+﻿// Copyright Epic Games, Inc. All Rights Reserved.
+
+#pragma once
+
+#include "CoreMinimal.h"
+#include "Factories/Factory.h"
+#include "NiagaraAssetTagDefinitionsFactory.generated.h"
+
+#define UE_API NIAGARAEDITOR_API
+
+UCLASS(MinimalAPI)
+class UNiagaraAssetTagDefinitionsFactory : public UFactory
+{
+	GENERATED_BODY()
+
+	UE_API UNiagaraAssetTagDefinitionsFactory();
+	
+	UE_API virtual UObject* FactoryCreateNew(UClass* InClass, UObject* InParent, FName InName, EObjectFlags Flags, UObject* Context, FFeedbackContext* Warn) override;
+};
+
+#undef UE_API
diff --git a/Engine/Plugins/FX/Niagara/Source/NiagaraEditor/Public/NiagaraEditorModule.h b/Engine/Plugins/FX/Niagara/Source/NiagaraEditor/Public/NiagaraEditorModule.h
--- a/Engine/Plugins/FX/Niagara/Source/NiagaraEditor/Public/NiagaraEditorModule.h
+++ b/Engine/Plugins/FX/Niagara/Source/NiagaraEditor/Public/NiagaraEditorModule.h
@@ -223,8 +223,6 @@ public:
 	/** Register/unregister niagara editor settings. */
 	void RegisterSettings();
 	void UnregisterSettings();
-
-	void InitializeTaggedAssetBrowserDetails();
 	
 	/** Gets Niagara editor type utilities for a specific type if there are any registered. */
 	TSharedPtr<INiagaraEditorTypeUtilities, ESPMode::ThreadSafe> NIAGARAEDITOR_API GetTypeUtilities(const FNiagaraTypeDefinition& Type);
diff --git a/Engine/Plugins/FX/Niagara/Source/NiagaraEditor/Public/NiagaraEditorSettings.h b/Engine/Plugins/FX/Niagara/Source/NiagaraEditor/Public/NiagaraEditorSettings.h
--- a/Engine/Plugins/FX/Niagara/Source/NiagaraEditor/Public/NiagaraEditorSettings.h
+++ b/Engine/Plugins/FX/Niagara/Source/NiagaraEditor/Public/NiagaraEditorSettings.h
@@ -361,14 +361,6 @@ public:
 	UPROPERTY(config, EditAnywhere, DisplayName="Minimal Emitter", Category = Niagara, meta=(AllowedClasses="/Script/Niagara.NiagaraEmitter"))
 	FSoftObjectPath DefaultEmptyEmitter;
 
-	/** The configuration used in the Niagara Asset Wizard when creating an new Niagara Emitter or adding an emitter to a system. */
-	UPROPERTY(config, EditAnywhere, DisplayName="Emitter Asset Wizard Configuration", Category = Niagara, meta=(AllowedClasses="/Script/UserAssetTagsEditor.TaggedAssetBrowserConfiguration"))
-	FSoftObjectPath EmitterAssetWizardConfiguration;
-
-	/** The configuration used in the Niagara Asset Wizard when creating a new Niagara System. */
-	UPROPERTY(config, EditAnywhere, DisplayName="System Asset Wizard Configuration", Category = Niagara, meta=(AllowedClasses="/Script/UserAssetTagsEditor.TaggedAssetBrowserConfiguration"))
-	FSoftObjectPath SystemAssetWizardConfiguration;
-
 	/** Default material to use with new decal renderers. */
 	UPROPERTY(config, EditAnywhere, Category = "Niagara Renderer", meta = (AllowedClasses = "/Script/Engine.MaterialInterface"))
 	FSoftObjectPath DefaultDecalRendererMaterial;
diff --git a/Engine/Plugins/FX/Niagara/Source/NiagaraEditor/Public/NiagaraEmitterEditorData.h b/Engine/Plugins/FX/Niagara/Source/NiagaraEditor/Public/NiagaraEmitterEditorData.h
--- a/Engine/Plugins/FX/Niagara/Source/NiagaraEditor/Public/NiagaraEmitterEditorData.h
+++ b/Engine/Plugins/FX/Niagara/Source/NiagaraEditor/Public/NiagaraEmitterEditorData.h
@@ -45,7 +45,7 @@ public:
 	NIAGARAEDITOR_API void ToggleShowSummaryView();
 	
 	UHierarchyRoot* GetSummaryRoot() const { return SummaryViewRoot; }
-	TConstArrayView<const TObjectPtr<UHierarchySection>> GetSummarySections() const;
+	const TArray<UHierarchySection*>& GetSummarySections() const;
 
 	UTexture2D* GetThumbnail() const { return EmitterThumbnail; }
 	void SetThumbnail(UTexture2D* InThumbnail) { EmitterThumbnail = InThumbnail; OnPersistentDataChanged().Broadcast(); }
