Niagara - Fix an issue where niagara components could retain invalid asset references due to uneditable stale override parameters.

diff --git a/Engine/Plugins/FX/Niagara/Source/Niagara/Private/NiagaraComponent.cpp b/Engine/Plugins/FX/Niagara/Source/Niagara/Private/NiagaraComponent.cpp
--- a/Engine/Plugins/FX/Niagara/Source/Niagara/Private/NiagaraComponent.cpp
+++ b/Engine/Plugins/FX/Niagara/Source/Niagara/Private/NiagaraComponent.cpp
@@ -3766,14 +3766,17 @@ void UNiagaraComponent::SynchronizeWithSourceSystem()
 #if WITH_EDITORONLY_DATA
 		OverrideParameters.Empty();
 		EditorOverridesValue_DEPRECATED.Empty();
-		
+		CacheAndClearEditorParameterOverrides();
+
 		OnSynchronizedWithAssetParametersDelegate.Broadcast();
 #endif
 		return;
 	}
 
 #if WITH_EDITOR
+	CacheAndClearEditorParameterOverrides();
 	CopyParametersFromAsset();
+	CopyCachedEditorParameterOverrides();
 	ApplyOverridesToParameterStore();
 #endif
 
@@ -4072,6 +4075,54 @@ FNiagaraVariant UNiagaraComponent::GetCurrentParameterValue(const FNiagaraVariab
 	return FNiagaraVariant(DataValue);
 }
 
+void UNiagaraComponent::CacheAndClearEditorParameterOverrides()
+{
+	auto CacheOverrides = [](const TMap<FNiagaraVariableBase, FNiagaraVariant>& SourceOverrides, TMap<FNiagaraVariableBase, FNiagaraVariant>& CacheOverrides)
+	{
+		for (const TPair<FNiagaraVariableBase, FNiagaraVariant>& SourcePair : SourceOverrides)
+		{
+			CacheOverrides.FindOrAdd(SourcePair.Key) = SourcePair.Value;
+		}
+	};
+	if (IsTemplate())
+	{
+		CacheOverrides(TemplateParameterOverrides, TemplateParameterOverridesCache);
+		TemplateParameterOverrides.Empty();
+	}
+	else
+	{
+		CacheOverrides(InstanceParameterOverrides, InstanceParameterOverridesCache);
+		InstanceParameterOverrides.Empty();
+	}
+}
+
+void UNiagaraComponent::CopyCachedEditorParameterOverrides()
+{
+	auto UpdateOverridesFromCache = [](const TArrayView<const FNiagaraVariableWithOffset>& OverrideParameterVariables, const TMap<FNiagaraVariableBase, FNiagaraVariant>& CacheOverrides, TMap<FNiagaraVariableBase, FNiagaraVariant>& TargetOverrides)
+	{
+		for (const FNiagaraVariableWithOffset& OverrideParameterVariable : OverrideParameterVariables)
+		{
+			const FNiagaraVariant* CachedOverrideValue = CacheOverrides.Find(OverrideParameterVariable);
+			if (CachedOverrideValue != nullptr)
+			{
+				if (TargetOverrides.Contains(OverrideParameterVariable) == false)
+				{
+					TargetOverrides.Add(OverrideParameterVariable, *CachedOverrideValue);
+				}
+			}
+		}
+	};
+	TArrayView<const FNiagaraVariableWithOffset> OverrideParameterVariables = OverrideParameters.ReadParameterVariables();
+	if (IsTemplate())
+	{
+		UpdateOverridesFromCache(OverrideParameterVariables, TemplateParameterOverridesCache, TemplateParameterOverrides);
+	}
+	else
+	{
+		UpdateOverridesFromCache(OverrideParameterVariables, InstanceParameterOverridesCache, InstanceParameterOverrides);
+	}
+}
+
 void UNiagaraComponent::SetOverrideParameterStoreValue(const FNiagaraVariableBase& InKey, const FNiagaraVariant& InValue)
 {
 	if (InKey.IsDataInterface())
@@ -4160,10 +4211,12 @@ void UNiagaraComponent::RemoveParameterOverride(const FNiagaraVariableBase& InKe
 	if (!IsTemplate())
 	{
 		InstanceParameterOverrides.Remove(UserVariable);
+		InstanceParameterOverridesCache.Remove(UserVariable);
 	}
 	else
 	{
 		TemplateParameterOverrides.Remove(UserVariable);
+		TemplateParameterOverridesCache.Remove(UserVariable);
 
 		// we are an archetype, but check if we have an archetype and inherit the value from there
 		const UNiagaraComponent* Archetype = Cast<UNiagaraComponent>(GetArchetype());
diff --git a/Engine/Plugins/FX/Niagara/Source/Niagara/Public/NiagaraComponent.h b/Engine/Plugins/FX/Niagara/Source/Niagara/Public/NiagaraComponent.h
--- a/Engine/Plugins/FX/Niagara/Source/Niagara/Public/NiagaraComponent.h
+++ b/Engine/Plugins/FX/Niagara/Source/Niagara/Public/NiagaraComponent.h
@@ -112,6 +112,12 @@ private:
 	UPROPERTY(EditAnywhere, Category="Niagara")
 	TMap<FNiagaraVariableBase, FNiagaraVariant> InstanceParameterOverrides;
 
+	UPROPERTY(transient)
+	TMap<FNiagaraVariableBase, FNiagaraVariant> TemplateParameterOverridesCache;
+
+	UPROPERTY(transient)
+	TMap<FNiagaraVariableBase, FNiagaraVariant> InstanceParameterOverridesCache;
+
 	FOnSystemInstanceChanged OnSystemInstanceChangedDelegate;
 
 	FOnSynchronizedWithAssetParameters OnSynchronizedWithAssetParametersDelegate;
@@ -726,6 +732,9 @@ private:
 	NIAGARA_API void CopyParametersFromAsset(bool bResetExistingOverrideParameters = true);
 
 #if WITH_EDITOR
+	void CacheAndClearEditorParameterOverrides();
+	void CopyCachedEditorParameterOverrides();
+
 	NIAGARA_API void SetOverrideParameterStoreValue(const FNiagaraVariableBase& InKey, const FNiagaraVariant& InValue);
 	NIAGARA_API void ApplyOverridesToParameterStore();
 
