Summary View Data Hierarchy Editor Fixup
- Fixed broken sections

diff --git a/Engine/Plugins/FX/Niagara/Source/NiagaraEditor/Private/ViewModels/HierarchyEditor/NiagaraSummaryViewViewModel.cpp b/Engine/Plugins/FX/Niagara/Source/NiagaraEditor/Private/ViewModels/HierarchyEditor/NiagaraSummaryViewViewModel.cpp
--- a/Engine/Plugins/FX/Niagara/Source/NiagaraEditor/Private/ViewModels/HierarchyEditor/NiagaraSummaryViewViewModel.cpp
+++ b/Engine/Plugins/FX/Niagara/Source/NiagaraEditor/Private/ViewModels/HierarchyEditor/NiagaraSummaryViewViewModel.cpp
@@ -1406,9 +1406,7 @@ void UNiagaraSummaryViewViewModel::PrepareSourceItems(UHierarchyRoot* SourceRoot
 	ScriptUsageSectionIconMap.Add(EventsSectionName, "NiagaraEditor.EventIcon");
 	ScriptUsageSectionIconMap.Add(SimStagesSectionName, "NiagaraEditor.SimulationStageIcon");
 	ScriptUsageSectionIconMap.Add(RenderersSectionName, "NiagaraEditor.RenderIcon");
-
-	TMap<FName, TSharedPtr<FHierarchySectionViewModel>> SectionViewModelMap;
-
+	
 	auto CreateAndInitializeSectionData = [this, SourceRoot, &Sections](FName SectionName) -> UHierarchySection*
 	{
 		UHierarchySection* Section = nullptr;
@@ -1420,8 +1418,9 @@ void UNiagaraSummaryViewViewModel::PrepareSourceItems(UHierarchyRoot* SourceRoot
 
 		if(FoundHierarchySection == nullptr)
 		{
-			Section = SourceRoot->AddChild<UHierarchySection>(GetSectionDataClass(), INDEX_NONE);
+			Section = NewObject<UHierarchySection>(SourceRoot, GetSectionDataClass());
 			Section->SetSectionName(SectionName);
+			SourceRoot->GetSectionDataMutable().Add(Section);
 		}
 		else
 		{
@@ -1432,11 +1431,10 @@ void UNiagaraSummaryViewViewModel::PrepareSourceItems(UHierarchyRoot* SourceRoot
 		Sections.Add(Section);
 		return Section;
 	};
-
-	TMap<ENiagaraScriptUsage, UHierarchySection*> ScriptUsageSectionMap;
+	
 	for(auto& SectionNamePair : ScriptUsageSectionNameMap)
 	{
-		ScriptUsageSectionMap.Add(SectionNamePair.Key, CreateAndInitializeSectionData(SectionNamePair.Value));
+		CreateAndInitializeSectionData(SectionNamePair.Value);
 	}
 	
 	CreateAndInitializeSectionData(EventsSectionName);
@@ -1446,10 +1444,9 @@ void UNiagaraSummaryViewViewModel::PrepareSourceItems(UHierarchyRoot* SourceRoot
 	// We create the view models now
 	SourceRootViewModel->SyncViewModelsToData();
 	
-	TArray<TSharedPtr<FHierarchySectionViewModel>> SectionViewModels;
-	SourceRootViewModel->GetChildrenViewModelsForType<UHierarchySection, FHierarchySectionViewModel>(SectionViewModels);
-
-	for(TSharedPtr<FHierarchySectionViewModel> SectionViewModel : SectionViewModels)
+	TMap<FName, TSharedPtr<FHierarchySectionViewModel>> SectionViewModelMap;
+	TMap<FName, int32> SectionElementCountMap;
+	for(TSharedPtr<FHierarchySectionViewModel> SectionViewModel : SourceRootViewModel->GetSectionViewModels())
 	{
 		SectionViewModelMap.Add(SectionViewModel->GetSectionName(), SectionViewModel);
 	}
@@ -1488,13 +1485,16 @@ void UNiagaraSummaryViewViewModel::PrepareSourceItems(UHierarchyRoot* SourceRoot
 		}
 		
 		ENiagaraScriptUsage ScriptUsage = FNiagaraStackGraphUtilities::GetOutputNodeUsage(*ModuleNode);
+
+		UNiagaraHierarchyModule* HierarchyModule = SourceRoot->AddChild<UNiagaraHierarchyModule>();
+		HierarchyModule->Initialize(*ModuleNode);		
+		NewItems.Add(HierarchyModule);
 		
-		if(TSharedPtr<FHierarchyElementViewModel> ParentElement = SectionViewModelMap[ScriptUsageSectionNameMap[ScriptUsage]])
+		if(TSharedPtr<FHierarchySectionViewModel> SectionViewModel = SectionViewModelMap[ScriptUsageSectionNameMap[ScriptUsage]])
 		{
-			UNiagaraHierarchyModule* HierarchyModule = ScriptUsageSectionMap[ScriptUsage]->AddChild<UNiagaraHierarchyModule>();		
-			HierarchyModule->Initialize(*ModuleNode);		
-		
-			NewItems.Add(HierarchyModule);
+			FDataHierarchyElementMetaData_SectionAssociation* SectionAssociation = HierarchyModule->FindOrAddMetaDataOfType<FDataHierarchyElementMetaData_SectionAssociation>();
+			SectionAssociation->Section = Cast<UHierarchySection>(SectionViewModel->GetData());
+			SectionElementCountMap.FindOrAdd(SectionViewModel->GetSectionName())++;
 		}
 	}
 	
@@ -1505,25 +1505,29 @@ void UNiagaraSummaryViewViewModel::PrepareSourceItems(UHierarchyRoot* SourceRoot
 		EventHandlerIdentity.Guids.Add(ScriptPropertiesItem.Script->GetUsageId());
 		EventHandlerIdentity.Guids.Add(ScriptPropertiesItem.SourceEmitterID);
 		
-		if(TSharedPtr<FHierarchyElementViewModel> ParentElement = SectionViewModelMap[EventsSectionName])
+		auto* FoundItem = SourceRoot->GetChildrenMutable().FindByPredicate([EventHandlerIdentity](UHierarchyElement* ItemBase)
 		{
-			auto* FoundItem = ParentElement->GetDataMutable()->GetChildrenMutable().FindByPredicate([EventHandlerIdentity](UHierarchyElement* ItemBase)
-			{
-				return ItemBase->GetPersistentIdentity() == EventHandlerIdentity;
-			});
-			
-			if(FoundItem == nullptr)		
-			{
-				HierarchyEventHandler = ParentElement->GetDataMutable()->AddChild<UNiagaraHierarchyEventHandler>();
-				HierarchyEventHandler->Initialize(ScriptPropertiesItem);
-			}
-			else
-			{
-				HierarchyEventHandler = CastChecked<UNiagaraHierarchyEventHandler>(*FoundItem);
-			}
-			
-			NewItems.Add(HierarchyEventHandler);
-			HierarchyEventHandlers.Add(HierarchyEventHandler);
+			return ItemBase->GetPersistentIdentity() == EventHandlerIdentity;
+		});
+		
+		if(FoundItem == nullptr)		
+		{
+			HierarchyEventHandler = SourceRoot->AddChild<UNiagaraHierarchyEventHandler>();
+			HierarchyEventHandler->Initialize(ScriptPropertiesItem);
+		}
+		else
+		{
+			HierarchyEventHandler = CastChecked<UNiagaraHierarchyEventHandler>(*FoundItem);
+		}
+		
+		NewItems.Add(HierarchyEventHandler);
+		HierarchyEventHandlers.Add(HierarchyEventHandler);
+		
+		if(TSharedPtr<FHierarchySectionViewModel> SectionViewModel = SectionViewModelMap[EventsSectionName])
+		{
+			FDataHierarchyElementMetaData_SectionAssociation* SectionAssociation = HierarchyEventHandler->FindOrAddMetaDataOfType<FDataHierarchyElementMetaData_SectionAssociation>();
+			SectionAssociation->Section = Cast<UHierarchySection>(SectionViewModel->GetData());
+			SectionElementCountMap.FindOrAdd(SectionViewModel->GetSectionName())++;
 		}
 	}
 
@@ -1534,25 +1538,29 @@ void UNiagaraSummaryViewViewModel::PrepareSourceItems(UHierarchyRoot* SourceRoot
 		FHierarchyElementIdentity SimStageID;
 		SimStageID.Guids.Add(SimStage->GetMergeId());
 
-		if(TSharedPtr<FHierarchyElementViewModel> ParentElement = SectionViewModelMap[SimStagesSectionName])
+		auto* FoundItem = SourceRoot->GetChildrenMutable().FindByPredicate([SimStageID](UHierarchyElement* ItemBase)
 		{
-			auto* FoundItem = ParentElement->GetDataMutable()->GetChildrenMutable().FindByPredicate([SimStageID](UHierarchyElement* ItemBase)
-			{
-				return ItemBase->GetPersistentIdentity() == SimStageID;
-			});
-			
-			if(FoundItem == nullptr)		
-			{
-				HierarchySimStage = ParentElement->GetDataMutable()->AddChild<UNiagaraHierarchySimStage>();
-				HierarchySimStage->Initialize(*SimStage);
-			}
-			else
-			{
-				HierarchySimStage = CastChecked<UNiagaraHierarchySimStage>(*FoundItem);
-			}
-			
-			NewItems.Add(HierarchySimStage);
-			HierarchySimStages.Add(HierarchySimStage);
+			return ItemBase->GetPersistentIdentity() == SimStageID;
+		});
+		
+		if(FoundItem == nullptr)		
+		{
+			HierarchySimStage = SourceRoot->AddChild<UNiagaraHierarchySimStage>();
+			HierarchySimStage->Initialize(*SimStage);
+		}
+		else
+		{
+			HierarchySimStage = CastChecked<UNiagaraHierarchySimStage>(*FoundItem);
+		}
+		
+		NewItems.Add(HierarchySimStage);
+		HierarchySimStages.Add(HierarchySimStage);
+		
+		if(TSharedPtr<FHierarchySectionViewModel> SectionViewModel = SectionViewModelMap[SimStagesSectionName])
+		{
+			FDataHierarchyElementMetaData_SectionAssociation* SectionAssociation = HierarchySimStage->FindOrAddMetaDataOfType<FDataHierarchyElementMetaData_SectionAssociation>();
+			SectionAssociation->Section = Cast<UHierarchySection>(SectionViewModel->GetData());
+			SectionElementCountMap.FindOrAdd(SectionViewModel->GetSectionName())++;
 		}
 	}
 	
@@ -1563,37 +1571,44 @@ void UNiagaraSummaryViewViewModel::PrepareSourceItems(UHierarchyRoot* SourceRoot
 		FHierarchyElementIdentity RendererIdentity;
 		RendererIdentity.Guids.Add(RendererProperties->GetMergeId());
 
-		if(TSharedPtr<FHierarchyElementViewModel> ParentElement = SectionViewModelMap[RenderersSectionName])
+		auto* FoundItem = SourceRoot->GetChildrenMutable().FindByPredicate([RendererIdentity](UHierarchyElement* ItemBase)
 		{
-			auto* FoundItem = ParentElement->GetDataMutable()->GetChildrenMutable().FindByPredicate([RendererIdentity](UHierarchyElement* ItemBase)
-			{
-				return ItemBase->GetPersistentIdentity() == RendererIdentity;
-			});
-			
-			if(FoundItem == nullptr)		
-			{
-				HierarchyRenderer = ParentElement->GetDataMutable()->AddChild<UNiagaraHierarchyRenderer>();
-				HierarchyRenderer->Initialize(*RendererProperties);
-			}
-			else
-			{
-				HierarchyRenderer = CastChecked<UNiagaraHierarchyRenderer>(*FoundItem);
-			}
-			
-			NewItems.Add(HierarchyRenderer);
-			HierarchyRenderers.Add(HierarchyRenderer);
+			return ItemBase->GetPersistentIdentity() == RendererIdentity;
+		});
+		
+		if(FoundItem == nullptr)		
+		{
+			HierarchyRenderer = SourceRoot->AddChild<UNiagaraHierarchyRenderer>();
+			HierarchyRenderer->Initialize(*RendererProperties);
+		}
+		else
+		{
+			HierarchyRenderer = CastChecked<UNiagaraHierarchyRenderer>(*FoundItem);
+		}
+		
+		NewItems.Add(HierarchyRenderer);
+		HierarchyRenderers.Add(HierarchyRenderer);
+		
+		if(TSharedPtr<FHierarchySectionViewModel> SectionViewModel = SectionViewModelMap[RenderersSectionName])
+		{
+			FDataHierarchyElementMetaData_SectionAssociation* SectionAssociation = HierarchyRenderer->FindOrAddMetaDataOfType<FDataHierarchyElementMetaData_SectionAssociation>();
+			SectionAssociation->Section = Cast<UHierarchySection>(SectionViewModel->GetData());
+			SectionElementCountMap.FindOrAdd(SectionViewModel->GetSectionName())++;
 		}
 	}
-	
-	SourceRootViewModel->SyncViewModelsToData();
 
-	for(TSharedPtr<FHierarchySectionViewModel> SectionViewModel : SectionViewModels)
+	// If any section has no actual elements, we delete it again
+	TArray<TSharedPtr<FHierarchySectionViewModel>> TmpSectionViewModels = SourceRootViewModel->GetSectionViewModels();
+	for(auto It = TmpSectionViewModels.CreateIterator(); It; ++It)
 	{
-		if(SectionViewModel->GetChildren().Num() == 0)
+		FName SectionName = (*It)->GetSectionName();
+		if(SectionElementCountMap.Contains(SectionName) == false || SectionElementCountMap[SectionName] == 0)
 		{
-			SectionViewModel->Delete();
+			(*It)->Delete();
 		}
 	}
+	
+	SourceRootViewModel->SyncViewModelsToData();
 }
 
 void UNiagaraSummaryViewViewModel::SetupCommands()
