Minor hierarchy identity fix for Niagara Summary View

diff --git a/Engine/Plugins/FX/Niagara/Source/NiagaraEditor/Private/ViewModels/HierarchyEditor/NiagaraSummaryViewViewModel.cpp b/Engine/Plugins/FX/Niagara/Source/NiagaraEditor/Private/ViewModels/HierarchyEditor/NiagaraSummaryViewViewModel.cpp
--- a/Engine/Plugins/FX/Niagara/Source/NiagaraEditor/Private/ViewModels/HierarchyEditor/NiagaraSummaryViewViewModel.cpp
+++ b/Engine/Plugins/FX/Niagara/Source/NiagaraEditor/Private/ViewModels/HierarchyEditor/NiagaraSummaryViewViewModel.cpp
@@ -73,6 +73,7 @@ void UNiagaraHierarchyEventHandler::Initialize(const FNiagaraEventScriptProperti
 {
 	FHierarchyElementIdentity EventHandlerIdentity;
 	EventHandlerIdentity.Guids.Add(EventHandler.Script->GetUsageId());
+	EventHandlerIdentity.Guids.Add(EventHandler.SourceEmitterID);
 	SetIdentity(EventHandlerIdentity);
 }
 
