Niagara Tagged Asset Browser Setup
- Resave of templates & learning content to transfer Niagara Asset Tags to User Asset Tags
- Added Emitter Wizard and System Wizard configuration assets, set them as default for niagara
- Added plugin icon as texture for section use

#virtualized

diff --git a/Engine/Plugins/FX/Niagara/Config/BaseNiagara.ini b/Engine/Plugins/FX/Niagara/Config/BaseNiagara.ini
--- a/Engine/Plugins/FX/Niagara/Config/BaseNiagara.ini
+++ b/Engine/Plugins/FX/Niagara/Config/BaseNiagara.ini
@@ -85,6 +85,8 @@ DefaultMeshRendererMesh=/Niagara/DefaultAssets/S_Gnomon.S_Gnomon
 DefaultDecalRendererMaterial=/Niagara/DefaultAssets/DefaultDecalMaterial.DefaultDecalMaterial
 DefaultRibbonRendererMaterial=/Niagara/DefaultAssets/DefaultRibbonMaterial.DefaultRibbonMaterial
 DefaultSpriteRendererMaterial=/Niagara/DefaultAssets/DefaultSpriteMaterial.DefaultSpriteMaterial
+SystemAssetWizardConfiguration=/Niagara/DefaultAssets/TABC_SystemWizard.TABC_SystemWizard
+EmitterAssetWizardConfiguration=/Niagara/DefaultAssets/TABC_EmitterWizard.TABC_EmitterWizard
 RequiredSystemUpdateScript=/Niagara/Modules/System/SystemState.SystemState
 +GraphCreationShortcuts=(Name="Numeric::Add",Input=(Key=A,bShift=False,bCtrl=False,bAlt=False,bCmd=False))
 +GraphCreationShortcuts=(Name="Numeric::Divide",Input=(Key=D,bShift=False,bCtrl=False,bAlt=False,bCmd=False))
