FNiagaraEditorModule::TAssetPreloadCache::RefreshCache and INiagaraParameterDefinitionsSubscriber::SynchronizeWithParameterDefinitions now do disk-only asset registry queries to ensure deterministic results during cook

diff --git a/Engine/Plugins/FX/Niagara/Source/Niagara/Private/NiagaraParameterDefinitionsSubscriber.cpp b/Engine/Plugins/FX/Niagara/Source/Niagara/Private/NiagaraParameterDefinitionsSubscriber.cpp
--- a/Engine/Plugins/FX/Niagara/Source/Niagara/Private/NiagaraParameterDefinitionsSubscriber.cpp
+++ b/Engine/Plugins/FX/Niagara/Source/Niagara/Private/NiagaraParameterDefinitionsSubscriber.cpp
@@ -180,7 +180,15 @@ void INiagaraParameterDefinitionsSubscriber::SynchronizeWithParameterDefinitions
 	if (AssetRegistryModule.Get().IsLoadingAssets() == false)
 	{
 		TArray<FAssetData> ParameterDefinitionsAssetData;
-		AssetRegistryModule.GetRegistry().GetAssetsByClass(FTopLevelAssetPath(TEXT("/Script/NiagaraEditor"), TEXT("NiagaraParameterDefinitions")), ParameterDefinitionsAssetData);
+
+		// Since this runs during the cook, we must query for disk-only assets to ensure deterministic results
+		// (any asset queries running during cook of a package will affect dependencies)
+		FARFilter Filter;
+		Filter.ClassPaths.Add(FTopLevelAssetPath(TEXT("/Script/NiagaraEditor"), TEXT("NiagaraParameterDefinitions")));
+		Filter.bRecursiveClasses = false;
+		Filter.bIncludeOnlyOnDiskAssets = IsRunningCookCommandlet();
+		AssetRegistryModule.GetRegistry().GetAssets(Filter, ParameterDefinitionsAssetData);
+		
 		for (const FAssetData& ParameterDefinitionsAssetDatum : ParameterDefinitionsAssetData)
 		{
 			if (!ParameterDefinitionsAssetDatum.IsUAsset() || !ParameterDefinitionsAssetDatum.IsAssetLoaded())
diff --git a/Engine/Plugins/FX/Niagara/Source/NiagaraEditor/Public/NiagaraEditorModule.h b/Engine/Plugins/FX/Niagara/Source/NiagaraEditor/Public/NiagaraEditorModule.h
--- a/Engine/Plugins/FX/Niagara/Source/NiagaraEditor/Public/NiagaraEditorModule.h
+++ b/Engine/Plugins/FX/Niagara/Source/NiagaraEditor/Public/NiagaraEditorModule.h
@@ -368,7 +368,12 @@ private:
 		{
 			const FAssetRegistryModule& AssetRegistryModule = FModuleManager::LoadModuleChecked<FAssetRegistryModule>(TEXT("AssetRegistry"));
 			TArray<FAssetData> AssetData;
-			AssetRegistryModule.GetRegistry().GetAssetsByClass(AssetType::StaticClass()->GetClassPathName(), AssetData);
+
+			FARFilter Filter;
+			Filter.ClassPaths.Add(AssetType::StaticClass()->GetClassPathName());
+			Filter.bRecursiveClasses = false;
+			Filter.bIncludeOnlyOnDiskAssets = IsRunningCookCommandlet();	// Enumerate disk-only assets during cook to ensure deterministic results
+			AssetRegistryModule.GetRegistry().GetAssets(Filter, AssetData);
 
 			FAssetToolsModule& AssetToolsModule = FModuleManager::Get().LoadModuleChecked<FAssetToolsModule>("AssetTools");
 			CachedAssets.Reset(AssetData.Num());
