- Concurrent sim cache reading when no data interfaces are cached
- Can be disabled using "fx.Niagara.SystemSimulation.AllowASyncSimCache"
- Enabled concurrent execution for fixed DT systems if it's a single tick
- Fixed a bug in the component recaching where it ran all the time in the editor

diff --git a/Engine/Plugins/FX/Niagara/Source/Niagara/Classes/NiagaraSimCache.h b/Engine/Plugins/FX/Niagara/Source/Niagara/Classes/NiagaraSimCache.h
--- a/Engine/Plugins/FX/Niagara/Source/Niagara/Classes/NiagaraSimCache.h
+++ b/Engine/Plugins/FX/Niagara/Source/Niagara/Classes/NiagaraSimCache.h
@@ -480,6 +480,9 @@ public:
 	UFUNCTION(BlueprintCallable, Category = NiagaraSimCache)
 	NIAGARA_API TArray<FName> GetEmitterNames() const;
 
+	/** Returns true if we have stored data interfaces, otherwise false. */
+	bool HasStoredDataInterfaces() const { return !DataInterfaceStorage.IsEmpty(); }
+
 	/** Returns a list of data interfaces we have captured in the SimCache. */
 	NIAGARA_API TArray<FNiagaraVariableBase> GetStoredDataInterfaces() const;
 	
@@ -650,6 +653,8 @@ private:
 
 	mutable std::atomic<int32> PendingCommandsInFlight;
 
+	mutable FRWLock RWLock;
+
 public:
 	static NIAGARA_API FOnCacheBeginWrite	OnCacheBeginWrite;
 	static NIAGARA_API FOnCacheEndWrite		OnCacheEndWrite;
diff --git a/Engine/Plugins/FX/Niagara/Source/Niagara/Private/NiagaraSimCache.cpp b/Engine/Plugins/FX/Niagara/Source/Niagara/Private/NiagaraSimCache.cpp
--- a/Engine/Plugins/FX/Niagara/Source/Niagara/Private/NiagaraSimCache.cpp
+++ b/Engine/Plugins/FX/Niagara/Source/Niagara/Private/NiagaraSimCache.cpp
@@ -314,6 +314,8 @@ bool UNiagaraSimCache::BeginWrite(FNiagaraSimCacheCreateParameters InCreateParam
 
 bool UNiagaraSimCache::BeginWrite(FNiagaraSimCacheCreateParameters InCreateParameters, UNiagaraComponent* NiagaraComponent, FNiagaraSimCacheFeedbackContext& FeedbackContext)
 {
+	UE::TWriteScopeLock RWGuard(RWLock);
+
 	check(PendingCommandsInFlight == 0);
 
 	OnCacheBeginWrite.Broadcast(this);
@@ -523,6 +525,8 @@ bool UNiagaraSimCache::BeginAppend(FNiagaraSimCacheCreateParameters InCreatePara
 
 bool UNiagaraSimCache::BeginAppend(FNiagaraSimCacheCreateParameters InCreateParameters, UNiagaraComponent* NiagaraComponent, FNiagaraSimCacheFeedbackContext& FeedbackContext)
 {
+	UE::TWriteScopeLock RWGuard(RWLock);
+
 	check(PendingCommandsInFlight == 0);
 
 	OnCacheBeginWrite.Broadcast(this);
@@ -597,6 +601,8 @@ bool UNiagaraSimCache::WriteFrame(UNiagaraComponent* NiagaraComponent)
 
 bool UNiagaraSimCache::WriteFrame(UNiagaraComponent* NiagaraComponent, FNiagaraSimCacheFeedbackContext& FeedbackContext)
 {
+	UE::TWriteScopeLock RWGuard(RWLock);
+
 	FNiagaraSimCacheHelper Helper(NiagaraComponent);
 	if (Helper.HasValidSimulationData() == false)
 	{
@@ -770,6 +776,8 @@ bool UNiagaraSimCache::WriteFrame(UNiagaraComponent* NiagaraComponent, FNiagaraS
 
 bool UNiagaraSimCache::EndWrite(bool bAllowAnalytics)
 {
+	UE::TWriteScopeLock RWGuard(RWLock);
+
 	check(PendingCommandsInFlight == 0);
 	if ( CacheFrames.Num() == 0 )
 	{
@@ -916,6 +924,7 @@ bool UNiagaraSimCache::CanRead(UNiagaraSystem* NiagaraSystem)
 				{
 					bNeedsReadComponentMappingRecache = true;
 				}
+				++CacheVMIndex;
 			}
 		);
 	}
@@ -923,6 +932,10 @@ bool UNiagaraSimCache::CanRead(UNiagaraSystem* NiagaraSystem)
 
 	if (bNeedsReadComponentMappingRecache)
 	{
+		// This modifies the component mappings that concurrent reading could be using
+		// We should never really hit any concurrent readers while entering here, this is for safety
+		UE::TWriteScopeLock RWGuard(RWLock);
+
 		const int32 NumEmitters = NiagaraSystem->GetEmitterHandles().Num();
 		if (NumEmitters != CacheLayout.EmitterLayouts.Num())
 		{
@@ -968,11 +981,13 @@ bool UNiagaraSimCache::CanRead(UNiagaraSystem* NiagaraSystem)
 
 bool UNiagaraSimCache::Read(float TimeSeconds, FNiagaraSystemInstance* SystemInstance) const
 {
+	UE::TReadScopeLock RWGuard(RWLock);
+
 	// Adjust time to match our allow resolution
 	TimeSeconds = FMath::RoundToFloat(TimeSeconds * CacheAgeResolution) / CacheAgeResolution;
 
 	const float RelativeTime = FMath::Max(TimeSeconds - StartSeconds, 0.0f);
-	if ( RelativeTime > DurationSeconds )
+	if ( RelativeTime > DurationSeconds || CacheFrames.Num() == 0 )
 	{
 		// Complete
 		return false;
@@ -1026,6 +1041,8 @@ bool UNiagaraSimCache::Read(float TimeSeconds, FNiagaraSystemInstance* SystemIns
 
 bool UNiagaraSimCache::ReadFrame(int32 FrameIndex, float FrameFraction, FNiagaraSystemInstance* SystemInstance) const
 {
+	UE::TReadScopeLock RWGuard(RWLock);
+
 	FNiagaraSimCacheHelper Helper(SystemInstance);
 	if ( !Helper.HasValidSimulation() )
 	{
diff --git a/Engine/Plugins/FX/Niagara/Source/Niagara/Private/NiagaraSystemInstance.cpp b/Engine/Plugins/FX/Niagara/Source/Niagara/Private/NiagaraSystemInstance.cpp
--- a/Engine/Plugins/FX/Niagara/Source/Niagara/Private/NiagaraSystemInstance.cpp
+++ b/Engine/Plugins/FX/Niagara/Source/Niagara/Private/NiagaraSystemInstance.cpp
@@ -2425,7 +2425,6 @@ void FNiagaraSystemInstance::SimCacheTick_GameThread(UNiagaraSimCache* SimCache,
 	{
 		WorldTransform = OwnerTransform.GetValue();
 	}
-	//TickInstanceParameters_GameThread(DeltaSeconds);
 
 	// We still need to tick the override parameters as these might pass through to material bindings
 	if (OverrideParameters)
@@ -2434,55 +2433,34 @@ void FNiagaraSystemInstance::SimCacheTick_GameThread(UNiagaraSimCache* SimCache,
 		OverrideParameters->Tick();
 	}
 
-	//TickDataInterfaces(DeltaSeconds, false);
-
 	Age = DesiredAge;
 	TickCount += 1;
 
-#if WITH_EDITOR
-	//// We need to tick the rapid iteration parameters when in the editor
-	//for (auto& EmitterInstance : Emitters)
-	//{
-	//	if (EmitterInstance->ShouldTick())
-	//	{
-	//		EmitterInstance->TickRapidIterationParameters();
-	//	}
-	//}
-#endif
+	SystemSimulation->SimCacheTick_GameThread(SimCache, DeltaSeconds, MyCompletionGraphEvent);
+}
 
-	//-OPT: On cooked builds we do not need to do this every frame
-	if ( SimCache->CanRead(GetSystem()) == false )
-	{
-		Complete(false);
-		return;
-	}
+void FNiagaraSystemInstance::SimCacheTickFinalize_GameThread()
+{
+	// Ensure concurrent work is complete and clear the finalize ref
+	check(ConcurrentTickGraphEvent == nullptr || ConcurrentTickGraphEvent->IsComplete());
+	ConcurrentTickGraphEvent = nullptr;
+	check(ConcurrentTickBatchGraphEvent == nullptr || ConcurrentTickBatchGraphEvent->IsComplete());
+	ConcurrentTickBatchGraphEvent = nullptr;
+	FinalizeRef.ConditionalClear();
 
-	//-TODO: This can run async
-	const bool bCacheValid = SimCache->Read(Age, this);
-	if (bCacheValid == false)
+	if (GetActualExecutionState() == ENiagaraExecutionState::Complete)
 	{
 		Complete(false);
-		return;
-	}
-
-	if ( SystemInstanceState == ENiagaraSystemInstanceState::PendingSpawn )
-	{
-		SystemSimulation->SetInstanceState(this, ENiagaraSystemInstanceState::Running);
 	}
-
-	SystemSimulation->SimCachePostTick_Concurrent(DeltaSeconds, MyCompletionGraphEvent);
-
-	if (OnPostTickDelegate.IsBound())
+	else
 	{
-		OnPostTickDelegate.Execute();
+		if (OnPostTickDelegate.IsBound())
+		{
+			OnPostTickDelegate.Execute();
+		}
 	}
 }
 
-void FNiagaraSystemInstance::SimCacheTick_Concurrent(UNiagaraSimCache* SimCache)
-{
-	//-OPT: Move SimCache read from GT to concurrent work
-}
-
 void FNiagaraSystemInstance::DumpStalledInfo()
 {
 	TStringBuilder<128> Builder;
diff --git a/Engine/Plugins/FX/Niagara/Source/Niagara/Private/NiagaraSystemSimulation.cpp b/Engine/Plugins/FX/Niagara/Source/Niagara/Private/NiagaraSystemSimulation.cpp
--- a/Engine/Plugins/FX/Niagara/Source/Niagara/Private/NiagaraSystemSimulation.cpp
+++ b/Engine/Plugins/FX/Niagara/Source/Niagara/Private/NiagaraSystemSimulation.cpp
@@ -6,6 +6,7 @@
 #include "NiagaraModule.h"
 #include "NiagaraTypes.h"
 #include "NiagaraParameterCollection.h"
+#include "NiagaraSimCache.h"
 #include "NiagaraSimulationTaskPriority.h"
 #include "NiagaraSystemEmitterState.h"
 #include "NiagaraSystemInstance.h"
@@ -78,7 +79,15 @@ static int32 GNiagaraSystemSimulationAllowASync = 1;
 static FAutoConsoleVariableRef CVarNiagaraSystemSimulationAllowASync(
 	TEXT("fx.Niagara.SystemSimulation.AllowASync"),
 	GNiagaraSystemSimulationAllowASync,
-	TEXT("If > 0, system post tick is parallelized. \n"),
+	TEXT("If > 0, system simulation tick is parallelized. \n"),
+	ECVF_Default
+);
+
+static int32 GNiagaraSystemSimulationAllowASyncSimCache = 1;
+static FAutoConsoleVariableRef CVarNiagaraSystemSimulationAllowASyncSimCache(
+	TEXT("fx.Niagara.SystemSimulation.AllowASyncSimCache"),
+	GNiagaraSystemSimulationAllowASyncSimCache,
+	TEXT("If > 0, sim cache tick is parallelized. \n"),
 	ECVF_Default
 );
 
@@ -435,7 +444,6 @@ struct FNiagaraSystemInstanceTickConcurrentTask
 	UWorld* WorldContext = nullptr;
 };
 
-
 //////////////////////////////////////////////////////////////////////////
 // Task to run FNiagaraSystemInstance::Finalize_GameThread
 struct FNiagaraSystemInstanceFinalizeTask
@@ -1114,7 +1122,7 @@ void FNiagaraSystemSimulation::Tick_GameThread(float DeltaSeconds, const FGraphE
 			TickInfo.TimeStepFraction = 1.0f * (i + 1) / Ticks;
 
 			//Cannot do multiple tick off the game thread here without additional work. So we pass in null for the completion event which will force GT execution.
-			Tick_GameThread_Internal(FixedDelta, nullptr);
+			Tick_GameThread_Internal(FixedDelta, Ticks == 1 ? MyCompletionGraphEvent : nullptr);
 			Budget -= FixedDelta;
 		}
 		FixedDeltaTickAge += DeltaSeconds;
@@ -1595,18 +1603,129 @@ void FNiagaraSystemSimulation::Spawn_Concurrent(FNiagaraSystemSimulationTickCont
 	bInSpawnPhase = false;
 }
 
-void FNiagaraSystemSimulation::SimCachePostTick_Concurrent(float DeltaSeconds, const FGraphEventRef& MyCompletionGraphEvent)
+void FNiagaraSystemSimulation::SimCacheTick_GameThread(UNiagaraSimCache* SimCache, float DeltaSeconds, const FGraphEventRef& MyCompletionGraphEvent)
 {
-	check(GetIsSolo());
 	check(System != nullptr);
 
+	// Work may not be complete if we back to back tick the GameThread without sending EOF updates
+	WaitForInstancesTickComplete();
+
+	// Move any pending spawn instance into the running state as they will be either complete or running by the end
+	{
+		TArray<FNiagaraSystemInstance*>& PendingInstances = GetSystemInstances(ENiagaraSystemInstanceState::PendingSpawn);
+		while (PendingInstances.Num() > 0)
+		{
+			SetInstanceState(PendingInstances.Last(), ENiagaraSystemInstanceState::Running);
+		}
+	}
+
+	// Get pointer to instances
 	TArray<FNiagaraSystemInstance*>& SystemInstances = GetSystemInstances(ENiagaraSystemInstanceState::Running);
-	if (SystemInstances.Num() == 0)
+
+	//-OPT: On cooked builds we do not need to do this every frame
+	if (SimCache->CanRead(System) == false)
 	{
+		for (FNiagaraSystemInstance* SystemInstance : SystemInstances)
+		{
+			SystemInstance->Complete(false);
+		}
 		return;
 	}
 
-	if (!bBindingsInitialized)
+	// Should we run concurrently or not?
+	// Data interfaces are not thread safe for reading, we will need to implement a way to do this....
+	if (GNiagaraSystemSimulationAllowASyncSimCache && !SimCache->HasStoredDataInterfaces())
+	{
+		// Create finalize tick batch information to fit into how things work
+		struct FFinalizeTickBatch
+		{
+			TArray<FNiagaraSystemInstance*, TInlineAllocator<4>> Instances;
+		#if DO_CHECK
+			std::atomic<int> DebugCounter = 0;
+		#endif
+		};
+		TUniquePtr<FFinalizeTickBatch> FinalizeTickBatch(new FFinalizeTickBatch());
+
+		FinalizeTickBatch->Instances.Reserve(SystemInstances.Num());
+		for (FNiagaraSystemInstance* SystemInstance : SystemInstances)
+		{
+			SystemInstance->ConcurrentTickGraphEvent = ConcurrentTickGraphEvent;
+			FinalizeTickBatch->Instances.Add(SystemInstance);
+
+			FNiagaraSystemInstanceFinalizeRef FinalizeRef(&FinalizeTickBatch->Instances.Last());
+		#if DO_CHECK
+			FinalizeRef.SetDebugCounter(&FinalizeTickBatch->DebugCounter);
+		#endif
+			SystemInstance->SetPendingFinalize(FinalizeRef);
+		}
+
+		// Create events to signal old task system
+		ConcurrentTickGraphEvent = FGraphEvent::CreateGraphEvent();
+		AllWorkCompleteGraphEvent = FGraphEvent::CreateGraphEvent();
+
+		// Execute concurrent work
+		UE::Tasks::FTask ConcurrentTask = UE::Tasks::Launch(
+			TEXT("Niagara::SimCacheTick_Concurrent"),
+			[this, SimCache, DeltaSeconds]()
+			{
+				TRACE_CPUPROFILER_EVENT_SCOPE(Niagara::SimCacheTick_Concurrent);
+
+				SimCacheTick_Concurrent(SimCache, DeltaSeconds);
+				ConcurrentTickGraphEvent->DispatchSubsequents();
+			}
+		);
+
+		// Execute finalize work
+		UE::Tasks::Launch(
+			TEXT("Niagara::SimCacheTick_Finalize"),
+			[this, TickBatch=MoveTemp(FinalizeTickBatch)]()
+			{
+				TRACE_CPUPROFILER_EVENT_SCOPE(Niagara::SimCacheTick_Finalize);
+
+				for (FNiagaraSystemInstance* Instance : TickBatch->Instances)
+				{
+					if (Instance)
+					{
+						Instance->SimCacheTickFinalize_GameThread();
+					}
+				}
+			#if DO_CHECK
+				ensureMsgf(TickBatch->DebugCounter.load() == 0, TEXT("Finalize batch is complete but counter is %d when it should be zero"), TickBatch->DebugCounter.load());
+			#endif
+				AllWorkCompleteGraphEvent->DispatchSubsequents();
+			},
+			ConcurrentTask,
+			UE::Tasks::ETaskPriority::Normal,
+			UE::Tasks::EExtendedTaskPriority::GameThreadNormalPri
+		);
+	}
+	else
+	{
+		SimCacheTick_Concurrent(SimCache, DeltaSeconds);
+
+		for (FNiagaraSystemInstance* SystemInstance : SystemInstances)
+		{
+			SystemInstance->SimCacheTickFinalize_GameThread();
+		}
+	}
+}
+
+void FNiagaraSystemSimulation::SimCacheTick_Concurrent(UNiagaraSimCache* SimCache, float DeltaSeconds)
+{
+	TArray<FNiagaraSystemInstance*>& SystemInstances = GetSystemInstances(ENiagaraSystemInstanceState::Running);
+
+	for (int32 iSystemInstance = 0; iSystemInstance < SystemInstances.Num(); ++iSystemInstance)
+	{
+		FNiagaraSystemInstance& SystemInstance = *SystemInstances[iSystemInstance];
+
+		const bool bCacheValid = SimCache->Read(SystemInstance.GetAge(), &SystemInstance);
+		if (bCacheValid == false)
+		{
+			SystemInstance.SetActualExecutionState(ENiagaraExecutionState::Complete);
+		}
+	}
+
+	if (!bBindingsInitialized && SystemInstances.Num() > 0)
 	{
 		InitParameterDataSetBindings(SystemInstances[0]);
 	}
diff --git a/Engine/Plugins/FX/Niagara/Source/Niagara/Public/NiagaraSystemInstance.h b/Engine/Plugins/FX/Niagara/Source/Niagara/Public/NiagaraSystemInstance.h
--- a/Engine/Plugins/FX/Niagara/Source/Niagara/Public/NiagaraSystemInstance.h
+++ b/Engine/Plugins/FX/Niagara/Source/Niagara/Public/NiagaraSystemInstance.h
@@ -206,22 +206,22 @@ public:
 	NIAGARA_API void MultiTick(int32 TicksToProcess, float DeltaSecondsPerTick, float MaxSimTime, const FGraphEventRef& MyCompletionGraphEvent);
 
 	/** Ticks the system using the a SimCache. */
-	NIAGARA_API void SimCacheTick_GameThread(UNiagaraSimCache* SimCache, float DesiredAge, float DeltaSeconds, const FGraphEventRef& MyCompletionGraphEvent);
-	/** Concurrent work for SimCache tick */
-	NIAGARA_API void SimCacheTick_Concurrent(UNiagaraSimCache* SimCache);
+	void SimCacheTick_GameThread(UNiagaraSimCache* SimCache, float DesiredAge, float DeltaSeconds, const FGraphEventRef& MyCompletionGraphEvent);
+	/** Final step after concurrent tick has run. */
+	void SimCacheTickFinalize_GameThread();
 
 	/** Initial phase of system instance tick. Must be executed on the game thread. */
-	NIAGARA_API void Tick_GameThread(float DeltaSeconds);
+	void Tick_GameThread(float DeltaSeconds);
 	/** Secondary phase of the system instance tick that can be executed on any thread. */
-	NIAGARA_API void Tick_Concurrent(bool bEnqueueGPUTickIfNeeded = true);
+	void Tick_Concurrent(bool bEnqueueGPUTickIfNeeded = true);
 	/** 
 		Final phase of system instance tick. Must be executed on the game thread. 
 		Returns whether the Finalize was actually done. It's possible for the finalize in a task to have already been done earlier on the GT by a WaitForAsyncAndFinalize call.
 	*/
-	NIAGARA_API void FinalizeTick_GameThread(bool bEnqueueGPUTickIfNeeded = true);
+	void FinalizeTick_GameThread(bool bEnqueueGPUTickIfNeeded = true);
 
-	NIAGARA_API void GenerateAndSubmitGPUTick();
-	NIAGARA_API void InitGPUTick(FNiagaraGPUSystemTick& OutTick);
+	void GenerateAndSubmitGPUTick();
+	void InitGPUTick(FNiagaraGPUSystemTick& OutTick);
 
 	void SetPendingFinalize(FNiagaraSystemInstanceFinalizeRef InFinalizeRef) { FinalizeRef = InFinalizeRef; }
 	bool HasPendingFinalize() const { return FinalizeRef.IsPending(); }
diff --git a/Engine/Plugins/FX/Niagara/Source/Niagara/Public/NiagaraSystemSimulation.h b/Engine/Plugins/FX/Niagara/Source/Niagara/Public/NiagaraSystemSimulation.h
--- a/Engine/Plugins/FX/Niagara/Source/Niagara/Public/NiagaraSystemSimulation.h
+++ b/Engine/Plugins/FX/Niagara/Source/Niagara/Public/NiagaraSystemSimulation.h
@@ -274,8 +274,10 @@ public:
 	/** Spawn any pending instances */
 	void Spawn_Concurrent(FNiagaraSystemSimulationTickContext& Context);
 
-	/** Called after the sim cache has been read. */
-	void SimCachePostTick_Concurrent(float DeltaSeconds, const FGraphEventRef& MyCompletionGraphEvent);
+	/** Game thread portion of sim cache tick.  Must run on GameThread. */
+	void SimCacheTick_GameThread(UNiagaraSimCache* SimCache, float DeltaSeconds, const FGraphEventRef& MyCompletionGraphEvent);
+	/** Concurrent portion of sim cache tick, can run on any thread. */
+	void SimCacheTick_Concurrent(UNiagaraSimCache* SimCache, float DeltaSeconds);
 
 	/** Wait for system simulation concurrent tick to complete.  If bEnsureComplete is true we will trigger an ensure if it is not complete. */
 	void WaitForConcurrentTickComplete(bool bEnsureComplete = false);
