- Add CanRunAsync to sim cache
- This will detect if it's possible to run async or not, right now it's overly conservative and we can relax that as time goes on

diff --git a/Engine/Plugins/FX/Niagara/Source/Niagara/Classes/NiagaraSimCache.h b/Engine/Plugins/FX/Niagara/Source/Niagara/Classes/NiagaraSimCache.h
--- a/Engine/Plugins/FX/Niagara/Source/Niagara/Classes/NiagaraSimCache.h
+++ b/Engine/Plugins/FX/Niagara/Source/Niagara/Classes/NiagaraSimCache.h
@@ -451,6 +451,8 @@ public:
 	NIAGARA_API bool Read(float TimeSeconds, FNiagaraSystemInstance* SystemInstance) const;
 	NIAGARA_API bool ReadFrame(int32 FrameIndex, float FrameFraction, FNiagaraSystemInstance* SystemInstance) const;
 
+	bool CanRunAsync() const { return bCanRunAsync; }
+
 	/** Get the time the simulation was at when recorded. */
 	UFUNCTION(BlueprintCallable, Category = NiagaraSimCache)
 	float GetStartSeconds() const { return StartSeconds; }
@@ -631,6 +633,9 @@ private:
 	UPROPERTY(transient)
 	bool bNeedsReadComponentMappingRecache = true;
 
+	UPROPERTY(transient)
+	bool bCanRunAsync = false;
+
 #if WITH_EDITORONLY_DATA
 	UPROPERTY(transient)
 	TArray<FNiagaraVMExecutableDataId> CachedScriptVMIds;
diff --git a/Engine/Plugins/FX/Niagara/Source/Niagara/Private/NiagaraSimCache.cpp b/Engine/Plugins/FX/Niagara/Source/Niagara/Private/NiagaraSimCache.cpp
--- a/Engine/Plugins/FX/Niagara/Source/Niagara/Private/NiagaraSimCache.cpp
+++ b/Engine/Plugins/FX/Niagara/Source/Niagara/Private/NiagaraSimCache.cpp
@@ -974,6 +974,28 @@ bool UNiagaraSimCache::CanRead(UNiagaraSystem* NiagaraSystem)
 		CachedScriptVMIds.Shrink();
 #endif
 		bNeedsReadComponentMappingRecache = false;
+
+		// Determine if we can run async or not
+		bCanRunAsync = !HasStoredDataInterfaces();
+		if (bCanRunAsync)
+		{
+			bool bHasAnyPostReadDataInterface = false;
+			FNiagaraDataInterfaceUtilities::ForEachDataInterface(
+				NiagaraSystem,
+				[&bHasAnyPostReadDataInterface](const FNiagaraDataInterfaceUtilities::FDataInterfaceUsageContext& DIContext)->bool
+				{
+					if (Cast<INiagaraSimCacheCustomStorageInterface>(DIContext.DataInterface) != nullptr)
+					{
+						bHasAnyPostReadDataInterface = true;
+					}
+					return !bHasAnyPostReadDataInterface;
+				}
+			);
+			if (bHasAnyPostReadDataInterface)
+			{
+				bCanRunAsync = false;
+			}
+		}
 	}
 
 	return true;
diff --git a/Engine/Plugins/FX/Niagara/Source/Niagara/Private/NiagaraSystemSimulation.cpp b/Engine/Plugins/FX/Niagara/Source/Niagara/Private/NiagaraSystemSimulation.cpp
--- a/Engine/Plugins/FX/Niagara/Source/Niagara/Private/NiagaraSystemSimulation.cpp
+++ b/Engine/Plugins/FX/Niagara/Source/Niagara/Private/NiagaraSystemSimulation.cpp
@@ -1634,7 +1634,7 @@ void FNiagaraSystemSimulation::SimCacheTick_GameThread(UNiagaraSimCache* SimCach
 
 	// Should we run concurrently or not?
 	// Data interfaces are not thread safe for reading, we will need to implement a way to do this....
-	if (GNiagaraSystemSimulationAllowASyncSimCache && !SimCache->HasStoredDataInterfaces())
+	if (GNiagaraSystemSimulationAllowASyncSimCache && SimCache->CanRunAsync())
 	{
 		// Create finalize tick batch information to fit into how things work
 		struct FFinalizeTickBatch
