- Couple of fixes for cache recording when multiple cache tracks exist

diff --git a/Engine/Plugins/FX/Niagara/Source/Niagara/Private/NiagaraSystemInstance.cpp b/Engine/Plugins/FX/Niagara/Source/Niagara/Private/NiagaraSystemInstance.cpp
--- a/Engine/Plugins/FX/Niagara/Source/Niagara/Private/NiagaraSystemInstance.cpp
+++ b/Engine/Plugins/FX/Niagara/Source/Niagara/Private/NiagaraSystemInstance.cpp
@@ -2448,7 +2448,7 @@ void FNiagaraSystemInstance::SimCacheTickFinalize_GameThread()
 	ConcurrentTickBatchGraphEvent = nullptr;
 	FinalizeRef.ConditionalClear();
 
-	if (GetActualExecutionState() == ENiagaraExecutionState::Complete)
+	if (RequestedExecutionState == ENiagaraExecutionState::Complete)
 	{
 		Complete(false);
 	}
diff --git a/Engine/Plugins/FX/Niagara/Source/Niagara/Private/NiagaraSystemSimulation.cpp b/Engine/Plugins/FX/Niagara/Source/Niagara/Private/NiagaraSystemSimulation.cpp
--- a/Engine/Plugins/FX/Niagara/Source/Niagara/Private/NiagaraSystemSimulation.cpp
+++ b/Engine/Plugins/FX/Niagara/Source/Niagara/Private/NiagaraSystemSimulation.cpp
@@ -1610,24 +1610,15 @@ void FNiagaraSystemSimulation::SimCacheTick_GameThread(UNiagaraSimCache* SimCach
 	// Work may not be complete if we back to back tick the GameThread without sending EOF updates
 	WaitForInstancesTickComplete();
 
-	// Move any pending spawn instance into the running state as they will be either complete or running by the end
-	{
-		TArray<FNiagaraSystemInstance*>& PendingInstances = GetSystemInstances(ENiagaraSystemInstanceState::PendingSpawn);
-		while (PendingInstances.Num() > 0)
-		{
-			SetInstanceState(PendingInstances.Last(), ENiagaraSystemInstanceState::Running);
-		}
-	}
-
 	// Get pointer to instances
 	TArray<FNiagaraSystemInstance*>& SystemInstances = GetSystemInstances(ENiagaraSystemInstanceState::Running);
 
 	//-OPT: On cooked builds we do not need to do this every frame
 	if (SimCache->CanRead(System) == false)
 	{
-		for (FNiagaraSystemInstance* SystemInstance : SystemInstances)
+		for (int32 i=SystemInstances.Num() - 1; i >= 0; --i)
 		{
-			SystemInstance->Complete(false);
+			SystemInstances[i]->Complete(false);
 		}
 		return;
 	}
@@ -1703,39 +1694,57 @@ void FNiagaraSystemSimulation::SimCacheTick_GameThread(UNiagaraSimCache* SimCach
 	{
 		SimCacheTick_Concurrent(SimCache, DeltaSeconds);
 
-		for (FNiagaraSystemInstance* SystemInstance : SystemInstances)
+		for ( int32 i=SystemInstances.Num()-1; i >= 0; --i)
 		{
-			SystemInstance->SimCacheTickFinalize_GameThread();
+			SystemInstances[i]->SimCacheTickFinalize_GameThread();
 		}
 	}
 }
 
 void FNiagaraSystemSimulation::SimCacheTick_Concurrent(UNiagaraSimCache* SimCache, float DeltaSeconds)
 {
+	// Read running instances
 	TArray<FNiagaraSystemInstance*>& SystemInstances = GetSystemInstances(ENiagaraSystemInstanceState::Running);
-
-	for (int32 iSystemInstance = 0; iSystemInstance < SystemInstances.Num(); ++iSystemInstance)
+	for (FNiagaraSystemInstance* SystemInstance : SystemInstances)
 	{
-		FNiagaraSystemInstance& SystemInstance = *SystemInstances[iSystemInstance];
+		if ( !SimCache->Read(SystemInstance->GetAge(), SystemInstance) )
+		{
+			SystemInstance->RequestedExecutionState = ENiagaraExecutionState::Complete;
+		}
+	}
 
-		const bool bCacheValid = SimCache->Read(SystemInstance.GetAge(), &SystemInstance);
-		if (bCacheValid == false)
+	// Move pending spawn instances into running instances
+	{
+		TArray<FNiagaraSystemInstance*>& PendingInstances = GetSystemInstances(ENiagaraSystemInstanceState::PendingSpawn);
+		for (int32 i=PendingInstances.Num()-1; i >= 0; --i)
 		{
-			SystemInstance.SetActualExecutionState(ENiagaraExecutionState::Complete);
+			FNiagaraSystemInstance* SystemInstance = PendingInstances[i];
+			if (!SimCache->Read(SystemInstance->GetAge(), SystemInstance))
+			{
+				SystemInstance->RequestedExecutionState = ENiagaraExecutionState::Complete;
+			}
+			else
+			{
+				SetInstanceState(SystemInstance, ENiagaraSystemInstanceState::Running);
+			}
 		}
 	}
 
+	// Initialize bindings (if needed)
 	if (!bBindingsInitialized && SystemInstances.Num() > 0)
 	{
 		InitParameterDataSetBindings(SystemInstances[0]);
 	}
 
-	FNiagaraDataSetReaderInt32<ENiagaraExecutionState> SystemExecutionStateAccessor = System->GetSystemExecutionStateAccessor().GetReader(MainDataSet);
+	// Copy over execution states push to stores, etc
 	TConstArrayView<FNiagaraDataSetAccessor<ENiagaraExecutionState>> EmitterExecutionStateAccessors = System->GetEmitterExecutionStateAccessors();
-
 	for ( int32 iSystemInstance=0; iSystemInstance < SystemInstances.Num(); ++iSystemInstance )
 	{
 		FNiagaraSystemInstance& SystemInstance = *SystemInstances[iSystemInstance];
+		if (SystemInstance.RequestedExecutionState == ENiagaraExecutionState::Complete)
+		{
+			continue;
+		}
 
 		TArrayView<FNiagaraEmitterInstanceRef> Emitters = SystemInstance.GetEmitters();
 		for (int32 iEmitter=0; iEmitter < Emitters.Num(); ++iEmitter)
