- Add all base Niagara float types to sim cache interpolation
- Enabled by default can be disabled changing bInterpolateAllTypes in the creation parameters
- This improves MRQ renders with temporal samples

diff --git a/Engine/Plugins/FX/Niagara/Source/Niagara/Classes/NiagaraSimCache.h b/Engine/Plugins/FX/Niagara/Source/Niagara/Classes/NiagaraSimCache.h
--- a/Engine/Plugins/FX/Niagara/Source/Niagara/Classes/NiagaraSimCache.h
+++ b/Engine/Plugins/FX/Niagara/Source/Niagara/Classes/NiagaraSimCache.h
@@ -68,11 +68,17 @@ struct FNiagaraSimCacheCreateParameters
 	/**
 	When enabled we allow the cache to be generated for interpolation.
 	This will increase the memory usage for the cache slightly but can allow you to reduce the capture rate.
-	By default we will capture and interpolate all Position & Quat types, you can adjust this using the include / exclude list.
 	*/
 	UPROPERTY(EditAnywhere, BlueprintReadWrite, Category="SimCache")
 	uint32 bAllowInterpolation : 1 = false;
 
+	/**
+	When enabled all support types for interpolation will be included, this set may increase in future releases.
+	When disabled we only interpolate Position / Quat attributes.
+	*/
+	UPROPERTY(EditAnywhere, BlueprintReadWrite, Category = "SimCache", meta = (EditCondition = "bAllowInterpolation"))
+	uint32 bInterpolateAllTypes : 1 = true;
+	
 	/**
 	When enabled we allow the cache to be generated for extrapolation.
 	This will force the velocity attribute to be maintained.
diff --git a/Engine/Plugins/FX/Niagara/Source/Niagara/Private/NiagaraSimCacheHelper.h b/Engine/Plugins/FX/Niagara/Source/Niagara/Private/NiagaraSimCacheHelper.h
--- a/Engine/Plugins/FX/Niagara/Source/Niagara/Private/NiagaraSimCacheHelper.h
+++ b/Engine/Plugins/FX/Niagara/Source/Niagara/Private/NiagaraSimCacheHelper.h
@@ -306,7 +306,7 @@ struct FNiagaraSimCacheHelper
 		{
 			for (const FNiagaraVariableBase& Variable : SystemCompileData.Variables)
 			{
-				if (!CanInterpolateVariable(Variable) || CreateParameters.RebaseExcludeAttributes.Contains(Variable.GetName()))
+				if (!CanInterpolateVariable(CreateParameters.bInterpolateAllTypes, Variable) || CreateParameters.InterpolationExcludeAttributes.Contains(Variable.GetName()))
 				{
 					continue;
 				}
@@ -431,7 +431,7 @@ struct FNiagaraSimCacheHelper
 		{
 			for (const FNiagaraVariableBase& Variable : ParticleDataSet->Variables)
 			{
-				if (!CanInterpolateVariable(Variable) || CreateParameters.RebaseExcludeAttributes.Contains(Variable.GetName()))
+				if (!CanInterpolateVariable(CreateParameters.bInterpolateAllTypes, Variable) || CreateParameters.InterpolationExcludeAttributes.Contains(Variable.GetName()))
 				{
 					continue;
 				}
@@ -616,6 +616,29 @@ struct FNiagaraSimCacheHelper
 							DestVariableLayout = nullptr;
 						}
 					}
+					else
+					{
+						if (SourceVariable.Variable.GetType() == FNiagaraTypeDefinition::GetFloatDef())
+						{
+							CacheBufferReadInfo.VariableCopyMappingsToDataBuffer.Emplace(SourceVariable.FloatOffset, IntCastChecked<uint16>(DestVariableLayout->GetFloatComponentStart()), &FNiagaraSimCacheHelper::InterpFloats<1>);
+							DestVariableLayout = nullptr;
+						}
+						else if (SourceVariable.Variable.GetType() == FNiagaraTypeDefinition::GetVec2Def())
+						{
+							CacheBufferReadInfo.VariableCopyMappingsToDataBuffer.Emplace(SourceVariable.FloatOffset, IntCastChecked<uint16>(DestVariableLayout->GetFloatComponentStart()), &FNiagaraSimCacheHelper::InterpFloats<2>);
+							DestVariableLayout = nullptr;
+						}
+						else if (SourceVariable.Variable.GetType() == FNiagaraTypeDefinition::GetVec3Def())
+						{
+							CacheBufferReadInfo.VariableCopyMappingsToDataBuffer.Emplace(SourceVariable.FloatOffset, IntCastChecked<uint16>(DestVariableLayout->GetFloatComponentStart()), &FNiagaraSimCacheHelper::InterpFloats<3>);
+							DestVariableLayout = nullptr;
+						}
+						else if (SourceVariable.Variable.GetType() == FNiagaraTypeDefinition::GetVec4Def() || SourceVariable.Variable.GetType() == FNiagaraTypeDefinition::GetColorDef())
+						{
+							CacheBufferReadInfo.VariableCopyMappingsToDataBuffer.Emplace(SourceVariable.FloatOffset, IntCastChecked<uint16>(DestVariableLayout->GetFloatComponentStart()), &FNiagaraSimCacheHelper::InterpFloats<4>);
+							DestVariableLayout = nullptr;
+						}
+					}
 				}
 			}
 
@@ -924,8 +947,21 @@ struct FNiagaraSimCacheHelper
 		);
 	}
 
-	static bool CanInterpolateVariable(const FNiagaraVariableBase& Variable)
+	static bool CanInterpolateVariable(bool bInterpolateAllTypes, const FNiagaraVariableBase& Variable)
 	{
+		if (bInterpolateAllTypes)
+		{
+			if (
+				Variable.GetType() == FNiagaraTypeDefinition::GetFloatDef() ||
+				Variable.GetType() == FNiagaraTypeDefinition::GetVec2Def() ||
+				Variable.GetType() == FNiagaraTypeDefinition::GetVec3Def() ||
+				Variable.GetType() == FNiagaraTypeDefinition::GetVec4Def() ||
+				Variable.GetType() == FNiagaraTypeDefinition::GetColorDef())
+			{
+				return true;
+			}
+		}
+
 		return
 			(Variable.GetType() == FNiagaraTypeDefinition::GetQuatDef()) ||
 			(Variable.GetType() == FNiagaraTypeDefinition::GetPositionDef());
@@ -1125,6 +1161,24 @@ struct FNiagaraSimCacheHelper
 		}
 	}
 
+	template<uint32 NumFloats>
+	static void InterpFloats(const FNiagaraSimCacheDataBuffersLayout::FVariableCopyContext& CopyDataContext)
+	{
+		for (uint32 iComponent = 0; iComponent < NumFloats; ++iComponent)
+		{
+			const float* SrcAFloats = reinterpret_cast<const float*>(CopyDataContext.SourceACurr + (CopyDataContext.SourceAStride * iComponent));
+			const float* SrcBFloats = reinterpret_cast<const float*>(CopyDataContext.SourceBCurr + (CopyDataContext.SourceBStride * iComponent));
+			float* DstFloats  = reinterpret_cast<float*>(CopyDataContext.DestCurr + (CopyDataContext.DestStride * iComponent));
+			for (uint32 iInstanceA=0; iInstanceA < CopyDataContext.NumInstances; ++iInstanceA)
+			{
+				const uint32 iInstanceB = CopyDataContext.InterpMappings[iInstanceA];
+				const float FloatA = SrcAFloats[iInstanceA];
+				const float FloatB = iInstanceB == INDEX_NONE ? FloatA : SrcBFloats[iInstanceB];
+				DstFloats[iInstanceA] = FMath::Lerp(FloatA, FloatB, CopyDataContext.FrameFraction);
+			}
+		}
+	}
+
 	UNiagaraSystem*						NiagaraSystem = nullptr;
 	FNiagaraSystemInstance*				SystemInstance = nullptr;
 	FNiagaraSystemSimulationPtr			SystemSimulation = nullptr;
