Fix for DDC barriers during niagara compilation that could cause deadlocks/checks

#tests Created fabricated test to generate a check() and verified that it was remedied

diff --git a/Engine/Plugins/FX/Niagara/Source/NiagaraEditor/Private/NiagaraCompilationTasks.cpp b/Engine/Plugins/FX/Niagara/Source/NiagaraEditor/Private/NiagaraCompilationTasks.cpp
--- a/Engine/Plugins/FX/Niagara/Source/NiagaraEditor/Private/NiagaraCompilationTasks.cpp
+++ b/Engine/Plugins/FX/Niagara/Source/NiagaraEditor/Private/NiagaraCompilationTasks.cpp
@@ -618,7 +618,6 @@ namespace NiagaraCompilationTasksImpl
 FNiagaraSystemCompilationTask::FNiagaraSystemCompilationTask(FNiagaraCompilationTaskHandle InTaskHandle, UNiagaraSystem* InSystem, ENiagaraCompileRIParamMode InRIParamMode)
 	: TaskHandle(InTaskHandle)
 	, RIParamMode(InRIParamMode)
-	, DDCRequestOwner(UE::DerivedData::EPriority::Normal)
 	, System_GT(InSystem)
 {
 	bGenerateCompileMetaData = GetDefault<UNiagaraSettings>()->bGenerateMetaDataOnCompile;
@@ -1256,6 +1255,12 @@ bool FNiagaraSystemCompilationTask::FDDCTaskInfo::ResolveGet(bool bSuccess)
 	return bSuccess;
 }
 
+FNiagaraSystemCompilationTask::FDataCacheRequestBase::FDataCacheRequestBase()
+	: RequestOwner(UE::DerivedData::EPriority::Normal)
+{
+	RequestOwner.KeepAlive();
+}
+
 void FNiagaraSystemCompilationTask::FDispatchAndProcessDataCacheGetRequests::Launch(FNiagaraSystemCompilationTask* SystemCompileTask)
 {
 	// todo - incorporate something here where we skip if we try to update the compileid (post ri collection)
@@ -1332,8 +1337,8 @@ void FNiagaraSystemCompilationTask::FDispatchAndProcessDataCacheGetRequests::Lau
 
 	if (PendingGetRequestCount > 0)
 	{
-		FRequestBarrier AsyncBarrier(SystemCompileTask->DDCRequestOwner);
-		GetCache().Get(GetRequests, SystemCompileTask->DDCRequestOwner, [this, SystemCompileTask = SystemCompileTask->AsShared()](FCacheGetResponse&& Response)
+		FRequestBarrier RequestBarrier(RequestOwner);
+		GetCache().Get(GetRequests, RequestOwner, [this, SystemCompileTask = SystemCompileTask->AsShared()](FCacheGetResponse&& Response)
 		{
 			FDDCUserData Index;
 			Index.UserData = Response.UserData;
@@ -1491,8 +1496,8 @@ void FNiagaraSystemCompilationTask::FDispatchDataCachePutRequests::Launch(FNiaga
 
 	if (!PutRequests.IsEmpty())
 	{
-		FRequestBarrier AsyncBarrier(SystemCompileTask->DDCRequestOwner);
-		GetCache().Put(PutRequests, SystemCompileTask->DDCRequestOwner, [this](FCachePutResponse&& Response)
+		FRequestBarrier RequestBarrier(RequestOwner);
+		GetCache().Put(PutRequests, RequestOwner, [this, SystemCompileTask = SystemCompileTask->AsShared()](FCachePutResponse&& Response)
 		{
 			if (PendingPutRequestCount.fetch_sub(1, std::memory_order_relaxed) == 1)
 			{
diff --git a/Engine/Plugins/FX/Niagara/Source/NiagaraEditor/Private/NiagaraCompilationTasks.h b/Engine/Plugins/FX/Niagara/Source/NiagaraEditor/Private/NiagaraCompilationTasks.h
--- a/Engine/Plugins/FX/Niagara/Source/NiagaraEditor/Private/NiagaraCompilationTasks.h
+++ b/Engine/Plugins/FX/Niagara/Source/NiagaraEditor/Private/NiagaraCompilationTasks.h
@@ -11,8 +11,16 @@
 
 struct FNiagaraSystemCompilationTask : public TSharedFromThis<FNiagaraSystemCompilationTask, ESPMode::ThreadSafe>
 {
+	struct FDataCacheRequestBase
+	{
+		FDataCacheRequestBase();
+
+	protected:
+		UE::DerivedData::FRequestOwner RequestOwner;
+	};
+
 	// helper for script DDC get requests
-	struct FDispatchAndProcessDataCacheGetRequests
+	struct FDispatchAndProcessDataCacheGetRequests : public FDataCacheRequestBase
 	{
 		void Launch(FNiagaraSystemCompilationTask* SystemCompileTask);
 
@@ -21,7 +29,7 @@ struct FNiagaraSystemCompilationTask : public TSharedFromThis<FNiagaraSystemComp
 	};
 
 	// helper for script DDC put requests
-	struct FDispatchDataCachePutRequests
+	struct FDispatchDataCachePutRequests : public FDataCacheRequestBase
 	{
 		void Launch(FNiagaraSystemCompilationTask* SystemCompileTask);
 
@@ -258,8 +266,6 @@ struct FNiagaraSystemCompilationTask : public TSharedFromThis<FNiagaraSystemComp
 	bool GenerateCompileMetaData() const;
 	const FGuid& GetCompilationTaskId() const;
 
-	UE::DerivedData::FRequestOwner DDCRequestOwner;
-
 private:
 	FNiagaraSystemCompilationTask() = delete;
 
