Have niagara shader compilation not worry using the SharedEnvironment, as not all values set there are necessarily being evaluted (compile flags are evaluated on the base Environment in some cases).  As our shader maps are generally sparse the SharedEnvironment won't actually provide that much value even if it didn't have weird edge cases.

#tests Confirmed compiler flags from a DI make it to the shaderformat

diff --git a/Engine/Plugins/FX/Niagara/Source/NiagaraShader/Private/NiagaraShader.cpp b/Engine/Plugins/FX/Niagara/Source/NiagaraShader/Private/NiagaraShader.cpp
--- a/Engine/Plugins/FX/Niagara/Source/NiagaraShader/Private/NiagaraShader.cpp
+++ b/Engine/Plugins/FX/Niagara/Source/NiagaraShader/Private/NiagaraShader.cpp
@@ -426,7 +426,11 @@ void FNiagaraShaderType::BeginCompileShaderFromSource(
 		return;
 	}
 	NewJob->ShaderParameters = MakeShared<const FParameters, ESPMode::ThreadSafe>(ShaderParameters);
-	NewJob->Input.SharedEnvironment = CompilationEnvironment;
+	if (CompilationEnvironment)
+	{
+		// Initialize the Environment with what we've captured from the DI and the rest of the compilation process
+		NewJob->Input.Environment = *CompilationEnvironment;
+	}
 	NewJob->Input.Target = Target;
 	NewJob->Input.ShaderFormat = LegacyShaderPlatformToShaderFormat(Platform);
 	NewJob->Input.VirtualSourceFilePath = TEXT("/Plugin/FX/Niagara/Private/NiagaraEmitterInstanceShader.usf");
