Fixed memory leak in data channel object destructor

diff --git a/Engine/Plugins/FX/Niagara/Source/Niagara/Private/NiagaraDataChannelData.cpp b/Engine/Plugins/FX/Niagara/Source/Niagara/Private/NiagaraDataChannelData.cpp
--- a/Engine/Plugins/FX/Niagara/Source/Niagara/Private/NiagaraDataChannelData.cpp
+++ b/Engine/Plugins/FX/Niagara/Source/Niagara/Private/NiagaraDataChannelData.cpp
@@ -37,14 +37,19 @@ FNiagaraDataChannelData::~FNiagaraDataChannelData()
 
 	/** We defer the deletion of the dataset to the RT to be sure all in-flight RT commands have finished using it.*/
 	ENQUEUE_RENDER_COMMAND(FDeleteContextCommand)(
-		[CPUDataChannelDataSet = CPUSimData](FRHICommandListImmediate& RHICmdList)
+		[CPUDataChannelDataSet = CPUSimData, GameDataStagingDataSet = GameDataStaging](FRHICommandListImmediate& RHICmdList)
 	{
 		if (CPUDataChannelDataSet != nullptr)
 		{
 			delete CPUDataChannelDataSet;
 		}
+		if (GameDataStagingDataSet != nullptr)
+		{
+			delete GameDataStagingDataSet;
+		}
 	}
 	);
+	GameDataStaging = nullptr;
 	CPUSimData = nullptr;
 	RTProxy = nullptr;
 }
