Small bugfixes caught by static analysis:
* using max instead of min distance in LoD parameter
* nullptr and weak ptr check before dereferencing
* return early in case of error during script merge

diff --git a/Engine/Plugins/FX/Niagara/Source/Niagara/Private/NiagaraCommon.cpp b/Engine/Plugins/FX/Niagara/Source/Niagara/Private/NiagaraCommon.cpp
--- a/Engine/Plugins/FX/Niagara/Source/Niagara/Private/NiagaraCommon.cpp
+++ b/Engine/Plugins/FX/Niagara/Source/Niagara/Private/NiagaraCommon.cpp
@@ -184,7 +184,7 @@ void FNiagaraSystemUpdateContext::CommitUpdate()
 			PostWork.ExecuteIfBound(Comp);
 		}
 	}
-	ComponentsToReInit.Empty();
+	ComponentsToNotifySimDestroy.Empty();
 }
 
 void FNiagaraSystemUpdateContext::AddAll(bool bReInit)
diff --git a/Engine/Plugins/FX/Niagara/Source/Niagara/Private/NiagaraComponent.cpp b/Engine/Plugins/FX/Niagara/Source/Niagara/Private/NiagaraComponent.cpp
--- a/Engine/Plugins/FX/Niagara/Source/Niagara/Private/NiagaraComponent.cpp
+++ b/Engine/Plugins/FX/Niagara/Source/Niagara/Private/NiagaraComponent.cpp
@@ -4328,7 +4328,7 @@ void UNiagaraComponent::SetPreviewLODDistance(bool bInEnablePreviewLODDistance,
 	{
 		if (bEnablePreviewLODDistance)
 		{
-			SystemInstanceController->SetLODDistance(PreviewLODDistance, InPreviewLODDistance, true);
+			SystemInstanceController->SetLODDistance(PreviewLODDistance, InPreviewMaxDistance, true);
 		}
 		else
 		{
diff --git a/Engine/Plugins/FX/Niagara/Source/Niagara/Private/NiagaraDataChannelHandler.cpp b/Engine/Plugins/FX/Niagara/Source/Niagara/Private/NiagaraDataChannelHandler.cpp
--- a/Engine/Plugins/FX/Niagara/Source/Niagara/Private/NiagaraDataChannelHandler.cpp
+++ b/Engine/Plugins/FX/Niagara/Source/Niagara/Private/NiagaraDataChannelHandler.cpp
@@ -201,11 +201,11 @@ void UNiagaraDataChannelHandler::OnComputeDispatchInterfaceDestroyed(FNiagaraGpu
 void UNiagaraDataChannelHandler::VerifyAccessContext(FNDCAccessContextInst& AccessContext)
 {
 #if !UE_BUILD_SHIPPING
-	if(const UNiagaraDataChannel* Channel = DataChannel.Get())
+	if (const UNiagaraDataChannel* Channel = DataChannel.Get())
 	{
 		const UScriptStruct* Expected = Channel->GetAccessContextType().Get();
 		const UScriptStruct* Actual = AccessContext.GetScriptStruct();
-		if(Actual->IsChildOf(Expected) == false)
+		if (Actual == nullptr || Actual->IsChildOf(Expected) == false)
 		{
 			UE_LOGFMT(LogNiagara, Warning, "Attempting to find Niagara Data Channel with incorrect access context type.\nData Channel: {DataChannel}\nExpected Type:{Expected}\nActual:{ActualType}", 
 				GetNameSafe(DataChannel->GetAsset()), GetNameSafe(Expected), GetNameSafe(Actual));
@@ -218,7 +218,7 @@ FNiagaraDataChannelDataPtr UNiagaraDataChannelHandler::FindData(FNiagaraDataChan
 {
 	UE_LOGFMT(LogNiagara, Error, 
 	"Using legacy SearchParams API when accessing Niagara Data Channel that does not support the legacy path.\nPlease update access code to use the newer Access Context based access path.\nData Channel: {DataChannel}",
-	DataChannel->GetAsset() ? DataChannel->GetAsset()->GetPathName() : TEXT(""));
+	DataChannel.IsValid() && DataChannel->GetAsset() ? DataChannel->GetAsset()->GetPathName() : TEXT(""));
 	return nullptr;
 }
 
diff --git a/Engine/Plugins/FX/Niagara/Source/NiagaraEditor/Private/NiagaraScriptMergeManager.cpp b/Engine/Plugins/FX/Niagara/Source/NiagaraEditor/Private/NiagaraScriptMergeManager.cpp
--- a/Engine/Plugins/FX/Niagara/Source/NiagaraEditor/Private/NiagaraScriptMergeManager.cpp
+++ b/Engine/Plugins/FX/Niagara/Source/NiagaraEditor/Private/NiagaraScriptMergeManager.cpp
@@ -2021,6 +2021,7 @@ FNiagaraScriptMergeManager::FApplyDiffResults FNiagaraScriptMergeManager::ResetM
 		Results.bModifiedGraph = false;
 		Results.ErrorMessages.Add(FText::Format(LOCTEXT("ResetFailedBecauseOfScratchPadScripts", "Failed to reset input back to it's base value.  Its scratch pad scripts were out of sync.  Emitter: {0}  Input:{1}"),
 			EmitterPath, FText::FromString(InputName)));
+		return Results;
 	}
 	
 	// Remove items from the diff which are not relevant to this input.
