Fix issue writing values to an access context into bit field bool properties coming from full fat bool properties.

diff --git a/Engine/Plugins/FX/Niagara/Source/Niagara/Private/NiagaraDataChannelFunctionLibrary.cpp b/Engine/Plugins/FX/Niagara/Source/Niagara/Private/NiagaraDataChannelFunctionLibrary.cpp
--- a/Engine/Plugins/FX/Niagara/Source/Niagara/Private/NiagaraDataChannelFunctionLibrary.cpp
+++ b/Engine/Plugins/FX/Niagara/Source/Niagara/Private/NiagaraDataChannelFunctionLibrary.cpp
@@ -215,7 +215,27 @@ DEFINE_FUNCTION(UNiagaraDataChannelLibrary::execSetSinglePropertyInNDCAccessCont
 		P_NATIVE_BEGIN
 			if (FProperty* DestProperty = ActualContextStruct->FindPropertyByName(PropertyName))
 			{
-				DestProperty->CopyCompleteValue(DestProperty->ContainerPtrToValuePtr<uint8>(DestPtr), ValuePtr);
+				//Need special case handling for bools coming from bit fields.
+				if (FBoolProperty* BoolDestProperty = CastField<FBoolProperty>(DestProperty))
+				{
+					if(const FBoolProperty* BoolSrcProperty = CastField<FBoolProperty>(ValueProp))
+					{
+						bool SrcValue = BoolSrcProperty->GetPropertyValue(ValuePtr);
+						uint8* DstBool = BoolDestProperty->ContainerPtrToValuePtr<uint8>(DestPtr);
+						BoolDestProperty->SetPropertyValue(DstBool, SrcValue);
+					}
+					else
+					{					
+						FBlueprintExceptionInfo ExceptionInfo(
+							EBlueprintExceptionType::AbortExecution,
+							FText::Format(LOCTEXT("NDC_AccessContext_MismatchedPropertyTypes_Fmt", "Cannot write {0} to {1}. Property types do not match."), FText::FromName(PropertyName), DestProperty->GetDisplayNameText()));
+						FBlueprintCoreDelegates::ThrowScriptException(P_THIS, Stack, ExceptionInfo);
+					}
+				}
+				else
+				{
+					DestProperty->CopyCompleteValue(DestProperty->ContainerPtrToValuePtr<uint8>(DestPtr), ValuePtr);
+				}
 			}
 			else
 			{
