non editor display name fix.

diff --git a/Engine/Plugins/FX/Niagara/Source/Niagara/Private/NiagaraDataChannelFunctionLibrary.cpp b/Engine/Plugins/FX/Niagara/Source/Niagara/Private/NiagaraDataChannelFunctionLibrary.cpp
--- a/Engine/Plugins/FX/Niagara/Source/Niagara/Private/NiagaraDataChannelFunctionLibrary.cpp
+++ b/Engine/Plugins/FX/Niagara/Source/Niagara/Private/NiagaraDataChannelFunctionLibrary.cpp
@@ -228,7 +228,7 @@ DEFINE_FUNCTION(UNiagaraDataChannelLibrary::execSetSinglePropertyInNDCAccessCont
 					{					
 						FBlueprintExceptionInfo ExceptionInfo(
 							EBlueprintExceptionType::AbortExecution,
-							FText::Format(LOCTEXT("NDC_AccessContext_MismatchedPropertyTypes_Fmt", "Cannot write {0} to {1}. Property types do not match."), FText::FromName(PropertyName), DestProperty->GetDisplayNameText()));
+							FText::Format(LOCTEXT("NDC_AccessContext_MismatchedPropertyTypes_Fmt", "Cannot write {0} to {1}. Property types do not match."), FText::FromName(PropertyName), FText::FromString(DestProperty->GetName())));
 						FBlueprintCoreDelegates::ThrowScriptException(P_THIS, Stack, ExceptionInfo);
 					}
 				}
