- Fix leaving a lock dangling when setting an individual array element from BP

#tests In BP set array elements on a NS that are out of bounds multiple times

diff --git a/Engine/Plugins/FX/Niagara/Source/Niagara/Classes/NiagaraDataInterfaceArrayImpl.h b/Engine/Plugins/FX/Niagara/Source/Niagara/Classes/NiagaraDataInterfaceArrayImpl.h
--- a/Engine/Plugins/FX/Niagara/Source/Niagara/Classes/NiagaraDataInterfaceArrayImpl.h
+++ b/Engine/Plugins/FX/Niagara/Source/Niagara/Classes/NiagaraDataInterfaceArrayImpl.h
@@ -584,6 +584,7 @@ struct FNDIArrayProxyImpl : public INDIArrayProxyBase
 		ensure(PerInstanceData_GameThread.Num() <= 1);
 
 		BeginSetArrayFromBP(true);
+		ON_SCOPE_EXIT { EndSetArrayFromBP(); };
 
 		TArray<TArrayType>& ArrayRef = Owner->GetArrayReference();
 		if (!ArrayRef.IsValidIndex(Index))
@@ -595,8 +596,6 @@ struct FNDIArrayProxyImpl : public INDIArrayProxyBase
 			ArrayRef.AddDefaulted(Index + 1 - ArrayRef.Num());
 		}
 		FNDIArrayImplHelper<TArrayType>::CopyCpuToCpuMemory(ArrayRef.GetData() + Index, &Value, 1);
-
-		EndSetArrayFromBP();
 	}
 
 	template<typename TToArrayType>
