- Fix dangling RenderData when Systems are recached
- The render state is already marked for destruction, but if we simulate before flushing EOF updates the render data could be left dangling resulting in incorrect access during post tick

#tests Create a NS with a sprite + nanite + component renderer, disable / enable the sprite renderer

diff --git a/Engine/Plugins/FX/Niagara/Source/Niagara/Private/NiagaraCommon.cpp b/Engine/Plugins/FX/Niagara/Source/Niagara/Private/NiagaraCommon.cpp
--- a/Engine/Plugins/FX/Niagara/Source/Niagara/Private/NiagaraCommon.cpp
+++ b/Engine/Plugins/FX/Niagara/Source/Niagara/Private/NiagaraCommon.cpp
@@ -90,6 +90,20 @@ FNiagaraSystemUpdateContext::~FNiagaraSystemUpdateContext()
 
 void FNiagaraSystemUpdateContext::CommitUpdate()
 {
+	auto DestroyInstanceNotComponent =
+		[this](UNiagaraComponent* Component)
+		{
+			// Destroy the instance not the component and if we are going to recache the simulation make sure to kill the render state
+			// Not killing the render state can result in crashes if it's recached and we tick before flushing EOF updates
+			Component->DestroyInstanceNotComponent();
+			if (Component->IsRenderStateCreated() && SystemSimsToRecache.Contains(Component->GetAsset()))
+			{
+				UPrimitiveComponent* PrimitiveComponent = static_cast<UPrimitiveComponent*>(Component);
+				PrimitiveComponent->MarkRenderStateDirty();
+				PrimitiveComponent->DestroyRenderState_Concurrent();
+			}
+		};
+
 	if (SystemSimsToDestroy.Num() > 0)
 	{
 		for (UNiagaraSystem* NiagaraSystem : SystemSimsToDestroy)
@@ -112,7 +126,7 @@ void FNiagaraSystemUpdateContext::CommitUpdate()
 					FNiagaraSystemSimulationPtr SoloSimulation = SystemInstanceController->GetSoloSystemSimulation();
 					if (SoloSimulation.IsValid())
 					{
-						Comp->DestroyInstanceNotComponent();
+						DestroyInstanceNotComponent(Comp);
 						SoloSimulation->Destroy();
 					}
 				}
@@ -126,7 +140,7 @@ void FNiagaraSystemUpdateContext::CommitUpdate()
 	{
 		if (Comp)
 		{
-			Comp->DestroyInstanceNotComponent();
+			DestroyInstanceNotComponent(Comp);
 			PostWork.ExecuteIfBound(Comp);
 		}
 	}
diff --git a/Engine/Plugins/FX/Niagara/Source/Niagara/Private/NiagaraRendererProperties.cpp b/Engine/Plugins/FX/Niagara/Source/Niagara/Private/NiagaraRendererProperties.cpp
--- a/Engine/Plugins/FX/Niagara/Source/Niagara/Private/NiagaraRendererProperties.cpp
+++ b/Engine/Plugins/FX/Niagara/Source/Niagara/Private/NiagaraRendererProperties.cpp
@@ -1338,11 +1338,6 @@ void UNiagaraRendererProperties::SetIsEnabled(bool bInIsEnabled)
 	#endif
 
 		bIsEnabled = bInIsEnabled;
-
-		if (UNiagaraSystem* NiagaraSystem = GetTypedOuter<UNiagaraSystem>())
-		{
-			NiagaraSystem->ComputeRenderersDrawOrder();
-		}
 	}
 }
 
