- Fix using ternary vs if when reading from buffers to avoid OOB access

#rnx

diff --git a/Engine/Plugins/FX/Niagara/Shaders/Private/Stateless/NiagaraStatelessCommon.ush b/Engine/Plugins/FX/Niagara/Shaders/Private/Stateless/NiagaraStatelessCommon.ush
--- a/Engine/Plugins/FX/Niagara/Shaders/Private/Stateless/NiagaraStatelessCommon.ush
+++ b/Engine/Plugins/FX/Niagara/Shaders/Private/Stateless/NiagaraStatelessCommon.ush
@@ -122,12 +122,12 @@ float2 GetParameterBufferFloat2(int BaseOffset, int Element) { return asfloat(Co
 float3 GetParameterBufferFloat3(int BaseOffset, int Element) { return asfloat(Common_DynamicDataBuffer.Load3((BaseOffset + Element*3) * 4)); }
 float4 GetParameterBufferFloat4(int BaseOffset, int Element) { return asfloat(Common_DynamicDataBuffer.Load4((BaseOffset + Element*4) * 4)); }
 
-uint   GetParameterBufferUint  (int BaseOffset, int Element, uint  DefaultValue)  { return BaseOffset >= 0 ? asuint (Common_DynamicDataBuffer.Load ((BaseOffset + Element) * 4)) : DefaultValue; }
-int    GetParameterBufferInt   (int BaseOffset, int Element, int    DefaultValue) { return BaseOffset >= 0 ? asint  (Common_DynamicDataBuffer.Load ((BaseOffset + Element) * 4)) : DefaultValue; }
-float  GetParameterBufferFloat (int BaseOffset, int Element, float  DefaultValue) { return BaseOffset >= 0 ? asfloat(Common_DynamicDataBuffer.Load ((BaseOffset + Element) * 4)) : DefaultValue; }
-float2 GetParameterBufferFloat2(int BaseOffset, int Element, float2 DefaultValue) { return BaseOffset >= 0 ? asfloat(Common_DynamicDataBuffer.Load2((BaseOffset + Element*2) * 4)) : DefaultValue; }
-float3 GetParameterBufferFloat3(int BaseOffset, int Element, float3 DefaultValue) { return BaseOffset >= 0 ? asfloat(Common_DynamicDataBuffer.Load3((BaseOffset + Element*3) * 4)) : DefaultValue; }
-float4 GetParameterBufferFloat4(int BaseOffset, int Element, float4 DefaultValue) { return BaseOffset >= 0 ? asfloat(Common_DynamicDataBuffer.Load4((BaseOffset + Element*4) * 4)) : DefaultValue; }
+uint   GetParameterBufferUint  (int BaseOffset, int Element, uint  DefaultValue)  { if ( BaseOffset >= 0 ) { return asuint (Common_DynamicDataBuffer.Load ((BaseOffset + Element) * 4)); }   return DefaultValue; }
+int    GetParameterBufferInt   (int BaseOffset, int Element, int    DefaultValue) { if ( BaseOffset >= 0 ) { return asint  (Common_DynamicDataBuffer.Load ((BaseOffset + Element) * 4)); }   return DefaultValue; }
+float  GetParameterBufferFloat (int BaseOffset, int Element, float  DefaultValue) { if ( BaseOffset >= 0 ) { return asfloat(Common_DynamicDataBuffer.Load ((BaseOffset + Element) * 4)); }   return DefaultValue; }
+float2 GetParameterBufferFloat2(int BaseOffset, int Element, float2 DefaultValue) { if ( BaseOffset >= 0 ) { return asfloat(Common_DynamicDataBuffer.Load2((BaseOffset + Element*2) * 4)); } return DefaultValue; }
+float3 GetParameterBufferFloat3(int BaseOffset, int Element, float3 DefaultValue) { if ( BaseOffset >= 0 ) { return asfloat(Common_DynamicDataBuffer.Load3((BaseOffset + Element*3) * 4)); } return DefaultValue; }
+float4 GetParameterBufferFloat4(int BaseOffset, int Element, float4 DefaultValue) { if ( BaseOffset >= 0 ) { return asfloat(Common_DynamicDataBuffer.Load4((BaseOffset + Element*4) * 4)); } return DefaultValue; }
 	
 float  StatelessLerp(float  lhs, float  rhs, float  u, bool bUniform) { return lerp(lhs, rhs, u); }
 float2 StatelessLerp(float2 lhs, float2 rhs, float2 u, bool bUniform) { return lerp(lhs, rhs, bUniform ? u.xx : u); }
