- Fix for bIsSeeking not behaving the same post MultiTick refactor

#tests Create a fountain, seek forward and observe the system flickering

diff --git a/Engine/Plugins/FX/Niagara/Source/Niagara/Private/NiagaraComponent.cpp b/Engine/Plugins/FX/Niagara/Source/Niagara/Private/NiagaraComponent.cpp
--- a/Engine/Plugins/FX/Niagara/Source/Niagara/Private/NiagaraComponent.cpp
+++ b/Engine/Plugins/FX/Niagara/Source/Niagara/Private/NiagaraComponent.cpp
@@ -980,7 +980,8 @@ void UNiagaraComponent::TickComponent(float DeltaSeconds, enum ELevelTick TickTy
 		else if(AgeUpdateMode == ENiagaraAgeUpdateMode::DesiredAge)
 		{
 			float AgeDiff = FMath::Max(DesiredAge, 0.0f) - SystemInstanceController->GetAge();
-			int32 TicksToProcess = 0;
+			bIsSeeking = false;
+
 			if (FMath::Abs(AgeDiff) >= UE_KINDA_SMALL_NUMBER)
 			{
 				if (AgeDiff < 0.0f)
@@ -998,10 +999,10 @@ void UNiagaraComponent::TickComponent(float DeltaSeconds, enum ELevelTick TickTy
 					if (bLockDesiredAgeDeltaTimeToSeekDelta || AgeDiff > TickDelta)
 					{
 						// If we're locking the delta time to the seek delta, or we need to seek more than a frame, tick the simulation by the seek delta.
-						TicksToProcess = FMath::FloorToInt(AgeDiff / TickDelta);
+						const int32 TicksToProcess = FMath::FloorToInt(AgeDiff / TickDelta);
 						CompletionEvent = TicksToProcess == 1 ? CompletionEvent : nullptr;
 
-						SystemInstanceController->MultiTick(TicksToProcess, TickDelta, MaxSimTime, CompletionEvent);
+						bIsSeeking = SystemInstanceController->MultiTick(TicksToProcess, TickDelta, MaxSimTime, CompletionEvent);
 					}
 					else
 					{
@@ -1013,11 +1014,6 @@ void UNiagaraComponent::TickComponent(float DeltaSeconds, enum ELevelTick TickTy
 					}
 				}
 			}
-
-			if (TicksToProcess == 0)
-			{
-				bIsSeeking = false;
-			}
 		}
 		else if (AgeUpdateMode == ENiagaraAgeUpdateMode::DesiredAgeNoSeek)
 		{
diff --git a/Engine/Plugins/FX/Niagara/Source/Niagara/Private/NiagaraSystemInstance.cpp b/Engine/Plugins/FX/Niagara/Source/Niagara/Private/NiagaraSystemInstance.cpp
--- a/Engine/Plugins/FX/Niagara/Source/Niagara/Private/NiagaraSystemInstance.cpp
+++ b/Engine/Plugins/FX/Niagara/Source/Niagara/Private/NiagaraSystemInstance.cpp
@@ -2347,7 +2347,7 @@ void FNiagaraSystemInstance::ManualTick(float DeltaSeconds, const FGraphEventRef
 	SystemSim->Tick_GameThread(DeltaSeconds, MyCompletionGraphEvent);
 }
 
-void FNiagaraSystemInstance::MultiTick(int32 TicksToProcess, float DeltaSecondsPerTick, float MaxSimTime, const FGraphEventRef& MyCompletionGraphEvent)
+bool FNiagaraSystemInstance::MultiTick(int32 TicksToProcess, float DeltaSecondsPerTick, float MaxSimTime, const FGraphEventRef& MyCompletionGraphEvent)
 {
 	SCOPE_CYCLE_COUNTER(STAT_NiagaraOverview_GT);
 	SCOPE_CYCLE_COUNTER(STAT_NiagaraSystemInst_ComponentTickGT);
@@ -2356,7 +2356,7 @@ void FNiagaraSystemInstance::MultiTick(int32 TicksToProcess, float DeltaSecondsP
 
 	if (IsDisabled())
 	{
-		return;
+		return false;
 	}
 
 	TSharedPtr<FNiagaraSystemSimulation> SystemSim = GetSystemSimulation();
@@ -2364,7 +2364,7 @@ void FNiagaraSystemInstance::MultiTick(int32 TicksToProcess, float DeltaSecondsP
 	check(IsInGameThread());
 	check(bSolo);
 
-	SystemSim->MultiTick_GameThread(TicksToProcess, DeltaSecondsPerTick, MaxSimTime, MyCompletionGraphEvent);
+	return SystemSim->MultiTick_GameThread(TicksToProcess, DeltaSecondsPerTick, MaxSimTime, MyCompletionGraphEvent);
 }
 
 void FNiagaraSystemInstance::SimCacheTick_GameThread(UNiagaraSimCache* SimCache, float DesiredAge, float DeltaSeconds, const FGraphEventRef& MyCompletionGraphEvent)
diff --git a/Engine/Plugins/FX/Niagara/Source/Niagara/Private/NiagaraSystemSimulation.cpp b/Engine/Plugins/FX/Niagara/Source/Niagara/Private/NiagaraSystemSimulation.cpp
--- a/Engine/Plugins/FX/Niagara/Source/Niagara/Private/NiagaraSystemSimulation.cpp
+++ b/Engine/Plugins/FX/Niagara/Source/Niagara/Private/NiagaraSystemSimulation.cpp
@@ -1140,11 +1140,11 @@ void FNiagaraSystemSimulation::Tick_GameThread(float DeltaSeconds, const FGraphE
 	}
 }
 
-void FNiagaraSystemSimulation::MultiTick_GameThread(int32 TicksToProcess, float DeltaSecondsPerTick, float MaxSimTime, const FGraphEventRef& MyCompletionGraphEvent)
+bool FNiagaraSystemSimulation::MultiTick_GameThread(int32 TicksToProcess, float DeltaSecondsPerTick, float MaxSimTime, const FGraphEventRef& MyCompletionGraphEvent)
 {
 	if (DeltaSecondsPerTick <= 0 || TicksToProcess <= 0)
 	{
-		return;
+		return false;
 	}
 
 	TickInfo.UsesFixedTick = true;
@@ -1180,6 +1180,8 @@ void FNiagaraSystemSimulation::MultiTick_GameThread(int32 TicksToProcess, float
 			}
 		}
 	}
+
+	return TickInfo.TickNumber != (TicksToProcess - 1);
 }
 
 void FNiagaraSystemSimulation::Tick_GameThread_Internal(float DeltaSeconds, const FGraphEventRef& MyCompletionGraphEvent)
diff --git a/Engine/Plugins/FX/Niagara/Source/Niagara/Public/NiagaraSystemInstance.h b/Engine/Plugins/FX/Niagara/Source/Niagara/Public/NiagaraSystemInstance.h
--- a/Engine/Plugins/FX/Niagara/Source/Niagara/Public/NiagaraSystemInstance.h
+++ b/Engine/Plugins/FX/Niagara/Source/Niagara/Public/NiagaraSystemInstance.h
@@ -200,10 +200,11 @@ public:
 
 	/** Manually advances the simulation by multiple ticks. Ticks on the game thread.
 	 *  If MaxSimTime >= 0, then unprocessed ticks will be discarded after the time is elapsed to prevent stalling the game thread.
+	 *  Returns true if we hit MaxSimTime before processing all the requested ticks
 	 *
 	 *  Note that this overrides a system's fixed tick delta time (if it has one) and uses DeltaSecondsPerTick instead. 
 	 */
-	NIAGARA_API void MultiTick(int32 TicksToProcess, float DeltaSecondsPerTick, float MaxSimTime, const FGraphEventRef& MyCompletionGraphEvent);
+	NIAGARA_API bool MultiTick(int32 TicksToProcess, float DeltaSecondsPerTick, float MaxSimTime, const FGraphEventRef& MyCompletionGraphEvent);
 
 	/** Ticks the system using the a SimCache. */
 	void SimCacheTick_GameThread(UNiagaraSimCache* SimCache, float DesiredAge, float DeltaSeconds, const FGraphEventRef& MyCompletionGraphEvent);
diff --git a/Engine/Plugins/FX/Niagara/Source/Niagara/Public/NiagaraSystemSimulation.h b/Engine/Plugins/FX/Niagara/Source/Niagara/Public/NiagaraSystemSimulation.h
--- a/Engine/Plugins/FX/Niagara/Source/Niagara/Public/NiagaraSystemSimulation.h
+++ b/Engine/Plugins/FX/Niagara/Source/Niagara/Public/NiagaraSystemSimulation.h
@@ -265,7 +265,7 @@ public:
 	void Tick_Concurrent(FNiagaraSystemSimulationTickContext& Context);
 
 	/** Processes multiple ticks at once on the game thread. */
-	void MultiTick_GameThread(int32 TicksToProcess, float DeltaSecondsPerTick, float MaxSimTime, const FGraphEventRef& MyCompletionGraphEvent);
+	bool MultiTick_GameThread(int32 TicksToProcess, float DeltaSecondsPerTick, float MaxSimTime, const FGraphEventRef& MyCompletionGraphEvent);
 
 	/** Update TickGroups for pending instances and execute tick group promotions. */
 	void UpdateTickGroups_GameThread();
