- Fix crash from dangling compilations when changing version and emitter merging takes place

diff --git a/Engine/Plugins/FX/Niagara/Source/Niagara/Private/NiagaraSystem.cpp b/Engine/Plugins/FX/Niagara/Source/Niagara/Private/NiagaraSystem.cpp
--- a/Engine/Plugins/FX/Niagara/Source/Niagara/Private/NiagaraSystem.cpp
+++ b/Engine/Plugins/FX/Niagara/Source/Niagara/Private/NiagaraSystem.cpp
@@ -3828,6 +3828,10 @@ void UNiagaraSystem::InitSystemCompiledData()
 
 void UNiagaraSystem::ResetToEmptySystem()
 {
+#if WITH_EDITORONLY_DATA
+	KillAllActiveCompilations();
+#endif
+
 	// in order to handle resetting transient systems that can be created (FNiagaraScriptMergeManager::UpdateModuleVersions) we make sure
 	// to mark the system as having been fully loaded to bypass any future work that might be done (generating compile requests)
 	bFullyLoaded = true;
diff --git a/Engine/Plugins/FX/Niagara/Source/NiagaraEditor/Private/NiagaraScriptMergeManager.cpp b/Engine/Plugins/FX/Niagara/Source/NiagaraEditor/Private/NiagaraScriptMergeManager.cpp
--- a/Engine/Plugins/FX/Niagara/Source/NiagaraEditor/Private/NiagaraScriptMergeManager.cpp
+++ b/Engine/Plugins/FX/Niagara/Source/NiagaraEditor/Private/NiagaraScriptMergeManager.cpp
@@ -1433,7 +1433,7 @@ void FNiagaraScriptMergeManager::UpdateModuleVersions(const FVersionedNiagaraEmi
 	{
 		TSharedPtr<FNiagaraSystemViewModel> SystemViewModel;
 		
-		System = NewObject<UNiagaraSystem>(GetTransientPackage(), NAME_None, RF_Transient | RF_Standalone);
+		System = NewObject<UNiagaraSystem>(GetTransientPackage(), NAME_None, RF_Transient);
 		UNiagaraSystemFactoryNew::InitializeSystem(System, true);
 		SystemViewModel = MakeShared<FNiagaraSystemViewModel>();
 		FNiagaraSystemViewModelOptions SystemOptions;
@@ -1507,6 +1507,7 @@ void FNiagaraScriptMergeManager::UpdateModuleVersions(const FVersionedNiagaraEmi
 	if (System)
 	{
 		System->ResetToEmptySystem();
+		System->MarkAsGarbage();
 	}
 }
 
