Clean up to the NiagaraParameterCollection editor workflow
-changes (editing values, adding/removing/renaming parameters) throughout an MPC -> NPC -> NPCI chain all update properly
-parameters that are provided by an MPC are marked as read-only

diff --git a/Engine/Plugins/FX/Niagara/Source/Niagara/Classes/NiagaraParameterCollection.h b/Engine/Plugins/FX/Niagara/Source/Niagara/Classes/NiagaraParameterCollection.h
--- a/Engine/Plugins/FX/Niagara/Source/Niagara/Classes/NiagaraParameterCollection.h
+++ b/Engine/Plugins/FX/Niagara/Source/Niagara/Classes/NiagaraParameterCollection.h
@@ -40,6 +40,7 @@ public:
 	NIAGARA_API void RenameParameter(const FNiagaraVariable& Parameter, FName NewName);
 	NIAGARA_API void Empty();
 	NIAGARA_API void GetParameters(TArray<FNiagaraVariable>& OutParameters);
+	NIAGARA_API bool IsExternallyDriven(const FNiagaraVariable& Parameter) const;
 
 	NIAGARA_API void Tick(UWorld* World);
 
@@ -57,12 +58,22 @@ public:
 
 	/** Synchronizes this instance with any changes with it's parent collection. */
 	NIAGARA_API void SyncWithCollection();
+	NIAGARA_API void SyncWithCollectionParameters(TConstArrayView<FName> ParametersToRefresh);
 
 	NIAGARA_API void Bind(UWorld* World);
+	NIAGARA_API void Unbind(UWorld* World);
+
+#if WITH_EDITOR
+	DECLARE_MULTICAST_DELEGATE_OneParam(FOnParameterValuesChanged, TConstArrayView<FName>)
+	FOnParameterValuesChanged& ParameterValuesChanged() { return OnParameterValuesChanged; }
+#endif
 
 private:
 	NIAGARA_API void RefreshSourceParameters(UWorld* World, const TArray<TPair<FName, float>>& ScalarParameters, const TArray<TPair<FName, FLinearColor>>& VectorParameters);
 	NIAGARA_API bool EnsureNotBoundToMaterialParameterCollection(FName InVariableName, FString CallingFunction) const;
+	
+	void InitDelegates(UNiagaraParameterCollection* ParentCollection);
+	void ResetDelegates(UNiagaraParameterCollection* ParentCollection);
 
 	UPROPERTY()
 	FNiagaraParameterStore ParameterStorage;
@@ -74,13 +85,28 @@ private:
 	TArray<TPair<FName, float>> DirtyScalarParameters;
 	TArray<TPair<FName, FLinearColor>> DirtyVectorParameters;
 
+	FDelegateHandle ScalarParamUpdateDelegate;
+	FDelegateHandle VectorParamUpdateDelegate;
+#if WITH_EDITOR
+	FDelegateHandle ParentCollectionChangedDelegate;
+	FDelegateHandle MaterialCollectionChangedDelegate;
+	FDelegateHandle MaterialParametersChangedHandle;
+	FDelegateHandle ParameterValuesEditedHandle;
+	std::atomic<bool> ParametersNeedReset;
+
+	void DirtyParameters() { ParametersNeedReset = true; }
+
+	FOnParameterValuesChanged OnParameterValuesChanged;
+#endif
+
 	//TODO: These overrides should be settable per platform.
 	//UPROPERTY()
 	//TMap<FString, FNiagaraParameterStore>
 
 	//~UObject interface
 #if WITH_EDITORONLY_DATA
-	NIAGARA_API virtual void PostEditChangeProperty(struct FPropertyChangedEvent& PropertyChangedEvent)override;
+	NIAGARA_API virtual void PreEditChange(FProperty* PropertyAboutToChange) override;
+	NIAGARA_API virtual void PostEditChangeProperty(struct FPropertyChangedEvent& PropertyChangedEvent) override;
 #endif
 	//~UObject interface
 public:
@@ -149,6 +175,7 @@ public:
 
 	//~UObject interface
 #if WITH_EDITORONLY_DATA
+	NIAGARA_API virtual void PreEditChange(FProperty* PropertyAboutToChange) override;
 	NIAGARA_API virtual void PostEditChangeProperty(struct FPropertyChangedEvent& PropertyChangedEvent)override;
 	NIAGARA_API virtual void PostDuplicate(EDuplicateMode::Type DuplicateMode) override;
 #endif
@@ -212,16 +239,24 @@ public:
 	FName GetNamespace()const { return Namespace; }
 
 #if WITH_EDITOR
-	DECLARE_MULTICAST_DELEGATE(FOnChanged);
+	DECLARE_MULTICAST_DELEGATE(FOnCollectionChanged);
+	FOnCollectionChanged OnCollectionChangedDelegate;
+
+	DECLARE_MULTICAST_DELEGATE_OneParam(FOnParameterValuesChanged, TConstArrayView<FName>)
+
+	// invoked when the collection's material source is updated with new parameter values
+	FOnParameterValuesChanged OnMaterialParametersChanged;
 
-	FOnChanged OnChangedDelegate;
+	// invoked when the collection's default values are updated from the NPC editor
+	FOnParameterValuesChanged OnParameterValuesEdited;
 #endif
 
 protected:
 	
 #if WITH_EDITORONLY_DATA
 	NIAGARA_API void MakeNamespaceNameUnique();
-	NIAGARA_API void AddDefaultSourceParameters();
+	NIAGARA_API void NamespaceChanged();
+	NIAGARA_API void RefreshSourceMaterialCollection();
 #endif
 
 	/** Namespace for this parameter collection. Is enforced to be unique across all parameter collections. */
@@ -246,4 +281,13 @@ protected:
 
 	// transient variable holding the full namespace (NPC.<namespace>.)
 	FName FullNamespace;
-};
+
+#if WITH_EDITOR
+	FDelegateHandle OnMaterialCollectionChanged;
+
+	// in order to manage updating the NPC if a source MPC deletes a parameter while in the editor we need to keep track of which
+	// parameters are actually 'owned' by the MPC.  This is only relevant for the UI.
+	TSet<FName> ParametersOwnedBySourceCollection;
+#endif
+
+};
\ No newline at end of file
diff --git a/Engine/Plugins/FX/Niagara/Source/Niagara/Private/NiagaraParameterCollection.cpp b/Engine/Plugins/FX/Niagara/Source/Niagara/Private/NiagaraParameterCollection.cpp
--- a/Engine/Plugins/FX/Niagara/Source/Niagara/Private/NiagaraParameterCollection.cpp
+++ b/Engine/Plugins/FX/Niagara/Source/Niagara/Private/NiagaraParameterCollection.cpp
@@ -1,15 +1,17 @@
 // Copyright Epic Games, Inc. All Rights Reserved.
 
 #include "NiagaraParameterCollection.h"
+
+#include "Algo/RemoveIf.h"
 #include "AssetRegistry/AssetData.h"
+#include "AssetRegistry/AssetRegistryModule.h"
 #include "Engine/World.h"
-#include "Modules/ModuleManager.h"
-#include "NiagaraDataInterface.h"
 #include "Materials/MaterialParameterCollection.h"
 #include "Materials/MaterialParameterCollectionInstance.h"
 #include "Misc/SecureHash.h"
+#include "Modules/ModuleManager.h"
 #include "NiagaraConstants.h"
-#include "AssetRegistry/AssetRegistryModule.h"
+#include "NiagaraDataInterface.h"
 
 #include UE_INLINE_GENERATED_CPP_BY_NAME(NiagaraParameterCollection)
 
@@ -17,14 +19,15 @@
 
 UNiagaraParameterCollectionInstance::UNiagaraParameterCollectionInstance(const FObjectInitializer& ObjectInitializer)
 	: Super(ObjectInitializer)
+#if WITH_EDITOR
+	, ParametersNeedReset(false)
+#endif
 {
 	ParameterStorage.SetOwner(this);
-	//Bind(ParameterStorage);
 }
 
 UNiagaraParameterCollectionInstance::~UNiagaraParameterCollectionInstance()
 {
-	//Unbind(ParameterStorage);
 }
 
 void UNiagaraParameterCollectionInstance::PostLoad()
@@ -39,6 +42,7 @@ void UNiagaraParameterCollectionInstance::PostLoad()
 	if (Collection && (Collection->GetDefaultInstance() != this))
 	{
 		Collection->ConditionalPostLoad();
+		InitDelegates(Collection);
 	}
 
 	//Ensure we're synced up with our collection. TODO: Do conditionally via a version number/guid?
@@ -47,7 +51,12 @@ void UNiagaraParameterCollectionInstance::PostLoad()
 
 void UNiagaraParameterCollectionInstance::SetParent(UNiagaraParameterCollection* InParent)
 {
+	ResetDelegates(Collection);
+
 	Collection = InParent;
+
+	InitDelegates(Collection);
+
 	SyncWithCollection();
 }
 
@@ -85,21 +94,49 @@ void UNiagaraParameterCollectionInstance::GetParameters(TArray<FNiagaraVariable>
 	ParameterStorage.GetParameters(OutParameters);
 }
 
+static void BuildCollectionValues(
+	const UMaterialParameterCollection& OwnerCollection,
+	const UMaterialParameterCollectionInstance& OwnerInstance,
+	TArray<TPair<FName, float>>& ScalarParameters,
+	TArray<TPair<FName, FLinearColor>>& VectorParameters)
+{
+	// initialize the source parameters from the material collection instance
+	const int32 ScalarParameterCount = OwnerCollection.ScalarParameters.Num();
+	ScalarParameters.AddUninitialized(ScalarParameterCount);
+	for (int32 ScalarIt = 0; ScalarIt < ScalarParameterCount; ++ScalarIt)
+	{
+		ScalarParameters[ScalarIt].Key = OwnerCollection.ScalarParameters[ScalarIt].ParameterName;
+		OwnerInstance.GetScalarParameterValue(OwnerCollection.ScalarParameters[ScalarIt], ScalarParameters[ScalarIt].Value);
+	}
+
+	const int32 VectorParameterCount = OwnerCollection.VectorParameters.Num();
+	VectorParameters.AddUninitialized(VectorParameterCount);
+	for (int32 VectorIt = 0; VectorIt < VectorParameterCount; ++VectorIt)
+	{
+		VectorParameters[VectorIt].Key = OwnerCollection.VectorParameters[VectorIt].ParameterName;
+		OwnerInstance.GetVectorParameterValue(OwnerCollection.VectorParameters[VectorIt], VectorParameters[VectorIt].Value);
+	}
+}
+
 void UNiagaraParameterCollectionInstance::Bind(UWorld* World)
 {
-	if (const UMaterialParameterCollection* SourceCollection = Collection ? Collection->GetSourceCollection() : nullptr)
+	if (const UMaterialParameterCollection* MaterialCollection = Collection ? Collection->GetSourceCollection() : nullptr)
 	{
-		SourceMaterialCollectionInstance = World->GetParameterCollectionInstance(SourceCollection);
+#if WITH_EDITOR
+		MaterialCollectionChangedDelegate = MaterialCollection->OnCollectionChanged().AddUObject(this, &UNiagaraParameterCollectionInstance::DirtyParameters);
+#endif
+
+		SourceMaterialCollectionInstance = World->GetParameterCollectionInstance(MaterialCollection);
 
 		if (SourceMaterialCollectionInstance)
 		{
-			SourceMaterialCollectionInstance->OnScalarParameterUpdated().AddLambda([this](UMaterialParameterCollectionInstance::ScalarParameterUpdate DirtyParameter)
+			ScalarParamUpdateDelegate = SourceMaterialCollectionInstance->OnScalarParameterUpdated().AddLambda([this](UMaterialParameterCollectionInstance::ScalarParameterUpdate DirtyParameter)
 			{
 				FRWScopeLock WriteLock(DirtyParameterLock, SLT_Write);
 				DirtyScalarParameters.Emplace(DirtyParameter);
 			});
 
-			SourceMaterialCollectionInstance->OnVectorParameterUpdated().AddLambda([this](UMaterialParameterCollectionInstance::VectorParameterUpdate DirtyParameter)
+			VectorParamUpdateDelegate = SourceMaterialCollectionInstance->OnVectorParameterUpdated().AddLambda([this](UMaterialParameterCollectionInstance::VectorParameterUpdate DirtyParameter)
 			{
 				FRWScopeLock WriteLock(DirtyParameterLock, SLT_Write);
 				DirtyVectorParameters.Emplace(DirtyParameter);
@@ -107,40 +144,64 @@ void UNiagaraParameterCollectionInstance::Bind(UWorld* World)
 
 			// initialize the source parameters from the material collection instance
 			TArray<TPair<FName, float>> ScalarParameters;
-			const int32 ScalarParameterCount = SourceCollection->ScalarParameters.Num();
-			ScalarParameters.AddUninitialized(ScalarParameterCount);
-			for (int32 ScalarIt = 0; ScalarIt < ScalarParameterCount; ++ScalarIt)
-			{
-				ScalarParameters[ScalarIt].Key = SourceCollection->ScalarParameters[ScalarIt].ParameterName;
-				SourceMaterialCollectionInstance->GetScalarParameterValue(SourceCollection->ScalarParameters[ScalarIt], ScalarParameters[ScalarIt].Value);
-			}
-
 			TArray<TPair<FName, FLinearColor>> VectorParameters;
-			const int32 VectorParameterCount = SourceCollection->VectorParameters.Num();
-			VectorParameters.AddUninitialized(VectorParameterCount);
-			for (int32 VectorIt = 0; VectorIt < VectorParameterCount; ++VectorIt)
-			{
-				VectorParameters[VectorIt].Key = SourceCollection->VectorParameters[VectorIt].ParameterName;
-				SourceMaterialCollectionInstance->GetVectorParameterValue(SourceCollection->VectorParameters[VectorIt], VectorParameters[VectorIt].Value);
-			}
 
+			BuildCollectionValues(*MaterialCollection, *SourceMaterialCollectionInstance, ScalarParameters, VectorParameters);
 			RefreshSourceParameters(World, ScalarParameters, VectorParameters);
 		}
 	}
 }
 
+bool UNiagaraParameterCollectionInstance::IsExternallyDriven(const FNiagaraVariable& Parameter) const
+{
+	if (Collection)
+	{
+		if (const UMaterialParameterCollection* MaterialCollection = Collection->GetSourceCollection())
+		{
+			FName FriendlyName = Collection->FriendlyNameFromParameterName(Parameter.GetName());
+
+			return MaterialCollection->ScalarParameters.ContainsByPredicate([&](const FCollectionParameterBase& InParam) { return InParam.ParameterName == FriendlyName; })
+				|| MaterialCollection->VectorParameters.ContainsByPredicate([&](const FCollectionParameterBase& InParam) { return InParam.ParameterName == FriendlyName; });
+		}
+	}
+
+	return false;
+}
+
+void UNiagaraParameterCollectionInstance::Unbind(UWorld* World)
+{
+	if (const UMaterialParameterCollection* SourceCollection = Collection ? Collection->GetSourceCollection() : nullptr)
+	{
+#if WITH_EDITOR
+		SourceCollection->OnCollectionChanged().Remove(MaterialCollectionChangedDelegate);
+#endif
+
+		if (UMaterialParameterCollectionInstance* SourceInstance = World->GetParameterCollectionInstance(SourceCollection))
+		{
+			SourceInstance->OnScalarParameterUpdated().Remove(ScalarParamUpdateDelegate);
+			SourceInstance->OnVectorParameterUpdated().Remove(VectorParamUpdateDelegate);
+		}
+	}
+
+#if WITH_EDITOR
+	MaterialCollectionChangedDelegate.Reset();
+#endif
+	ScalarParamUpdateDelegate.Reset();
+	VectorParamUpdateDelegate.Reset();
+}
+
 void UNiagaraParameterCollectionInstance::RefreshSourceParameters(
 	UWorld* World,
 	const TArray<TPair<FName, float>>& ScalarParameters,
 	const TArray<TPair<FName, FLinearColor>>& VectorParameters)
 {
 	// if the NPC uses any MPC as sources, the make those bindings now
-	if (const UMaterialParameterCollection* SourceCollection = Collection ? Collection->GetSourceCollection() : nullptr)
+	if (const UMaterialParameterCollection* MaterialCollection = Collection ? Collection->GetSourceCollection() : nullptr)
 	{
 		// find the appropriate Instance
-		if (!SourceMaterialCollectionInstance || SourceMaterialCollectionInstance->GetCollection() != SourceCollection)
+		if (!SourceMaterialCollectionInstance || SourceMaterialCollectionInstance->GetCollection() != MaterialCollection)
 		{
-			SourceMaterialCollectionInstance = World->GetParameterCollectionInstance(SourceCollection);
+			SourceMaterialCollectionInstance = World->GetParameterCollectionInstance(MaterialCollection);
 		}
 
 		if (SourceMaterialCollectionInstance)
@@ -180,11 +241,70 @@ void UNiagaraParameterCollectionInstance::RefreshSourceParameters(
 	}
 }
 
+void UNiagaraParameterCollectionInstance::InitDelegates(UNiagaraParameterCollection* ParentCollection)
+{
+#if WITH_EDITOR
+	if (ParentCollection)
+	{
+		ParentCollectionChangedDelegate = ParentCollection->OnCollectionChangedDelegate.AddUObject(this, &UNiagaraParameterCollectionInstance::SyncWithCollection);
+		MaterialParametersChangedHandle = ParentCollection->OnMaterialParametersChanged.AddUObject(this, &UNiagaraParameterCollectionInstance::SyncWithCollectionParameters);
+		if (!IsDefaultInstance())
+		{
+			ParameterValuesEditedHandle = ParentCollection->OnParameterValuesEdited.AddUObject(this, &UNiagaraParameterCollectionInstance::SyncWithCollectionParameters);
+		}
+	}
+#endif
+
+}
+
+void UNiagaraParameterCollectionInstance::ResetDelegates(UNiagaraParameterCollection* ParentCollection)
+{
+#if WITH_EDITOR
+	if (ParentCollection)
+	{
+		if (ParentCollectionChangedDelegate.IsValid())
+		{
+			ParentCollection->OnCollectionChangedDelegate.Remove(ParentCollectionChangedDelegate);
+		}
+
+		if (MaterialParametersChangedHandle.IsValid())
+		{
+			ParentCollection->OnMaterialParametersChanged.Remove(MaterialParametersChangedHandle);
+		}
+
+		if (ParameterValuesEditedHandle.IsValid())
+		{
+			ParentCollection->OnParameterValuesEdited.Remove(ParameterValuesEditedHandle);
+		}
+	}
+#endif
+}
+
 void UNiagaraParameterCollectionInstance::Tick(UWorld* World)
 {
 	{
 		FRWScopeLock WriteLock(DirtyParameterLock, SLT_Write);
 
+#if WITH_EDITOR
+		if (ParametersNeedReset)
+		{
+			if (const UMaterialParameterCollection* SourceCollection = Collection ? Collection->GetSourceCollection() : nullptr)
+			{
+				if (UMaterialParameterCollectionInstance* SourceInstance = World->GetParameterCollectionInstance(SourceCollection))
+				{
+					// initialize the source parameters from the material collection instance
+					TArray<TPair<FName, float>> ScalarParameters;
+					TArray<TPair<FName, FLinearColor>> VectorParameters;
+
+					BuildCollectionValues(*SourceCollection, *SourceInstance, ScalarParameters, VectorParameters);
+					RefreshSourceParameters(World, ScalarParameters, VectorParameters);
+				}
+			}
+
+			ParametersNeedReset = false;
+		}
+		else
+#endif
 		if (DirtyScalarParameters.Num() || DirtyVectorParameters.Num())
 		{
 			RefreshSourceParameters(World, DirtyScalarParameters, DirtyVectorParameters);
@@ -208,7 +328,9 @@ void UNiagaraParameterCollectionInstance::SyncWithCollection()
 		return;
 	}
 
-	for (FNiagaraVariable& Param : Collection->GetParameters())
+	const TArray<FNiagaraVariable>& CollectionParameters = Collection->GetParameters();
+
+	for (const FNiagaraVariable& Param : CollectionParameters)
 	{
 		int32 Offset = OldStore.IndexOf(Param);
 		if (Offset != INDEX_NONE && OverridesParameter(Param))
@@ -234,26 +356,86 @@ void UNiagaraParameterCollectionInstance::SyncWithCollection()
 			//If the parameter did not exist in the old store or we don't override this parameter, sync it up to the parent collection.
 			FNiagaraParameterStore& DefaultStore = Collection->GetDefaultInstance()->GetParameterStore();
 			Offset = DefaultStore.IndexOf(Param);
-			check(Offset != INDEX_NONE);
 
-			int32 ParameterStorageOffset = INDEX_NONE;
-			ParameterStorage.AddParameter(Param, false, true, &ParameterStorageOffset);
+			// if the parameter was not in the parent collection than we can delete the orphaned parameter
+			if (Offset != INDEX_NONE)
+			{
+				int32 ParameterStorageOffset = INDEX_NONE;
+				ParameterStorage.AddParameter(Param, false, true, &ParameterStorageOffset);
+				if (Param.IsDataInterface())
+				{
+					ParameterStorage.SetDataInterface(CastChecked<UNiagaraDataInterface>(StaticDuplicateObject(DefaultStore.GetDataInterface(Offset), this)), Param);
+				}
+				else if (Param.IsUObject())
+				{
+					ParameterStorage.SetUObject(DefaultStore.GetUObject(Offset), Param);
+				}
+				else
+				{
+					ParameterStorage.SetParameterData(DefaultStore.GetParameterData(Offset, Param.GetType()), ParameterStorageOffset, Param.GetSizeInBytes());
+				}
+			}
+		}
+	}
+
+	// clean up OverridenParameters as well
+	OverridenParameters.SetNum(Algo::RemoveIf(OverridenParameters, [&CollectionParameters](const FNiagaraVariable& OverrideParameter) -> bool
+	{
+		return !CollectionParameters.Contains(OverrideParameter);
+	}));
+
+	ParameterStorage.Rebind();
+}
+
+void UNiagaraParameterCollectionInstance::SyncWithCollectionParameters(TConstArrayView<FName> ParametersToRefresh)
+{
+	if (ParametersToRefresh.IsEmpty())
+	{
+		return;
+	}
+
+	for (FNiagaraVariable& Param : Collection->GetParameters())
+	{
+		if (!ParametersToRefresh.Contains(Param.GetName()))
+		{
+			continue;
+		}
+
+		// only worry about parameters we're not overriding
+		if (OverridesParameter(Param))
+		{
+			continue;
+		}
+
+		const FNiagaraParameterStore& DefaultStore = Collection->GetDefaultInstance()->GetParameterStore();
+		const int32 CollectionIndex = DefaultStore.IndexOf(Param);
+		const int32 LocalIndex = ParameterStorage.IndexOf(Param);
+
+		// if the parameter was not in the parent collection than we can delete the orphaned parameter
+		if (CollectionIndex != INDEX_NONE && LocalIndex != INDEX_NONE)
+		{
 			if (Param.IsDataInterface())
 			{
-				ParameterStorage.SetDataInterface(CastChecked<UNiagaraDataInterface>(StaticDuplicateObject(DefaultStore.GetDataInterface(Offset), this)), Param);
+				ParameterStorage.SetDataInterface(CastChecked<UNiagaraDataInterface>(StaticDuplicateObject(DefaultStore.GetDataInterface(CollectionIndex), this)), Param);
 			}
 			else if (Param.IsUObject())
 			{
-				ParameterStorage.SetUObject(DefaultStore.GetUObject(Offset), Param);
+				ParameterStorage.SetUObject(DefaultStore.GetUObject(CollectionIndex), Param);
 			}
 			else
 			{
-				ParameterStorage.SetParameterData(DefaultStore.GetParameterData(Offset, Param.GetType()), ParameterStorageOffset, Param.GetSizeInBytes());
+				const int32* LocalOffset = ParameterStorage.FindParameterOffset(Param);
+				if (LocalOffset)
+				{
+					ParameterStorage.SetParameterData(DefaultStore.GetParameterData(CollectionIndex, Param.GetType()), *LocalOffset, Param.GetSizeInBytes());
+				}
 			}
 		}
 	}
 
-	ParameterStorage.Rebind();
+#if WITH_EDITOR
+	OnParameterValuesChanged.Broadcast(ParametersToRefresh);
+#endif
 }
 
 bool UNiagaraParameterCollectionInstance::OverridesParameter(const FNiagaraVariable& Parameter)const
@@ -274,6 +456,16 @@ void UNiagaraParameterCollectionInstance::SetOverridesParameter(const FNiagaraVa
 }
 
 #if WITH_EDITORONLY_DATA
+void UNiagaraParameterCollectionInstance::PreEditChange(FProperty* PropertyAboutToChange)
+{
+	Super::PreEditChange(PropertyAboutToChange);
+
+	if (PropertyAboutToChange && PropertyAboutToChange->GetFName() == GET_MEMBER_NAME_CHECKED(UNiagaraParameterCollectionInstance, Collection))
+	{
+		SetParent(nullptr);
+	}
+}
+
 void UNiagaraParameterCollectionInstance::PostEditChangeProperty(FPropertyChangedEvent& PropertyChangedEvent)
 {
 	Super::PostEditChangeProperty(PropertyChangedEvent);
@@ -437,7 +629,7 @@ void UNiagaraParameterCollectionInstance::SetQuatParameter(const FString& InVari
 UNiagaraParameterCollection::UNiagaraParameterCollection(const FObjectInitializer& ObjectInitializer)
 	: Super(ObjectInitializer)
 {
-	Namespace = *GetName();
+	Namespace = GetFName();
 	BuildFullNamespace();
 
 	DefaultInstance = ObjectInitializer.CreateDefaultSubobject<UNiagaraParameterCollectionInstance>(this, TEXT("Default Instance"));
@@ -445,17 +637,48 @@ UNiagaraParameterCollection::UNiagaraParameterCollection(const FObjectInitialize
 }
 
 #if WITH_EDITORONLY_DATA
+void UNiagaraParameterCollection::PreEditChange(FProperty* PropertyAboutToChange)
+{
+	Super::PreEditChange(PropertyAboutToChange);
+
+#if WITH_EDITOR
+	if (PropertyAboutToChange
+		&& PropertyAboutToChange->GetFName() == GET_MEMBER_NAME_CHECKED(UNiagaraParameterCollection, SourceMaterialCollection)
+		&& SourceMaterialCollection
+		&& OnMaterialCollectionChanged.IsValid())
+	{
+		SourceMaterialCollection->OnCollectionChanged().Remove(OnMaterialCollectionChanged);
+		OnMaterialCollectionChanged.Reset();
+	}
+#endif // WITH_EDITOR
+}
+
 void UNiagaraParameterCollection::PostEditChangeProperty(struct FPropertyChangedEvent& PropertyChangedEvent)
 {
 	Super::PostEditChangeProperty(PropertyChangedEvent);
 
 	MakeNamespaceNameUnique();
 
-	if (PropertyChangedEvent.Property && PropertyChangedEvent.Property->GetFName() == GET_MEMBER_NAME_CHECKED(UNiagaraParameterCollection, SourceMaterialCollection))
+	if (PropertyChangedEvent.Property)
 	{
-		AddDefaultSourceParameters();
-		OnChangedDelegate.Broadcast();
+		if (PropertyChangedEvent.Property->GetFName() == GET_MEMBER_NAME_CHECKED(UNiagaraParameterCollection, SourceMaterialCollection))
+		{
+#if WITH_EDITOR
+			if (SourceMaterialCollection)
+			{
+				OnMaterialCollectionChanged = SourceMaterialCollection->OnCollectionChanged().AddUObject(this, &UNiagaraParameterCollection::RefreshSourceMaterialCollection);
+			}
+#endif
+			RefreshSourceMaterialCollection();
+			OnCollectionChangedDelegate.Broadcast();
+		}
+		if (PropertyChangedEvent.Property->GetFName() == GET_MEMBER_NAME_CHECKED(UNiagaraParameterCollection, Namespace))
+		{
+			NamespaceChanged();
+			OnCollectionChangedDelegate.Broadcast();
+		}
 	}
+
 }
 
 void UNiagaraParameterCollection::PostDuplicate(EDuplicateMode::Type DuplicateMode)
@@ -631,29 +854,83 @@ void UNiagaraParameterCollection::MakeNamespaceNameUnique()
 
 	if (ExistingNames.Contains(Namespace))
 	{
-		FString CandidateNameString = Namespace.ToString();
-		FString BaseNameString = CandidateNameString;
-		if (CandidateNameString.Len() >= 3 && CandidateNameString.Right(3).IsNumeric())
+		// first try to use the name of the newly created asset as the namespace
+		FName ObjectName = GetFName();
+		if (!ExistingNames.Contains(ObjectName))
 		{
-			BaseNameString = CandidateNameString.Left(CandidateNameString.Len() - 3);
+			Namespace = ObjectName;
 		}
-
-		FName UniqueName = FName(*BaseNameString);
-		int32 NameIndex = 1;
-		while (ExistingNames.Contains(UniqueName))
+		else
 		{
-			UniqueName = FName(*FString::Printf(TEXT("%s%03i"), *BaseNameString, NameIndex));
-			NameIndex++;
-		}
+			FString CandidateNameString = ObjectName.ToString();
+			FString BaseNameString = CandidateNameString;
+			if (CandidateNameString.Len() >= 3 && CandidateNameString.Right(3).IsNumeric())
+			{
+				BaseNameString = CandidateNameString.Left(CandidateNameString.Len() - 3);
+			}
+
+			FName UniqueName = FName(*BaseNameString);
+			int32 NameIndex = 1;
+			while (ExistingNames.Contains(UniqueName))
+			{
+				UniqueName = FName(*FString::Printf(TEXT("%s%03i"), *BaseNameString, NameIndex));
+				NameIndex++;
+			}
 
-		UE_LOG(LogNiagara, Warning, TEXT("Parameter collection namespace conflict found. \"%s\" is already in use!"), *Namespace.ToString());
-		Namespace = UniqueName;
+			UE_LOG(LogNiagara, Warning, TEXT("Parameter collection namespace conflict found. \"%s\" is already in use!"), *Namespace.ToString());
+			Namespace = UniqueName;
+		}
 		BuildFullNamespace();
+		NamespaceChanged();
 	}
 }
 
-void UNiagaraParameterCollection::AddDefaultSourceParameters()
+void UNiagaraParameterCollection::RefreshSourceMaterialCollection()
 {
+#if WITH_EDITOR
+	TSet<FName> PreviouslyOwnedParameters = MoveTemp(ParametersOwnedBySourceCollection);
+	ParametersOwnedBySourceCollection.Reset();
+
+	bool bParametersAddedOrRemoved = false;
+	TArray<FName> UpdatedParameters;
+	auto AddOrUpdateParameter = [this, &bParametersAddedOrRemoved, &UpdatedParameters](const FNiagaraVariable& Parameter)
+#else
+	auto AddOrUpdateParameter = [this](const FNiagaraVariable& Parameter)
+#endif
+	{
+		// go through the existing elements to see if we already have an entry for the supplied parameter
+		int32 Idx = IndexOfParameter(Parameter);
+		if (Idx == INDEX_NONE)
+		{
+			Modify();
+
+			Idx = Parameters.Add(Parameter);
+			DefaultInstance->AddParameter(Parameter);
+
+#if WITH_EDITOR
+			bParametersAddedOrRemoved = true;
+#endif
+			return;
+		}
+		else if (Parameter.IsDataAllocated())
+		{
+			if (!Parameters[Idx].HoldsSameData(Parameter))
+			{
+				if (const uint8* ParameterData = Parameter.GetData())
+				{
+					Modify();
+
+					Parameters[Idx].SetData(ParameterData);
+					DefaultInstance->GetParameterStore().SetParameterData(ParameterData, Parameter);
+
+#if WITH_EDITOR
+					UpdatedParameters.Add(Parameter.GetName());
+#endif
+				}
+			}
+		}
+	};
+
 	if (SourceMaterialCollection)
 	{
 		TArray<FName> ScalarParameterNames;
@@ -670,16 +947,98 @@ void UNiagaraParameterCollection::AddDefaultSourceParameters()
 			FNiagaraVariable ScalarParameter(ScalarDef, ConditionalAddFullNamespace(ScalarParameterName));
 			ScalarParameter.SetValue(SourceMaterialCollection->GetScalarParameterByName(ScalarParameterName)->DefaultValue);
 
-			AddParameter(ScalarParameter);
+			AddOrUpdateParameter(ScalarParameter);
 		}
 
 		for (const FName& VectorParameterName : VectorParameterNames)
 		{
 			FNiagaraVariable VectorParameter(ColorDef, ConditionalAddFullNamespace(VectorParameterName));
 			VectorParameter.SetValue(SourceMaterialCollection->GetVectorParameterByName(VectorParameterName)->DefaultValue);
-			AddParameter(VectorParameter);
+
+			AddOrUpdateParameter(VectorParameter);
+		}
+
+#if WITH_EDITOR
+		ParametersOwnedBySourceCollection.Append(ScalarParameterNames);
+		ParametersOwnedBySourceCollection.Append(VectorParameterNames);
+#endif
+	}
+
+#if WITH_EDITOR
+	for (FName PreviousParameter : PreviouslyOwnedParameters)
+	{
+		if (!ParametersOwnedBySourceCollection.Contains(PreviousParameter))
+		{
+			const FName FullParameterName = ConditionalAddFullNamespace(PreviousParameter);
+
+			auto FindVariableByName = [FullParameterName](const FNiagaraVariable& Var) -> bool
+			{
+				return Var.GetName() == FullParameterName;
+			};
+
+			if (const FNiagaraVariable* FoundParameter = Parameters.FindByPredicate(FindVariableByName))
+			{
+				FNiagaraVariable ParameterToRemove = *FoundParameter;
+				RemoveParameter(ParameterToRemove);
+				bParametersAddedOrRemoved = true;
+			}
 		}
 	}
+
+	if (bParametersAddedOrRemoved)
+	{
+		OnCollectionChangedDelegate.Broadcast();
+	}
+
+	if (!UpdatedParameters.IsEmpty())
+	{
+		OnMaterialParametersChanged.Broadcast(MakeConstArrayView(UpdatedParameters));
+	}
+#endif
+}
+
+void UNiagaraParameterCollection::NamespaceChanged()
+{
+	TArray<FNiagaraVariable> ParametersToRemove;
+	TSet<FName> ExistingVariableNames;
+
+	for (FNiagaraVariable Parameter : Parameters)
+	{
+		FNameBuilder ParameterName(Parameter.GetName());
+		FStringView ParameterNameView = ParameterName.ToView();
+
+		// look for the form NPC.<OldNamespace>.<Name> and replace it with NPC.<Namespace>.<Name>
+		if (ParameterNameView.StartsWith(TEXT("NPC.")))
+		{
+			int32 LastSeparatorIndex = INDEX_NONE;
+			if (ParameterNameView.FindLastChar('.', LastSeparatorIndex))
+			{
+				FStringView VariableNameView = ParameterNameView.Mid(LastSeparatorIndex + 1);
+				if (!VariableNameView.IsEmpty())
+				{
+					FName VariableName(VariableNameView);
+
+					bool DuplicateFound = false;
+					ExistingVariableNames.Add(VariableName, &DuplicateFound);
+					if (DuplicateFound)
+					{
+						// mark the parameter as requiring removal
+						ParametersToRemove.Add(Parameter);
+					}
+
+					FNameBuilder NewParameterName;
+					FullNamespace.AppendString(NewParameterName);
+					VariableName.AppendString(NewParameterName);
+					RenameParameter(Parameter, FName(NewParameterName));
+				}
+			}
+		}
+	}
+
+	for (const FNiagaraVariable& Parameter : ParametersToRemove)
+	{
+		RemoveParameter(Parameter);
+	}
 }
 
 #endif
@@ -700,11 +1059,15 @@ void UNiagaraParameterCollection::PostLoad()
 
 	if (SourceMaterialCollection)
 	{
+#if WITH_EDITOR
+		OnMaterialCollectionChanged = SourceMaterialCollection->OnCollectionChanged().AddUObject(this, &UNiagaraParameterCollection::RefreshSourceMaterialCollection);
+#endif
+
 		SourceMaterialCollection->ConditionalPostLoad();
 
-#if WITH_EDITOR
+#if WITH_EDITORONLY_DATA
 		// catch up with any changes that may have been made to the MPC
-		AddDefaultSourceParameters();
+		RefreshSourceMaterialCollection();
 #endif
 	}
 }
diff --git a/Engine/Plugins/FX/Niagara/Source/Niagara/Private/NiagaraWorldManager.cpp b/Engine/Plugins/FX/Niagara/Source/Niagara/Private/NiagaraWorldManager.cpp
--- a/Engine/Plugins/FX/Niagara/Source/Niagara/Private/NiagaraWorldManager.cpp
+++ b/Engine/Plugins/FX/Niagara/Source/Niagara/Private/NiagaraWorldManager.cpp
@@ -627,7 +627,7 @@ UNiagaraParameterCollectionInstance* FNiagaraWorldManager::GetParameterCollectio
 		UNiagaraParameterCollectionInstance* DefaultInstance = Collection->GetDefaultInstance();
 		OverrideInst = &ParameterCollections.Add(Collection);
 		*OverrideInst = CastChecked<UNiagaraParameterCollectionInstance>(StaticDuplicateObject(DefaultInstance, World));
-#if WITH_EDITORONLY_DATA
+#if WITH_EDITOR
 		//Bind to the default instance so that changes to the collection propagate through.
 		DefaultInstance->GetParameterStore().Bind(&(*OverrideInst)->GetParameterStore());
 #endif
@@ -641,16 +641,19 @@ UNiagaraParameterCollectionInstance* FNiagaraWorldManager::GetParameterCollectio
 
 void FNiagaraWorldManager::CleanupParameterCollections()
 {
-#if WITH_EDITOR
 	for (TPair<UNiagaraParameterCollection*, UNiagaraParameterCollectionInstance*> CollectionInstPair : ParameterCollections)
 	{
 		UNiagaraParameterCollection* Collection = CollectionInstPair.Key;
 		UNiagaraParameterCollectionInstance* CollectionInst = CollectionInstPair.Value;
+
+		CollectionInst->Unbind(World);
+
+#if WITH_EDITOR
 		//Ensure that the default instance is not bound to the override.
 		UNiagaraParameterCollectionInstance* DefaultInst = Collection->GetDefaultInstance();
 		DefaultInst->GetParameterStore().Unbind(&CollectionInst->GetParameterStore());
-	}
 #endif
+	}
 	ParameterCollections.Empty();
 }
 
diff --git a/Engine/Plugins/FX/Niagara/Source/NiagaraEditor/Private/ViewModels/NiagaraCollectionParameterViewModel.cpp b/Engine/Plugins/FX/Niagara/Source/NiagaraEditor/Private/ViewModels/NiagaraCollectionParameterViewModel.cpp
--- a/Engine/Plugins/FX/Niagara/Source/NiagaraEditor/Private/ViewModels/NiagaraCollectionParameterViewModel.cpp
+++ b/Engine/Plugins/FX/Niagara/Source/NiagaraEditor/Private/ViewModels/NiagaraCollectionParameterViewModel.cpp
@@ -10,7 +10,7 @@
 
 #define LOCTEXT_NAMESPACE "CollectionParameterViewModel"
 
-FNiagaraCollectionParameterViewModel::FNiagaraCollectionParameterViewModel(FNiagaraVariable& Variable, UNiagaraParameterCollectionInstance* CollectionInstance, ENiagaraParameterEditMode ParameterEditMode)
+FNiagaraCollectionParameterViewModel::FNiagaraCollectionParameterViewModel(const FNiagaraVariable& Variable, UNiagaraParameterCollectionInstance* CollectionInstance, ENiagaraParameterEditMode ParameterEditMode)
 	: FNiagaraParameterViewModel(ParameterEditMode)
 	, CollectionInst(CollectionInstance)
 	, Parameter(Variable)
@@ -108,7 +108,7 @@ bool FNiagaraCollectionParameterViewModel::CanChangeParameterType()const
 
 bool FNiagaraCollectionParameterViewModel::IsEditingEnabled()const
 {
-	return IsProvided() == ECheckBoxState::Checked;
+	return IsProvided() == ECheckBoxState::Checked && !CollectionInst->IsExternallyDriven(Parameter);
 }
 
 bool FNiagaraCollectionParameterViewModel::CanChangeSortOrder() const
@@ -128,12 +128,22 @@ void FNiagaraCollectionParameterViewModel::SetSortOrder(int32 SortOrder)
 
 bool FNiagaraCollectionParameterViewModel::CanRenameParameter()const
 {
-	return CollectionInst->IsDefaultInstance();
+	return !CollectionInst->IsExternallyDriven(Parameter) && CollectionInst->IsDefaultInstance();
+}
+
+FText FNiagaraCollectionParameterViewModel::GetTooltip() const
+{
+	if (CollectionInst->IsExternallyDriven(Parameter))
+	{
+		return FText::Format(LOCTEXT("NPCExternallyDrivenTip", "\"{0}\" can not be edited as it is being driven from the parent MaterialParameterCollection"), GetNameText());
+	}
+
+	return FNiagaraParameterViewModel::GetTooltip();
 }
 
 bool FNiagaraCollectionParameterViewModel::IsOptional()const
 {
-	return !CollectionInst->IsDefaultInstance(); 
+	return !CollectionInst->IsExternallyDriven(Parameter) && !CollectionInst->IsDefaultInstance();
 }
 
 ECheckBoxState FNiagaraCollectionParameterViewModel::IsProvided()const
@@ -157,6 +167,7 @@ void FNiagaraCollectionParameterViewModel::SetProvided(ECheckBoxState CheckboxSt
 		else
 		{
 			CollectionInst->SetOverridesParameter(Parameter, false);
+			CollectionInst->SyncWithCollectionParameters({ Parameter.GetName() });
 		}
 	}
 
diff --git a/Engine/Plugins/FX/Niagara/Source/NiagaraEditor/Private/ViewModels/NiagaraCollectionParameterViewModel.h b/Engine/Plugins/FX/Niagara/Source/NiagaraEditor/Private/ViewModels/NiagaraCollectionParameterViewModel.h
--- a/Engine/Plugins/FX/Niagara/Source/NiagaraEditor/Private/ViewModels/NiagaraCollectionParameterViewModel.h
+++ b/Engine/Plugins/FX/Niagara/Source/NiagaraEditor/Private/ViewModels/NiagaraCollectionParameterViewModel.h
@@ -13,7 +13,7 @@ public:
 
 public:
 
-	FNiagaraCollectionParameterViewModel(FNiagaraVariable& Variable, UNiagaraParameterCollectionInstance* CollectionInstance, ENiagaraParameterEditMode ParameterEditMode);
+	FNiagaraCollectionParameterViewModel(const FNiagaraVariable& Variable, UNiagaraParameterCollectionInstance* CollectionInstance, ENiagaraParameterEditMode ParameterEditMode);
 
 	void Reset();
 
@@ -38,6 +38,8 @@ public:
 	virtual int32 GetSortOrder() const override;
 	virtual void SetSortOrder(int32 SortOrder) override;
 	virtual bool CanRenameParameter()const override;
+	virtual FText GetTooltip() const override;
+
 
 	virtual bool CanChangeParameterType()const override;
 	virtual bool IsEditingEnabled()const override;
diff --git a/Engine/Plugins/FX/Niagara/Source/NiagaraEditor/Private/ViewModels/NiagaraParameterCollectionAssetViewModel.cpp b/Engine/Plugins/FX/Niagara/Source/NiagaraEditor/Private/ViewModels/NiagaraParameterCollectionAssetViewModel.cpp
--- a/Engine/Plugins/FX/Niagara/Source/NiagaraEditor/Private/ViewModels/NiagaraParameterCollectionAssetViewModel.cpp
+++ b/Engine/Plugins/FX/Niagara/Source/NiagaraEditor/Private/ViewModels/NiagaraParameterCollectionAssetViewModel.cpp
@@ -31,7 +31,8 @@ FNiagaraParameterCollectionAssetViewModel::FNiagaraParameterCollectionAssetViewM
 	RegisteredHandle = RegisterViewModelWithMap(Collection, this);
 	GEditor->RegisterForUndo(this);
 
-	ExternalChangeHandle = InCollection->OnChangedDelegate.AddRaw(this, &FNiagaraParameterCollectionAssetViewModel::OnCollectionChangedExternally);
+	ExternalChangeHandle = Collection->OnCollectionChangedDelegate.AddRaw(this, &FNiagaraParameterCollectionAssetViewModel::OnCollectionChangedExternally);
+	ParameterValuesChangedHandle = Instance->ParameterValuesChanged().AddRaw(this, &FNiagaraParameterCollectionAssetViewModel::RefreshParameterViewModels);
 
 	RefreshParameterViewModels();
 }
@@ -44,8 +45,15 @@ FNiagaraParameterCollectionAssetViewModel::FNiagaraParameterCollectionAssetViewM
 {
 	check(Instance);
 
-	RegisteredHandle = RegisterViewModelWithMap(Collection, this);
-	GEditor->UnregisterForUndo(this);
+	GEditor->RegisterForUndo(this);
+
+	if (Collection)
+	{
+		RegisteredHandle = RegisterViewModelWithMap(Collection, this);
+		ExternalChangeHandle = Collection->OnCollectionChangedDelegate.AddRaw(this, &FNiagaraParameterCollectionAssetViewModel::OnCollectionChangedExternally);
+	}
+
+	ParameterValuesChangedHandle = Instance->ParameterValuesChanged().AddRaw(this, &FNiagaraParameterCollectionAssetViewModel::RefreshParameterViewModels);
 
 	RefreshParameterViewModels();
 }
@@ -69,8 +77,10 @@ FNiagaraParameterCollectionAssetViewModel::~FNiagaraParameterCollectionAssetView
 
 	if (Collection)
 	{
-		Collection->OnChangedDelegate.Remove(ExternalChangeHandle);
+		Collection->OnCollectionChangedDelegate.Remove(ExternalChangeHandle);
 	}
+
+	Instance->ParameterValuesChanged().Remove(ParameterValuesChangedHandle);
 }
 
 void FNiagaraParameterCollectionAssetViewModel::NotifyPreChange(FProperty* PropertyAboutToChange)
@@ -78,6 +88,14 @@ void FNiagaraParameterCollectionAssetViewModel::NotifyPreChange(FProperty* Prope
 	if (PropertyAboutToChange->GetFName() == GET_MEMBER_NAME_CHECKED(UNiagaraParameterCollectionInstance, Collection))
 	{
 		GEditor->BeginTransaction(LOCTEXT("ChangeNPCInstanceParent", "Change Parent"));
+
+		if (Collection)
+		{
+			Collection->OnCollectionChangedDelegate.Remove(ExternalChangeHandle);
+			UnregisterViewModelWithMap(RegisteredHandle);
+		}
+
+		Instance->ParameterValuesChanged().Remove(ParameterValuesChangedHandle);
 		Instance->Empty();
 	}
 }
@@ -89,6 +107,15 @@ void FNiagaraParameterCollectionAssetViewModel::NotifyPostChange(const FProperty
 		//Instance has changed parent;
 		check(!Instance->IsDefaultInstance());
 		Collection = Instance->GetParent();
+
+		if (Collection)
+		{
+			RegisteredHandle = RegisterViewModelWithMap(Collection, this);
+			ExternalChangeHandle = Collection->OnCollectionChangedDelegate.AddRaw(this, &FNiagaraParameterCollectionAssetViewModel::OnCollectionChangedExternally);
+		}
+
+		ParameterValuesChangedHandle = Instance->ParameterValuesChanged().AddRaw(this, &FNiagaraParameterCollectionAssetViewModel::RefreshParameterViewModels);
+
 		RefreshParameterViewModels();
 		GEditor->EndTransaction();
 	}
@@ -281,6 +308,18 @@ void FNiagaraParameterCollectionAssetViewModel::CollectionChanged(bool bRecompil
 	}
 }
 
+TSharedPtr<FNiagaraCollectionParameterViewModel> FNiagaraParameterCollectionAssetViewModel::CreateParameterViewModel(const FNiagaraVariable& Var)
+{
+	TSharedPtr<FNiagaraCollectionParameterViewModel> ParameterViewModel = MakeShareable(new FNiagaraCollectionParameterViewModel(Var, Instance, ParameterEditMode));
+
+	ParameterViewModel->OnNameChanged().AddRaw(this, &FNiagaraParameterCollectionAssetViewModel::OnParameterNameChanged, Var);
+	ParameterViewModel->OnTypeChanged().AddRaw(this, &FNiagaraParameterCollectionAssetViewModel::OnParameterTypeChanged, Var);
+	ParameterViewModel->OnDefaultValueChanged().AddRaw(this, &FNiagaraParameterCollectionAssetViewModel::OnParameterValueChangedInternal, ParameterViewModel.ToSharedRef());
+	ParameterViewModel->OnProvidedChanged().AddRaw(this, &FNiagaraParameterCollectionAssetViewModel::OnParameterProvidedChanged, Var);
+
+	return ParameterViewModel;
+}
+
 void FNiagaraParameterCollectionAssetViewModel::RefreshParameterViewModels()
 {
 	if (!Collection)
@@ -329,18 +368,38 @@ void FNiagaraParameterCollectionAssetViewModel::RefreshParameterViewModels()
 			}
 		}
 		
-		TSharedPtr<FNiagaraCollectionParameterViewModel> ParameterViewModel = MakeShareable(new FNiagaraCollectionParameterViewModel(Var, Instance, ParameterEditMode));
-
-		ParameterViewModel->OnNameChanged().AddRaw(this, &FNiagaraParameterCollectionAssetViewModel::OnParameterNameChanged, Var);
-		ParameterViewModel->OnTypeChanged().AddRaw(this, &FNiagaraParameterCollectionAssetViewModel::OnParameterTypeChanged, Var);
-		ParameterViewModel->OnDefaultValueChanged().AddRaw(this, &FNiagaraParameterCollectionAssetViewModel::OnParameterValueChangedInternal, ParameterViewModel.ToSharedRef());
-		ParameterViewModel->OnProvidedChanged().AddRaw(this, &FNiagaraParameterCollectionAssetViewModel::OnParameterProvidedChanged, Var);
-		ParameterViewModels.Add(ParameterViewModel.ToSharedRef());
+		ParameterViewModels.Add(CreateParameterViewModel(Var).ToSharedRef());
 	}
 
 	OnCollectionChangedDelegate.Broadcast();
 }
 
+void FNiagaraParameterCollectionAssetViewModel::RefreshParameterViewModels(TConstArrayView<FName> ChangedParameters)
+{
+	bool bChangeMade = false;
+
+	for (TSharedRef<INiagaraParameterViewModel>& ParameterViewModel : ParameterViewModels)
+	{
+		FNiagaraVariable Var = ParameterViewModel->GetVariable();
+		if (ChangedParameters.Contains(Var.GetName()))
+		{
+			if (FNiagaraCollectionParameterViewModel* CastParameterViewModel = (FNiagaraCollectionParameterViewModel*)(&ParameterViewModel.Get()))
+			{
+				CastParameterViewModel->Reset();
+			}
+
+			ParameterViewModel = CreateParameterViewModel(Var).ToSharedRef();
+
+			bChangeMade = true;
+		}
+	}
+
+	if (bChangeMade)
+	{
+		OnCollectionChangedDelegate.Broadcast();
+	}
+}
+
 bool FNiagaraParameterCollectionAssetViewModel::SupportsType(const FNiagaraTypeDefinition& Type) const
 {
 	return Type != FNiagaraTypeDefinition::GetGenericNumericDef() && !Type.IsInternalType();
@@ -441,7 +500,12 @@ void FNiagaraParameterCollectionAssetViewModel::OnParameterValueChangedInternal(
 	//restart any systems using this collection.
 	FNiagaraSystemUpdateContext UpdateContext(Collection, true);
 
-	OnParameterValueChanged().Broadcast(ChangedParameter->GetName());
+	OnParameterValueChanged().Broadcast({ ChangedParameter->GetName() });
+
+	if (Collection->GetDefaultInstance() == Instance)
+	{
+		Collection->OnParameterValuesEdited.Broadcast({ ChangedParameter->GetVariable().GetName() });
+	}
 	
 	//Push the change to anyone already bound.
 	Instance->GetParameterStore().Tick();
diff --git a/Engine/Plugins/FX/Niagara/Source/NiagaraEditor/Private/ViewModels/NiagaraParameterCollectionAssetViewModel.h b/Engine/Plugins/FX/Niagara/Source/NiagaraEditor/Private/ViewModels/NiagaraParameterCollectionAssetViewModel.h
--- a/Engine/Plugins/FX/Niagara/Source/NiagaraEditor/Private/ViewModels/NiagaraParameterCollectionAssetViewModel.h
+++ b/Engine/Plugins/FX/Niagara/Source/NiagaraEditor/Private/ViewModels/NiagaraParameterCollectionAssetViewModel.h
@@ -65,6 +65,8 @@ protected:
 	//~ FNiagaraParameterCollectionViewModel interface.
 	virtual bool SupportsType(const FNiagaraTypeDefinition& Type) const override;
 
+	void RefreshParameterViewModels(TConstArrayView<FName> ChangedParameters);
+
 private:	
 	/** Handles when the name on a parameter changes. */
 	void OnParameterNameChanged(FName OldName, FName NewName, FNiagaraVariable ParameterVariable);
@@ -83,6 +85,8 @@ private:
 	FName GenerateNewName(FNiagaraTypeDefinition Type)const;
 
 private:
+	TSharedPtr<FNiagaraCollectionParameterViewModel> CreateParameterViewModel(const FNiagaraVariable& Var);
+
 	/** The parameter view models. */
 	TArray<TSharedRef<INiagaraParameterViewModel>> ParameterViewModels;
 
@@ -97,5 +101,6 @@ private:
 	TNiagaraViewModelManager<UNiagaraParameterCollection, FNiagaraParameterCollectionAssetViewModel>::Handle RegisteredHandle;
 
 	FDelegateHandle ExternalChangeHandle;
+	FDelegateHandle ParameterValuesChangedHandle;
 };
 
diff --git a/Engine/Plugins/FX/Niagara/Source/NiagaraEditor/Private/ViewModels/NiagaraParameterCollectionViewModel.h b/Engine/Plugins/FX/Niagara/Source/NiagaraEditor/Private/ViewModels/NiagaraParameterCollectionViewModel.h
--- a/Engine/Plugins/FX/Niagara/Source/NiagaraEditor/Private/ViewModels/NiagaraParameterCollectionViewModel.h
+++ b/Engine/Plugins/FX/Niagara/Source/NiagaraEditor/Private/ViewModels/NiagaraParameterCollectionViewModel.h
@@ -20,7 +20,7 @@ class INiagaraParameterCollectionViewModel
 public:
 	DECLARE_MULTICAST_DELEGATE(FOnCollectionChanged);
 	DECLARE_MULTICAST_DELEGATE(FOnExpandedChanged);
-	DECLARE_MULTICAST_DELEGATE_OneParam(FOnParameterValueChanged, FName);
+	DECLARE_MULTICAST_DELEGATE_OneParam(FOnParameterValueChanged, TConstArrayView<FName>);
 
 public:
 	virtual ~INiagaraParameterCollectionViewModel() { }
diff --git a/Engine/Plugins/FX/Niagara/Source/NiagaraEditor/Private/ViewModels/NiagaraParameterViewModel.h b/Engine/Plugins/FX/Niagara/Source/NiagaraEditor/Private/ViewModels/NiagaraParameterViewModel.h
--- a/Engine/Plugins/FX/Niagara/Source/NiagaraEditor/Private/ViewModels/NiagaraParameterViewModel.h
+++ b/Engine/Plugins/FX/Niagara/Source/NiagaraEditor/Private/ViewModels/NiagaraParameterViewModel.h
@@ -42,10 +42,10 @@ public:
 	/** Gets whether or not this parameter can be renamed. */
 	virtual bool CanRenameParameter() const = 0;
 
-	/** Gets the text representation of the name of the paramter. */
+	/** Gets the text representation of the name of the parameter. */
 	virtual FText GetNameText() const = 0;
 
-	/** Handles a hame text change being comitter from the UI. */
+	/** Handles a name text change being committer from the UI. */
 	virtual void NameTextCommitted(const FText& Name, ETextCommit::Type CommitInfo) = 0;
 
 	/** Handles verification of an in-progress variable name change in the UI.*/
@@ -57,10 +57,10 @@ public:
 	/** Gets whether or not the type of this parameter can be changed. */
 	virtual bool CanChangeParameterType() const = 0;
 
-	/** Gets the type of the paramter. */
+	/** Gets the type of the parameter. */
 	virtual FNiagaraTypeDefinition GetType() const = 0;
 
-	/** Handles the paramter type being changed from the UI. */
+	/** Handles the parameter type being changed from the UI. */
 	virtual void SelectedTypeChanged(TSharedPtr<FNiagaraTypeDefinition> Item, ESelectInfo::Type SelectionType) = 0;
 
 	/** Gets the type of default value this view model provides. */
@@ -101,7 +101,7 @@ public:
 	/** Gets whether or not this parameter is editable in this context.*/
 	virtual bool IsEditingEnabled() const = 0;
 
-	/** Sets whether or not this parameter is editible in this context.*/
+	/** Sets whether or not this parameter is editable in this context.*/
 	virtual void SetEditingEnabled(bool bEnabled) = 0;
 
 	/** Gets the tooltip when hovering over this parameter.*/
diff --git a/Engine/Plugins/FX/Niagara/Source/NiagaraEditor/Private/ViewModels/NiagaraScriptInputCollectionViewModel.cpp b/Engine/Plugins/FX/Niagara/Source/NiagaraEditor/Private/ViewModels/NiagaraScriptInputCollectionViewModel.cpp
--- a/Engine/Plugins/FX/Niagara/Source/NiagaraEditor/Private/ViewModels/NiagaraScriptInputCollectionViewModel.cpp
+++ b/Engine/Plugins/FX/Niagara/Source/NiagaraEditor/Private/ViewModels/NiagaraScriptInputCollectionViewModel.cpp
@@ -563,7 +563,7 @@ void FNiagaraScriptInputCollectionViewModel::OnParameterValueChangedInternal(TSh
 		}
 	}
 
-	OnParameterValueChanged().Broadcast(ChangedParameter->GetName());
+	OnParameterValueChanged().Broadcast({ ChangedParameter->GetName() });
 }
 
 #undef LOCTEXT_NAMESPACE // "NiagaraScriptInputCollection"
diff --git a/Engine/Plugins/FX/Niagara/Source/NiagaraEditor/Private/ViewModels/NiagaraScriptOutputCollectionViewModel.cpp b/Engine/Plugins/FX/Niagara/Source/NiagaraEditor/Private/ViewModels/NiagaraScriptOutputCollectionViewModel.cpp
--- a/Engine/Plugins/FX/Niagara/Source/NiagaraEditor/Private/ViewModels/NiagaraScriptOutputCollectionViewModel.cpp
+++ b/Engine/Plugins/FX/Niagara/Source/NiagaraEditor/Private/ViewModels/NiagaraScriptOutputCollectionViewModel.cpp
@@ -324,7 +324,7 @@ void FNiagaraScriptOutputCollectionViewModel::OnParameterValueChangedInternal(FN
 		GraphOutputNode->NotifyOutputVariablesChanged();
 	}
 
-	OnParameterValueChanged().Broadcast(ParameterVariable->GetName());
+	OnParameterValueChanged().Broadcast({ ParameterVariable->GetName() });
 }
 
 #undef LOCTEXT_NAMESPACE // "NiagaraScriptInputCollection"
diff --git a/Engine/Plugins/FX/Niagara/Source/NiagaraEditor/Private/ViewModels/NiagaraScriptViewModel.cpp b/Engine/Plugins/FX/Niagara/Source/NiagaraEditor/Private/ViewModels/NiagaraScriptViewModel.cpp
--- a/Engine/Plugins/FX/Niagara/Source/NiagaraEditor/Private/ViewModels/NiagaraScriptViewModel.cpp
+++ b/Engine/Plugins/FX/Niagara/Source/NiagaraEditor/Private/ViewModels/NiagaraScriptViewModel.cpp
@@ -547,12 +547,12 @@ void FNiagaraScriptViewModel::InputViewModelSelectionChanged()
 	}
 }
 
-void FNiagaraScriptViewModel::InputParameterValueChanged(FName ParameterName)
+void FNiagaraScriptViewModel::InputParameterValueChanged(TConstArrayView<FName> ParameterName)
 {
 	MarkAllDirty(TEXT("Input parameter value changed"));
 }
 
-void FNiagaraScriptViewModel::OutputParameterValueChanged(FName ParameterName)
+void FNiagaraScriptViewModel::OutputParameterValueChanged(TConstArrayView<FName> ParameterName)
 {
 	MarkAllDirty(TEXT("Output parameter value changed"));
 }
diff --git a/Engine/Plugins/FX/Niagara/Source/NiagaraEditor/Private/Widgets/SNiagaraParameterCollection.cpp b/Engine/Plugins/FX/Niagara/Source/NiagaraEditor/Private/Widgets/SNiagaraParameterCollection.cpp
--- a/Engine/Plugins/FX/Niagara/Source/NiagaraEditor/Private/Widgets/SNiagaraParameterCollection.cpp
+++ b/Engine/Plugins/FX/Niagara/Source/NiagaraEditor/Private/Widgets/SNiagaraParameterCollection.cpp
@@ -579,8 +579,8 @@ TSharedRef<ITableRow> SNiagaraParameterCollection::OnGenerateRowForParameter(TSh
 	}
 
 	TSharedRef<SWidget> ParameterDeleteButton = SNew(SButton)
-		.OnClicked_Raw(this, &SNiagaraParameterCollection::DeleteParameter, Item)
-		.Visibility_Raw(this, &SNiagaraParameterCollection::ShowParameterDeleteButton)
+		.OnClicked(this, &SNiagaraParameterCollection::DeleteParameter, Item)
+		.Visibility(this, &SNiagaraParameterCollection::ShowParameterDeleteButton, Item)
 		.Content()
 		[
 			SNew(SImage)
@@ -862,10 +862,15 @@ FReply SNiagaraParameterCollection::OnItemAcceptDrop(FDragDropEvent const& DragD
 	return bWasDropHandled ? FReply::Handled() : FReply::Unhandled();
 }
 
-EVisibility SNiagaraParameterCollection::ShowParameterDeleteButton() const
+EVisibility SNiagaraParameterCollection::ShowParameterDeleteButton(TSharedRef<INiagaraParameterViewModel> Item) const
 {
-	// if we can add, we can also remove
-	return Collection->GetAddButtonVisibility();
+	// if we can add, we can also remove, unless it's an externally provided parameter
+	if (Item->IsEditingEnabled())
+	{
+		return Collection->GetAddButtonVisibility();
+	}
+
+	return EVisibility::Hidden;
 }
 
 void SNiagaraParameterCollection::ParameterViewModelTypeChanged()
diff --git a/Engine/Plugins/FX/Niagara/Source/NiagaraEditor/Private/Widgets/SNiagaraParameterCollection.h b/Engine/Plugins/FX/Niagara/Source/NiagaraEditor/Private/Widgets/SNiagaraParameterCollection.h
--- a/Engine/Plugins/FX/Niagara/Source/NiagaraEditor/Private/Widgets/SNiagaraParameterCollection.h
+++ b/Engine/Plugins/FX/Niagara/Source/NiagaraEditor/Private/Widgets/SNiagaraParameterCollection.h
@@ -111,7 +111,7 @@ private:
 	TOptional<EItemDropZone> OnItemCanAcceptDrop(const FDragDropEvent& DragDropEvent, EItemDropZone DropZone, TSharedRef<INiagaraParameterViewModel> DropItem);
 	FReply OnItemAcceptDrop(FDragDropEvent const& DragDropEvent, EItemDropZone DropZone, TSharedRef<INiagaraParameterViewModel> DropItem);
 
-	EVisibility ShowParameterDeleteButton() const;
+	EVisibility ShowParameterDeleteButton(TSharedRef<INiagaraParameterViewModel> Item) const;
 
 private:
 	/** The view model for the parameter collection. */
diff --git a/Engine/Plugins/FX/Niagara/Source/NiagaraEditor/Public/ViewModels/NiagaraScriptViewModel.h b/Engine/Plugins/FX/Niagara/Source/NiagaraEditor/Public/ViewModels/NiagaraScriptViewModel.h
--- a/Engine/Plugins/FX/Niagara/Source/NiagaraEditor/Public/ViewModels/NiagaraScriptViewModel.h
+++ b/Engine/Plugins/FX/Niagara/Source/NiagaraEditor/Public/ViewModels/NiagaraScriptViewModel.h
@@ -102,11 +102,11 @@ private:
 
 	void SetScripts(UNiagaraScriptSource* InScriptSource, TArray<FVersionedNiagaraScript>& InScripts);
 
-	/** Handles when a value in the input parameter collection changes. */
-	void InputParameterValueChanged(FName ParameterName);
+	/** Handles when values in the input parameter collection changes. */
+	void InputParameterValueChanged(TConstArrayView<FName> ParameterName);
 
-	/** Handles when a value in the output parameter collection changes. */
-	void OutputParameterValueChanged(FName ParameterName);
+	/** Handles when values in the output parameter collection changes. */
+	void OutputParameterValueChanged(TConstArrayView<FName> ParameterName);
 
 protected:
 	/** The script which provides the data for this view model. */
diff --git a/Engine/Source/Runtime/Engine/Private/Materials/ParameterCollection.cpp b/Engine/Source/Runtime/Engine/Private/Materials/ParameterCollection.cpp
--- a/Engine/Source/Runtime/Engine/Private/Materials/ParameterCollection.cpp
+++ b/Engine/Source/Runtime/Engine/Private/Materials/ParameterCollection.cpp
@@ -18,7 +18,7 @@ int32 GDeferUpdateRenderStates = 1;
 FAutoConsoleVariableRef CVarDeferUpdateRenderStates(
 	TEXT("r.DeferUpdateRenderStates"),
 	GDeferUpdateRenderStates,
-	TEXT("Whether to defer updating the render states of material parameter collections when a paramter is changed until a rendering command needs them up to date.  Deferring updates is more efficient because multiple SetVectorParameterValue and SetScalarParameterValue calls in a frame will only result in one update."),
+	TEXT("Whether to defer updating the render states of material parameter collections when a parameter is changed until a rendering command needs them up to date.  Deferring updates is more efficient because multiple SetVectorParameterValue and SetScalarParameterValue calls in a frame will only result in one update."),
 	ECVF_RenderThreadSafe
 );
 
@@ -365,6 +365,10 @@ void UMaterialParameterCollection::PostEditChangeProperty(FPropertyChangedEvent&
 		UpdateDefaultResource(false);
 	}
 
+#if WITH_EDITOR
+	CollectionChangedDelegate.Broadcast();
+#endif
+
 	Super::PostEditChangeProperty(PropertyChangedEvent);
 }
 
diff --git a/Engine/Source/Runtime/Engine/Public/Materials/MaterialParameterCollection.h b/Engine/Source/Runtime/Engine/Public/Materials/MaterialParameterCollection.h
--- a/Engine/Source/Runtime/Engine/Public/Materials/MaterialParameterCollection.h
+++ b/Engine/Source/Runtime/Engine/Public/Materials/MaterialParameterCollection.h
@@ -196,6 +196,12 @@ public:
 
 #if WITH_EDITOR
 	static ENGINE_API FName GetBaseParameterCollectionMemberName();
+
+	DECLARE_MULTICAST_DELEGATE(FOnCollectionChanged);
+	FOnCollectionChanged& OnCollectionChanged() const
+	{
+		return CollectionChangedDelegate;
+	}
 #endif
 
 private:
@@ -229,6 +235,11 @@ private:
 
 	TUniquePtr<FShaderParametersMetadata> UniformBufferStruct;
 
+#if WITH_EDITOR
+	/** Delegate for when the collection changes and consummers may need to be updated */
+	mutable FOnCollectionChanged CollectionChangedDelegate;
+#endif
+
 	void CreateBufferStruct();
 
 	/** Gets default values into data to be set on the uniform buffer. */
