Fix comment spelling mistake

diff --git a/Engine/Plugins/FX/Niagara/Source/NiagaraEditor/Private/ViewModels/NiagaraParameterViewModel.h b/Engine/Plugins/FX/Niagara/Source/NiagaraEditor/Private/ViewModels/NiagaraParameterViewModel.h
--- a/Engine/Plugins/FX/Niagara/Source/NiagaraEditor/Private/ViewModels/NiagaraParameterViewModel.h
+++ b/Engine/Plugins/FX/Niagara/Source/NiagaraEditor/Private/ViewModels/NiagaraParameterViewModel.h
@@ -45,7 +45,7 @@ public:
 	/** Gets the text representation of the name of the parameter. */
 	virtual FText GetNameText() const = 0;
 
-	/** Handles a name text change being committer from the UI. */
+	/** Handles a name text change being committed from the UI. */
 	virtual void NameTextCommitted(const FText& Name, ETextCommit::Type CommitInfo) = 0;
 
 	/** Handles verification of an in-progress variable name change in the UI.*/
