Fixes to the RVT streaming options within the landscape sampling DI

#tests Validated that trying to preload a large section wouldn't blow array limits

diff --git a/Engine/Plugins/FX/Niagara/Source/Niagara/Private/NiagaraDataInterfaceLandscape.cpp b/Engine/Plugins/FX/Niagara/Source/Niagara/Private/NiagaraDataInterfaceLandscape.cpp
--- a/Engine/Plugins/FX/Niagara/Source/Niagara/Private/NiagaraDataInterfaceLandscape.cpp
+++ b/Engine/Plugins/FX/Niagara/Source/Niagara/Private/NiagaraDataInterfaceLandscape.cpp
@@ -249,6 +249,7 @@ struct FNDILandscapeData_GameThread
 	int32 HeightVirtualTextureIndex = INDEX_NONE;
 	int32 NormalVirtualTextureIndex = INDEX_NONE;
 	ERuntimeVirtualTextureMaterialType NormalVirtualTextureMode = ERuntimeVirtualTextureMaterialType::Count;
+	FBox2D RVTLandscapeBounds = FBox2D(EForceInit::ForceInit);
 	bool RequiresCollisionCacheCpu = false;
 	bool RequiresCollisionCacheGpu = false;
 	bool RequiresPhysMatCacheGpu = false;
@@ -267,6 +268,7 @@ struct FNDILandscapeData_GameThread
 		HeightVirtualTextureIndex = INDEX_NONE;
 		NormalVirtualTextureIndex = INDEX_NONE;
 		NormalVirtualTextureMode = ERuntimeVirtualTextureMaterialType::Count;
+		RVTLandscapeBounds = FBox2D(EForceInit::ForceInit);
 		RequiresCollisionCacheCpu = false;
 		RequiresCollisionCacheGpu = false;
 		RequiresPhysMatCacheGpu = false;
@@ -714,10 +716,13 @@ struct FNiagaraDataInterfaceProxyLandscape : public FNiagaraDataInterfaceProxy
 					{
 						if (TexturePage && TexturePage->IsValid() && TexturePage->AllocatedTexture)
 						{
-							const FVector2D UVMin = TexturePage->WorldToUv(FVector(RegionRequest.Region.Min.X, RegionRequest.Region.Min.Y, 0.0)) * TexturePage->SizeInPages;
-							const FVector2D UVMax = TexturePage->WorldToUv(FVector(RegionRequest.Region.Max.X, RegionRequest.Region.Max.Y, 0.0)) * TexturePage->SizeInPages;
+							const FVector2D UVMin = TexturePage->WorldToUv(FVector(RegionRequest.Region.Min.X, RegionRequest.Region.Min.Y, 0.0));
+							const FVector2D UVMax = TexturePage->WorldToUv(FVector(RegionRequest.Region.Max.X, RegionRequest.Region.Max.Y, 0.0));
 
-							RenderModule.RequestVirtualTextureTiles(TexturePage->AllocatedTexture, FVector2D::One(), FVector2D::Zero(), FVector2D::One(), UVMin, UVMax, RegionRequest.MipLevel);
+							const FVector2D ClampedUVMin = FVector2D::Clamp(UVMin, FVector2D::Zero(), FVector2D::One());
+							const FVector2D ClampedUVMax = FVector2D::Clamp(UVMax, FVector2D::Zero(), FVector2D::One());
+
+							RenderModule.RequestVirtualTextureTiles(TexturePage->AllocatedTexture, FVector2D::One(), FVector2D::Zero(), FVector2D::One(), ClampedUVMin, ClampedUVMax, RegionRequest.MipLevel);
 						}
 					}
 				}
@@ -1533,24 +1538,38 @@ void UNiagaraDataInterfaceLandscape::VMGenerateRvtFeedback(FVectorVMExternalFunc
 	FNDIInputParam<bool> InIncludeNormal(Context);
 	FNDIInputParam<bool> InEnabled(Context);
 
-	for (int32 InstanceIdx = 0; InstanceIdx < Context.GetNumInstances(); ++InstanceIdx)
+	const bool bHasValidLandscapeBounds = InstanceData->RVTLandscapeBounds.bIsValid;
+	const int32 InstanceCount = Context.GetNumInstances();
+	if (InstanceCount)
 	{
-		const FVector2f ExtentMin = InExtentMin.GetAndAdvance();
-		const FVector2f ExtentMax = InExtentMax.GetAndAdvance();
-		const int32 RequestedMipLevel = InMipLevel.GetAndAdvance();
-		const bool bIncludeBaseColor = InIncludeBaseColor.GetAndAdvance();
-		const bool bIncludeHeight = InIncludeHeight.GetAndAdvance();
-		const bool bIncludeNormal = InIncludeNormal.GetAndAdvance();
-		const bool bEnabled = InEnabled.GetAndAdvance();
+		if (bHasValidLandscapeBounds)
+		{
+			InstanceData->RegionRequests.Reserve(InstanceData->RegionRequests.Num() + InstanceCount);
+		}
 
-		if (bEnabled)
+		for (int32 InstanceIdx = 0; InstanceIdx < Context.GetNumInstances(); ++InstanceIdx)
 		{
-			FRegionRequest& Request = InstanceData->RegionRequests.AddDefaulted_GetRef();
-			Request.Region = FBox2D(FVector2D(ExtentMin), FVector2D(ExtentMax));
-			Request.MipLevel = RequestedMipLevel;
-			Request.bIncludeBaseColor = bIncludeBaseColor;
-			Request.bIncludeHeight = bIncludeHeight;
-			Request.bIncludeNormal = bIncludeNormal;
+			const FVector2f ExtentMin = InExtentMin.GetAndAdvance();
+			const FVector2f ExtentMax = InExtentMax.GetAndAdvance();
+			const int32 RequestedMipLevel = InMipLevel.GetAndAdvance();
+			const bool bIncludeBaseColor = InIncludeBaseColor.GetAndAdvance();
+			const bool bIncludeHeight = InIncludeHeight.GetAndAdvance();
+			const bool bIncludeNormal = InIncludeNormal.GetAndAdvance();
+			const bool bEnabled = InEnabled.GetAndAdvance();
+
+			if (bEnabled && bHasValidLandscapeBounds)
+			{
+				const FBox2D RequestBounds = InstanceData->RVTLandscapeBounds.Overlap(FBox2D(FVector2D(ExtentMin), FVector2D(ExtentMax)));
+				if (RequestBounds.bIsValid)
+				{
+					FRegionRequest& Request = InstanceData->RegionRequests.AddDefaulted_GetRef();
+					Request.Region = RequestBounds;
+					Request.MipLevel = RequestedMipLevel;
+					Request.bIncludeBaseColor = bIncludeBaseColor;
+					Request.bIncludeHeight = bIncludeHeight;
+					Request.bIncludeNormal = bIncludeNormal;
+				}
+			}
 		}
 	}
 }
@@ -1813,6 +1832,7 @@ void UNiagaraDataInterfaceLandscape::ApplyLandscape(const FNiagaraSystemInstance
 	InstanceData.BaseColorVirtualTextureIndex = INDEX_NONE;
 	InstanceData.HeightVirtualTextureIndex = INDEX_NONE;
 	InstanceData.NormalVirtualTextureIndex = INDEX_NONE;
+	InstanceData.RVTLandscapeBounds = FBox2D(EForceInit::ForceInit);
 
 	// only worry about virtual textures if our current platform supports them
 	if (bVirtualTexturesSupported && UseVirtualTexturing(GetFeatureLevelShaderPlatform(SystemInstance.GetFeatureLevel())))
@@ -1865,6 +1885,20 @@ void UNiagaraDataInterfaceLandscape::ApplyLandscape(const FNiagaraSystemInstance
 				}
 			}
 		}
+
+	}
+
+	const bool bHasVirtualTexture = InstanceData.HeightVirtualTextureIndex != INDEX_NONE
+		|| InstanceData.NormalVirtualTextureIndex != INDEX_NONE
+		|| InstanceData.BaseColorVirtualTextureIndex != INDEX_NONE;
+
+	if (bHasVirtualTexture)
+	{
+		FVector LandscapeBoundsOrigin;
+		FVector LandscapeBoundsExtent;
+		Landscape->GetActorBounds(true /*bOnlyCollidingComponents*/, LandscapeBoundsOrigin, LandscapeBoundsExtent);
+
+		InstanceData.RVTLandscapeBounds = FBox2D(FVector2D(LandscapeBoundsOrigin - LandscapeBoundsExtent), FVector2D(LandscapeBoundsOrigin + LandscapeBoundsExtent));
 	}
 
 	// we need to create our own copy of the collision geometry if either the heights are needed, and they're not
