- Fix crash when pasting a curve into a non curve distribution
- The number of channels was not recached which resulted in OOB access when reading the curves

diff --git a/Engine/Plugins/FX/Niagara/Source/NiagaraEditor/Private/Stateless/NiagaraDistributionPropertyCustomization.cpp b/Engine/Plugins/FX/Niagara/Source/NiagaraEditor/Private/Stateless/NiagaraDistributionPropertyCustomization.cpp
--- a/Engine/Plugins/FX/Niagara/Source/NiagaraEditor/Private/Stateless/NiagaraDistributionPropertyCustomization.cpp
+++ b/Engine/Plugins/FX/Niagara/Source/NiagaraEditor/Private/Stateless/NiagaraDistributionPropertyCustomization.cpp
@@ -801,6 +801,7 @@ private:
 		if (NewEditorMode != EditorMode && bUpdatingDistributionModeInternally == false)
 		{
 			EditorMode = NewEditorMode;
+			EditorNumChannels = FNiagaraDistributionEditorUtilities::IsUniform(EditorMode) ? 1 : SourceNumChannels;
 			OnDistributionEditorModeChangedDelegate.Broadcast();
 		}
 	}
