Modular/MergeModules build fixes by moving dllimport-defined functions into cpp files.

diff --git a/Engine/Plugins/FX/Niagara/Source/NiagaraEditor/Private/Widgets/AssetBrowser/SNiagaraAddEmitterToSystemWindow.cpp b/Engine/Plugins/FX/Niagara/Source/NiagaraEditor/Private/Widgets/AssetBrowser/SNiagaraAddEmitterToSystemWindow.cpp
--- a/Engine/Plugins/FX/Niagara/Source/NiagaraEditor/Private/Widgets/AssetBrowser/SNiagaraAddEmitterToSystemWindow.cpp
+++ b/Engine/Plugins/FX/Niagara/Source/NiagaraEditor/Private/Widgets/AssetBrowser/SNiagaraAddEmitterToSystemWindow.cpp
@@ -98,6 +98,8 @@ bool FFrontendFilter_NiagaraSystemEffectType::PassesFilter(const FContentBrowser
 	return false;
 }
 
+SNiagaraAddEmitterToSystemWindow::~SNiagaraAddEmitterToSystemWindow() = default;
+
 void SNiagaraAddEmitterToSystemWindow::Construct(const FArguments& InArgs, const UTaggedAssetBrowserConfiguration& Configuration, TSharedRef<FNiagaraSystemViewModel> SystemViewModel)
 {
 	WeakSystemViewModel = SystemViewModel;
diff --git a/Engine/Plugins/FX/Niagara/Source/NiagaraEditor/Public/Widgets/AssetBrowser/SNiagaraAddEmitterToSystemWindow.h b/Engine/Plugins/FX/Niagara/Source/NiagaraEditor/Public/Widgets/AssetBrowser/SNiagaraAddEmitterToSystemWindow.h
--- a/Engine/Plugins/FX/Niagara/Source/NiagaraEditor/Public/Widgets/AssetBrowser/SNiagaraAddEmitterToSystemWindow.h
+++ b/Engine/Plugins/FX/Niagara/Source/NiagaraEditor/Public/Widgets/AssetBrowser/SNiagaraAddEmitterToSystemWindow.h
@@ -52,7 +52,7 @@ public:
 	SLATE_END_ARGS()
 
 	UE_API void Construct(const FArguments& InArgs, const UTaggedAssetBrowserConfiguration& Configuration, TSharedRef<FNiagaraSystemViewModel> SystemViewModel);
-	UE_API virtual ~SNiagaraAddEmitterToSystemWindow() override {}
+	UE_API virtual ~SNiagaraAddEmitterToSystemWindow() override;
 private:
 	/** The function that will be called by our buttons or by the asset picker itself if double-clicking, hitting enter etc. */
 	void OnAssetsActivatedInternal(const TArray<FAssetData>& AssetData, EAssetTypeActivationMethod::Type) const;
diff --git a/Engine/Plugins/Runtime/MeshModelingToolset/Source/MeshModelingToolsEditorOnly/Private/SkeletalMesh/SkeletalMeshToolsHelper.cpp b/Engine/Plugins/Runtime/MeshModelingToolset/Source/MeshModelingToolsEditorOnly/Private/SkeletalMesh/SkeletalMeshToolsHelper.cpp
--- a/Engine/Plugins/Runtime/MeshModelingToolset/Source/MeshModelingToolsEditorOnly/Private/SkeletalMesh/SkeletalMeshToolsHelper.cpp
+++ b/Engine/Plugins/Runtime/MeshModelingToolset/Source/MeshModelingToolsEditorOnly/Private/SkeletalMesh/SkeletalMeshToolsHelper.cpp
@@ -219,3 +219,5 @@ void SkeletalMeshToolsHelper::FPoseChangeDetector::CheckPose(const TArray<FTrans
 		PreviousMorphTargetWeights = MorphTargetWeights;	
 	}
 }
+
+SkeletalMeshToolsHelper::FPoseChangeDetector::FNotifier& SkeletalMeshToolsHelper::FPoseChangeDetector::GetNotifier() { return Notifier; }
diff --git a/Engine/Plugins/Runtime/MeshModelingToolset/Source/MeshModelingToolsEditorOnly/Public/SkeletalMesh/SkeletalMeshToolsHelper.h b/Engine/Plugins/Runtime/MeshModelingToolset/Source/MeshModelingToolsEditorOnly/Public/SkeletalMesh/SkeletalMeshToolsHelper.h
--- a/Engine/Plugins/Runtime/MeshModelingToolset/Source/MeshModelingToolsEditorOnly/Public/SkeletalMesh/SkeletalMeshToolsHelper.h
+++ b/Engine/Plugins/Runtime/MeshModelingToolset/Source/MeshModelingToolsEditorOnly/Public/SkeletalMesh/SkeletalMeshToolsHelper.h
@@ -66,7 +66,7 @@ namespace SkeletalMeshToolsHelper
 		DECLARE_MULTICAST_DELEGATE_OneParam(FNotifier, FPayload);	
 		
 		UE_API void CheckPose(const TArray<FTransform>& ComponentSpaceTransforms, const TMap<FName,float>& MorphTargetWeights);
-		UE_API FNotifier& GetNotifier() {return Notifier;}
+		UE_API FNotifier& GetNotifier();
 	
 	protected:
 		FNotifier Notifier;
diff --git a/Engine/Plugins/Runtime/RigVM/Source/RigVMDeveloper/Private/RigVMAsset.cpp b/Engine/Plugins/Runtime/RigVM/Source/RigVMDeveloper/Private/RigVMAsset.cpp
--- a/Engine/Plugins/Runtime/RigVM/Source/RigVMDeveloper/Private/RigVMAsset.cpp
+++ b/Engine/Plugins/Runtime/RigVM/Source/RigVMDeveloper/Private/RigVMAsset.cpp
@@ -92,6 +92,8 @@ FRigVMEdGraphDisplaySettings::FRigVMEdGraphDisplaySettings(): bShowNodeInstructi
 {
 }
 
+FRigVMEdGraphDisplaySettings::~FRigVMEdGraphDisplaySettings() = default;
+
 void FRigVMEdGraphDisplaySettings::SetTotalMicroSeconds(double InTotalMicroSeconds)
 {
 	TotalMicroSeconds = AggregateAverage(TotalMicroSecondsFrames, TotalMicroSeconds, InTotalMicroSeconds);
diff --git a/Engine/Plugins/Runtime/RigVM/Source/RigVMDeveloper/Private/RigVMBlueprintLegacy.cpp b/Engine/Plugins/Runtime/RigVM/Source/RigVMDeveloper/Private/RigVMBlueprintLegacy.cpp
--- a/Engine/Plugins/Runtime/RigVM/Source/RigVMDeveloper/Private/RigVMBlueprintLegacy.cpp
+++ b/Engine/Plugins/Runtime/RigVM/Source/RigVMDeveloper/Private/RigVMBlueprintLegacy.cpp
@@ -62,6 +62,8 @@ FEdGraphPinType FRigVMOldPublicFunctionArg::GetPinType() const
 	return RigVMTypeUtils::PinTypeFromExternalVariable(Variable);
 }
 
+FRigVMOldPublicFunctionData::~FRigVMOldPublicFunctionData() = default;
+
 bool FRigVMOldPublicFunctionData::IsMutable() const
 {
 	for(const FRigVMOldPublicFunctionArg& Arg : Arguments)
diff --git a/Engine/Plugins/Runtime/RigVM/Source/RigVMDeveloper/Public/RigVMAsset.h b/Engine/Plugins/Runtime/RigVM/Source/RigVMDeveloper/Public/RigVMAsset.h
--- a/Engine/Plugins/Runtime/RigVM/Source/RigVMDeveloper/Public/RigVMAsset.h
+++ b/Engine/Plugins/Runtime/RigVM/Source/RigVMDeveloper/Public/RigVMAsset.h
@@ -78,6 +78,7 @@ struct UE_API FRigVMEdGraphDisplaySettings
 	GENERATED_BODY();
 
 	FRigVMEdGraphDisplaySettings();
+	~FRigVMEdGraphDisplaySettings();
 
 	// When enabled shows the first node instruction index
 	// matching the execution stack window.
diff --git a/Engine/Plugins/Runtime/RigVM/Source/RigVMDeveloper/Public/RigVMBlueprintLegacy.h b/Engine/Plugins/Runtime/RigVM/Source/RigVMDeveloper/Public/RigVMBlueprintLegacy.h
--- a/Engine/Plugins/Runtime/RigVM/Source/RigVMDeveloper/Public/RigVMBlueprintLegacy.h
+++ b/Engine/Plugins/Runtime/RigVM/Source/RigVMDeveloper/Public/RigVMBlueprintLegacy.h
@@ -67,7 +67,7 @@ struct RIGVMDEVELOPER_API FRigVMOldPublicFunctionData
 		:Name(NAME_None)
 	{}
 
-	virtual ~FRigVMOldPublicFunctionData() {}
+	virtual ~FRigVMOldPublicFunctionData();
 
 	UPROPERTY()
 	FName Name;
diff --git a/Engine/Source/Developer/Zen/Internal/ZenAsyncCbPackageReceiver.h b/Engine/Source/Developer/Zen/Internal/ZenAsyncCbPackageReceiver.h
--- a/Engine/Source/Developer/Zen/Internal/ZenAsyncCbPackageReceiver.h
+++ b/Engine/Source/Developer/Zen/Internal/ZenAsyncCbPackageReceiver.h
@@ -30,9 +30,8 @@ public:
 		FOnComplete&& InOnComplete,
 		int InMaxAttempts = 1);
 
-	ZEN_API const IHttpResponse& GetHttpResponse() { return *Response; }
-	ZEN_API const FCbPackage& GetResponsePackage() { return Package; }
-
+	ZEN_API const IHttpResponse& GetHttpResponse();
+	ZEN_API const FCbPackage& GetResponsePackage();
 
 	ZEN_API void SendAsync();
 
diff --git a/Engine/Source/Developer/Zen/Private/ZenAsyncCbPackageReceiver.cpp b/Engine/Source/Developer/Zen/Private/ZenAsyncCbPackageReceiver.cpp
--- a/Engine/Source/Developer/Zen/Private/ZenAsyncCbPackageReceiver.cpp
+++ b/Engine/Source/Developer/Zen/Private/ZenAsyncCbPackageReceiver.cpp
@@ -26,6 +26,9 @@ FAsyncCbPackageReceiver::FAsyncCbPackageReceiver(
 {
 }
 
+const IHttpResponse& FAsyncCbPackageReceiver::GetHttpResponse() { return *Response; }
+const FCbPackage& FAsyncCbPackageReceiver::GetResponsePackage() { return Package; }
+
 void FAsyncCbPackageReceiver::SendAsync()
 {
 	Request->SendAsync(this, Response);
