- Fix race when allocating ribbon buffers
- This can occur when not visible in the main view but visible in multiple shadow views as they will run in parallel calling GetDynamicMeshElements from multiple threads

#tests Inject checking around callsight to verify if a race can occur or not

diff --git a/Engine/Plugins/FX/Niagara/Source/Niagara/Private/NiagaraRendererRibbons.cpp b/Engine/Plugins/FX/Niagara/Source/Niagara/Private/NiagaraRendererRibbons.cpp
--- a/Engine/Plugins/FX/Niagara/Source/Niagara/Private/NiagaraRendererRibbons.cpp
+++ b/Engine/Plugins/FX/Niagara/Source/Niagara/Private/NiagaraRendererRibbons.cpp
@@ -2399,8 +2399,13 @@ void FNiagaraRendererRibbons::InitializeVertexBuffersResources(FRHICommandListBa
                                                                FGlobalDynamicReadBuffer& DynamicReadBuffer, const TSharedPtr<FNiagaraRibbonRenderingFrameResources>& RenderingResources, bool bShouldUseGPUInit) const
 {
 
-	// Make sure our ribbon data buffers are setup
-	VertexBuffers.InitializeOrUpdateBuffers(RHICmdList, GenerationConfig, DynamicDataRibbon->GenerationOutput, SourceParticleData, DynamicDataRibbon->MaxAllocationCount, bShouldUseGPUInit);
+	{
+		//-OPT: We only need to update this data once for all GDME passes
+		UE::TScopeLock VertexBuffersLock(VertexBuffersGuard);
+
+		// Make sure our ribbon data buffers are setup
+		VertexBuffers.InitializeOrUpdateBuffers(RHICmdList, GenerationConfig, DynamicDataRibbon->GenerationOutput, SourceParticleData, DynamicDataRibbon->MaxAllocationCount, bShouldUseGPUInit);
+	}
 	
 	// Now we need to bind the source particle data, copying it to the gpu if necessary
 	if (DynamicDataRibbon->bIsGPUSystem)
