Fix up a couple more ux issues with the NPC editor
-handling undo around changes to the Source material property of the NPC
-ensure that the editor is refreshed when an NPCI has it's parent cleared

diff --git a/Engine/Plugins/FX/Niagara/Source/Niagara/Classes/NiagaraParameterCollection.h b/Engine/Plugins/FX/Niagara/Source/Niagara/Classes/NiagaraParameterCollection.h
--- a/Engine/Plugins/FX/Niagara/Source/Niagara/Classes/NiagaraParameterCollection.h
+++ b/Engine/Plugins/FX/Niagara/Source/Niagara/Classes/NiagaraParameterCollection.h
@@ -176,9 +176,14 @@ public:
 	//~UObject interface
 #if WITH_EDITORONLY_DATA
 	NIAGARA_API virtual void PreEditChange(FProperty* PropertyAboutToChange) override;
-	NIAGARA_API virtual void PostEditChangeProperty(struct FPropertyChangedEvent& PropertyChangedEvent)override;
+	NIAGARA_API virtual void PostEditChangeProperty(struct FPropertyChangedEvent& PropertyChangedEvent) override;
 	NIAGARA_API virtual void PostDuplicate(EDuplicateMode::Type DuplicateMode) override;
 #endif
+
+#if WITH_EDITOR
+	NIAGARA_API virtual void PreEditUndo() override;
+	NIAGARA_API virtual void PostEditUndo() override;
+#endif
 	//~UObject interface
 
 	NIAGARA_API int32 IndexOfParameter(const FNiagaraVariable& Var);
diff --git a/Engine/Plugins/FX/Niagara/Source/Niagara/Private/NiagaraParameterCollection.cpp b/Engine/Plugins/FX/Niagara/Source/Niagara/Private/NiagaraParameterCollection.cpp
--- a/Engine/Plugins/FX/Niagara/Source/Niagara/Private/NiagaraParameterCollection.cpp
+++ b/Engine/Plugins/FX/Niagara/Source/Niagara/Private/NiagaraParameterCollection.cpp
@@ -678,9 +678,36 @@ void UNiagaraParameterCollection::PostEditChangeProperty(struct FPropertyChanged
 			OnCollectionChangedDelegate.Broadcast();
 		}
 	}
+}
 
+#if WITH_EDITOR
+void UNiagaraParameterCollection::PreEditUndo()
+{
+	Super::PreEditUndo();
+
+	// managing the editor delegates in Pre/Post EditChange breaks down with undo, so if we hit an undo we clean up the delegate
+	if (SourceMaterialCollection && OnMaterialCollectionChanged.IsValid())
+	{
+		SourceMaterialCollection->OnCollectionChanged().Remove(OnMaterialCollectionChanged);
+		OnMaterialCollectionChanged.Reset();
+	}
 }
 
+void UNiagaraParameterCollection::PostEditUndo()
+{
+	Super::PostEditUndo();
+
+	// managing the editor delegates in Pre/Post EditChange breaks down with undo, so if we hit an undo, verify that they are setup
+	if (SourceMaterialCollection && !OnMaterialCollectionChanged.IsValid())
+	{
+		OnMaterialCollectionChanged = SourceMaterialCollection->OnCollectionChanged().AddUObject(this, &UNiagaraParameterCollection::RefreshSourceMaterialCollection);
+		RefreshSourceMaterialCollection();
+		OnCollectionChangedDelegate.Broadcast();
+	}
+}
+#endif
+
+
 void UNiagaraParameterCollection::PostDuplicate(EDuplicateMode::Type DuplicateMode)
 {
 	Super::PostDuplicate(DuplicateMode);
diff --git a/Engine/Plugins/FX/Niagara/Source/NiagaraEditor/Private/ViewModels/NiagaraParameterCollectionAssetViewModel.cpp b/Engine/Plugins/FX/Niagara/Source/NiagaraEditor/Private/ViewModels/NiagaraParameterCollectionAssetViewModel.cpp
--- a/Engine/Plugins/FX/Niagara/Source/NiagaraEditor/Private/ViewModels/NiagaraParameterCollectionAssetViewModel.cpp
+++ b/Engine/Plugins/FX/Niagara/Source/NiagaraEditor/Private/ViewModels/NiagaraParameterCollectionAssetViewModel.cpp
@@ -325,6 +325,7 @@ void FNiagaraParameterCollectionAssetViewModel::RefreshParameterViewModels()
 	if (!Collection)
 	{
 		ParameterViewModels.Empty();
+		OnCollectionChangedDelegate.Broadcast();
 		return;
 	}
 
