Fix issue with lingering NDCs under some flows that don't clean up worlds using either OnWorldBeginTearDown or OnWorldCleanup.

diff --git a/Engine/Plugins/FX/Niagara/Source/Niagara/Private/NiagaraWorldManager.cpp b/Engine/Plugins/FX/Niagara/Source/Niagara/Private/NiagaraWorldManager.cpp
--- a/Engine/Plugins/FX/Niagara/Source/Niagara/Private/NiagaraWorldManager.cpp
+++ b/Engine/Plugins/FX/Niagara/Source/Niagara/Private/NiagaraWorldManager.cpp
@@ -835,6 +835,8 @@ void FNiagaraWorldManager::OnPostWorldCleanup(bool bSessionEnded, bool bCleanupR
 	{
 		ComponentPool->Cleanup(World);
 	}
+
+	DataChannelManager->Cleanup();
 }
 
 void FNiagaraWorldManager::PreGarbageCollect()
