Disallow dynamic NDC refs by default in BP API. Overriding explicitly for the nodes that allow it.
Fixes issue where legacy nodes would incorrectly allow a dynamic ref and crash.

diff --git a/Engine/Plugins/FX/Niagara/Source/NiagaraBlueprintNodes/Internal/K2Node_DataChannelBase.h b/Engine/Plugins/FX/Niagara/Source/NiagaraBlueprintNodes/Internal/K2Node_DataChannelBase.h
--- a/Engine/Plugins/FX/Niagara/Source/NiagaraBlueprintNodes/Internal/K2Node_DataChannelBase.h
+++ b/Engine/Plugins/FX/Niagara/Source/NiagaraBlueprintNodes/Internal/K2Node_DataChannelBase.h
@@ -50,7 +50,7 @@ protected:
 #endif
 
 	/**	Some nodes require and explicit NDC reference to generate derived pins. We ban pin connections for these nodes but allow them by default. */
-	[[nodiscard]] virtual bool SupportsDynamicDataChannel()const { return true; }
+	[[nodiscard]] virtual bool SupportsDynamicDataChannel()const { return false; }
 
 	virtual void AddNDCDerivedPins();
 	virtual void RemoveNDCDerivedPins();
diff --git a/Engine/Plugins/FX/Niagara/Source/NiagaraBlueprintNodes/Internal/K2Node_DataChannel_WithContext.h b/Engine/Plugins/FX/Niagara/Source/NiagaraBlueprintNodes/Internal/K2Node_DataChannel_WithContext.h
--- a/Engine/Plugins/FX/Niagara/Source/NiagaraBlueprintNodes/Internal/K2Node_DataChannel_WithContext.h
+++ b/Engine/Plugins/FX/Niagara/Source/NiagaraBlueprintNodes/Internal/K2Node_DataChannel_WithContext.h
@@ -27,6 +27,8 @@ class UK2Node_WriteDataChannel_WithContext : public UK2Node_DataChannelBase
 	GENERATED_BODY()
 public:
 	UK2Node_WriteDataChannel_WithContext();
+
+	[[nodiscard]] virtual bool SupportsDynamicDataChannel()const override { return true; }
 };
 
 /**
@@ -38,6 +40,8 @@ class UK2Node_ReadDataChannel_WithContext : public UK2Node_DataChannelBase
 	GENERATED_BODY()
 public:
 	UK2Node_ReadDataChannel_WithContext();
+
+	[[nodiscard]] virtual bool SupportsDynamicDataChannel()const override { return true; }
 };
 
 
@@ -51,6 +55,8 @@ class UK2Node_DataChannelGetNum_WithContext : public UK2Node_DataChannelBase
 
 public:
 	UK2Node_DataChannelGetNum_WithContext();
+	
+	[[nodiscard]] virtual bool SupportsDynamicDataChannel()const override { return true; }
 };
 
 /**
@@ -65,9 +71,6 @@ public:
 	UK2Node_WriteDataChannelSingle_WithContext();
 protected:
 	[[nodiscard]] virtual FName GetWriterFunctionName()const { return GET_FUNCTION_NAME_CHECKED(UNiagaraDataChannelLibrary, WriteToNiagaraDataChannel_WithContext); }
-
-	/** The single node cannot support a dynamic channel as it needs an inline channel reference to generate it's pins. */
-	[[nodiscard]] virtual bool SupportsDynamicDataChannel()const override { return false; }
 };
 
 
@@ -83,9 +86,6 @@ public:
 	UK2Node_ReadDataChannelSingle_WithContext();
 protected:
 	[[nodiscard]] virtual FName GetReaderFunctionName()const { return GET_FUNCTION_NAME_CHECKED(UNiagaraDataChannelLibrary, ReadFromNiagaraDataChannel_WithContext); }
-
-	/** The single node cannot support a dynamic channel as it needs an inline channel reference to generate it's pins. */
-	[[nodiscard]] virtual bool SupportsDynamicDataChannel()const override { return false; }
 };
 
 //////////////////////////////////////////////////////////////////////////
diff --git a/Engine/Plugins/FX/Niagara/Source/NiagaraBlueprintNodes/Private/K2Node_ReadDataChannel.cpp b/Engine/Plugins/FX/Niagara/Source/NiagaraBlueprintNodes/Private/K2Node_ReadDataChannel.cpp
--- a/Engine/Plugins/FX/Niagara/Source/NiagaraBlueprintNodes/Private/K2Node_ReadDataChannel.cpp
+++ b/Engine/Plugins/FX/Niagara/Source/NiagaraBlueprintNodes/Private/K2Node_ReadDataChannel.cpp
@@ -51,9 +51,10 @@ void UK2Node_ReadDataChannel::ExpandNode(FKismetCompilerContext& CompilerContext
 {
 	Super::ExpandNode(CompilerContext, SourceGraph);
 	
-	if (!HasValidDataChannel() && !SupportsDynamicDataChannel())
+	check(!SupportsDynamicDataChannel());
+	if (!HasValidDataChannel())
 	{
-		return;
+		return;//Compile error generated in base.
 	}
 
 	const UEdGraphSchema_K2* Schema = CompilerContext.GetSchema();
diff --git a/Engine/Plugins/FX/Niagara/Source/NiagaraBlueprintNodes/Private/K2Node_WriteDataChannel.cpp b/Engine/Plugins/FX/Niagara/Source/NiagaraBlueprintNodes/Private/K2Node_WriteDataChannel.cpp
--- a/Engine/Plugins/FX/Niagara/Source/NiagaraBlueprintNodes/Private/K2Node_WriteDataChannel.cpp
+++ b/Engine/Plugins/FX/Niagara/Source/NiagaraBlueprintNodes/Private/K2Node_WriteDataChannel.cpp
@@ -47,9 +47,10 @@ void UK2Node_WriteDataChannel::ExpandNode(FKismetCompilerContext& CompilerContex
 {
 	Super::ExpandNode(CompilerContext, SourceGraph);
 	
-	if (!HasValidDataChannel() && !SupportsDynamicDataChannel())
+	check(!SupportsDynamicDataChannel());
+	if (!HasValidDataChannel())
 	{
-		return;
+		return;//Compile error generated in base.
 	}
 
 	// create function call node to init the writer object 
