Update tooltip

diff --git a/Engine/Plugins/FX/Niagara/Source/NiagaraBlueprintNodes/Internal/K2Node_DataChannelBase.h b/Engine/Plugins/FX/Niagara/Source/NiagaraBlueprintNodes/Internal/K2Node_DataChannelBase.h
--- a/Engine/Plugins/FX/Niagara/Source/NiagaraBlueprintNodes/Internal/K2Node_DataChannelBase.h
+++ b/Engine/Plugins/FX/Niagara/Source/NiagaraBlueprintNodes/Internal/K2Node_DataChannelBase.h
@@ -49,7 +49,7 @@ protected:
 	FGuid DataChannelVersion;
 #endif
 
-	/**	Some nodes require and explicit NDC reference to generate derived pins. We ban pin connections for these nodes but allow them by default. */
+	/** Some nodes support a dynamic NDC reference via a pin connection. By default this is false and nodes must specify an inline NDC asset reference. */
 	[[nodiscard]] virtual bool SupportsDynamicDataChannel()const { return false; }
 
 	virtual void AddNDCDerivedPins();
