add null check

diff --git a/Engine/Plugins/FX/Niagara/Source/NiagaraBlueprintNodes/Private/K2Node_ReadDataChannel.cpp b/Engine/Plugins/FX/Niagara/Source/NiagaraBlueprintNodes/Private/K2Node_ReadDataChannel.cpp
--- a/Engine/Plugins/FX/Niagara/Source/NiagaraBlueprintNodes/Private/K2Node_ReadDataChannel.cpp
+++ b/Engine/Plugins/FX/Niagara/Source/NiagaraBlueprintNodes/Private/K2Node_ReadDataChannel.cpp
@@ -171,7 +171,11 @@ void UK2Node_ReadDataChannel::ExpandNode(FKismetCompilerContext& CompilerContext
 	UK2Node_IfThenElse* IsValidBranchNode = CompilerContext.SpawnIntermediateNode<UK2Node_IfThenElse>(this, SourceGraph);
 	IsValidBranchNode->AllocateDefaultPins();
 	Schema->TryCreateConnection(IfElseNode->GetThenPin(), IsValidBranchNode->GetExecPin());
-	Schema->TryCreateConnection(LastIsValidPin, IsValidBranchNode->GetConditionPin());
+	
+	if(LastIsValidPin)
+	{
+		Schema->TryCreateConnection(LastIsValidPin, IsValidBranchNode->GetConditionPin());
+	}
 
 	// connect the last exec pins
 	CompilerContext.MovePinLinksToIntermediate(*FailurePin, *IsValidBranchNode->GetElsePin());
