Niagara Rightclick Option on Niagara Asset Tag Definition Assets to transfer previously tagged assets to User Asset Tags
Also added fx.Niagara.MigrateInternalTagsToUserAssetTags command that handles the c++ defined tags within the current CB path

diff --git a/Engine/Plugins/FX/Niagara/Source/NiagaraEditor/Private/NiagaraEditorModule.cpp b/Engine/Plugins/FX/Niagara/Source/NiagaraEditor/Private/NiagaraEditorModule.cpp
--- a/Engine/Plugins/FX/Niagara/Source/NiagaraEditor/Private/NiagaraEditorModule.cpp
+++ b/Engine/Plugins/FX/Niagara/Source/NiagaraEditor/Private/NiagaraEditorModule.cpp
@@ -151,6 +151,7 @@
 
 #include "AssetCompilingManager.h"
 #include "Editor.h"
+#include "FileHelpers.h"
 #include "FileMediaSource.h"
 #include "ILevelSequenceModule.h"
 #include "HAL/PlatformFileManager.h"
@@ -237,6 +238,62 @@ static FAutoConsoleVariableRef CVarEnableCustomInlineDynamicInputFormats(
 	TEXT("If > 0 and experimental inline editors for dynamic input trees are enabled, custom formats which are defined on scripts will be applied.\n"),
 	ECVF_Default);
 
+static FAutoConsoleCommand LogFolderPermissionsCommand = FAutoConsoleCommand(
+	TEXT("fx.Niagara.MigrateInternalTagsToUserAssetTags"),
+	TEXT("Finds and resaves Niagara assets within the current Content Browser view (recursive) that were tagged with an internal Niagara Asset Tag (such as Template, Hidden, Deprecated etc.)"),
+	FConsoleCommandWithArgsDelegate::CreateLambda(
+		[](const TArray<FString>& Args)
+		{
+			FContentBrowserModule& ContentBrowserModule = FModuleManager::Get().LoadModuleChecked<FContentBrowserModule>("ContentBrowser");
+			FContentBrowserItemPath CurrentPath = ContentBrowserModule.Get().GetCurrentPath();
+				
+			if(CurrentPath.HasInternalPath())
+			{
+				UE_LOG(LogNiagaraEditor, Display, TEXT("Trying to find tagged Niagara assets under path %s"), *CurrentPath.GetInternalPathString());
+				
+				FAssetRegistryModule& AssetRegistryModule = FModuleManager::Get().LoadModuleChecked<FAssetRegistryModule>("AssetRegistry");
+				
+				FARFilter Filter;
+				Filter.ClassPaths.Add(UNiagaraEmitter::StaticClass()->GetClassPathName());
+				Filter.ClassPaths.Add(UNiagaraSystem::StaticClass()->GetClassPathName());
+				Filter.PackagePaths.Add(CurrentPath.GetInternalPathName());
+				Filter.bRecursivePaths = true;
+
+PRAGMA_DISABLE_DEPRECATION_WARNINGS
+				for(const FNiagaraAssetTagDefinition* Definition : INiagaraModule::Get().GetInternalAssetTagDefinitions())
+				{
+					Filter.TagsAndValues.Add(FName(Definition->GetGuidAsString()));
+				}
+PRAGMA_ENABLE_DEPRECATION_WARNINGS
+				
+				TArray<FAssetData> OutAssetData;
+				AssetRegistryModule.Get().GetAssets(Filter, OutAssetData);
+
+				TArray<UPackage*> PackagesToResave;
+				Algo::Transform(OutAssetData, PackagesToResave, [](const FAssetData& AssetData)
+				{
+					return AssetData.GetPackage();
+				});
+
+				for(UPackage* Package : PackagesToResave)
+				{
+					Package->MarkPackageDirty();
+				}
+
+				if(PackagesToResave.Num() > 0)
+				{
+					UEditorLoadingAndSavingUtils::SavePackagesWithDialog(PackagesToResave, false);
+				}
+				else
+				{
+					UE_LOG(LogNiagaraEditor, Warning, TEXT("Didn't find any assets to resave under %s"), *CurrentPath.GetInternalPathString());
+				}
+			}
+			else
+			{
+				UE_LOG(LogNiagaraEditor, Warning, TEXT("No valid path found"));	
+			}
+		}));
 // this is required for gpu script compilation ticks
 static FNiagaraShaderQueueTickable NiagaraShaderQueueProcessor;
 
@@ -1185,6 +1242,8 @@ void FNiagaraEditorModule::StartupModule()
 	
 	InitializeTaggedAssetBrowserDetails();
 
+	RegisterMenuExtensions();
+
 	TempPackage = NewObject<UPackage>(nullptr, TEXT("/Temp/NiagaraEditor"), RF_Transient);
 	TempPackage->AddToRoot();
 
@@ -2723,6 +2782,86 @@ void FNiagaraEditorModule::RegisterDynamicCustomizations()
 	}
 }
 
+void FNiagaraEditorModule::RegisterMenuExtensions()
+{
+	RegisterMenuExtension_MigrateToUserAssetTags();
+}
+
+void FNiagaraEditorModule::RegisterMenuExtension_MigrateToUserAssetTags()
+{
+	UToolMenus::RegisterStartupCallback(FSimpleMulticastDelegate::FDelegate::CreateLambda([]()
+	{
+		UToolMenu* MenuExtension = UToolMenus::Get()->ExtendMenu("ContentBrowser.AssetContextMenu.NiagaraAssetTagDefinitions");
+		FToolMenuSection& Section = MenuExtension->FindOrAddSection("GetAssetActions");
+		Section.AddDynamicEntry("MigrateToUserAssetTagsDynamic", FNewToolMenuSectionDelegate::CreateLambda([](FToolMenuSection& InSection)
+			{
+				UContentBrowserAssetContextMenuContext* Context = InSection.FindContext<UContentBrowserAssetContextMenuContext>();
+				if (Context)
+				{
+					if(Context->SelectedAssets.Num() > 0)
+					{
+						InSection.AddMenuEntry(
+							"MigrateToUserAssetTags",
+							LOCTEXT("MigrateToUserAssetTags_Label", "Migrate Tagged Assets to User Asset Tags"),
+							LOCTEXT("MigrateToUserAssetTags_Tooltip",
+								"Migrate the assets that were tagged using tags from this Tag Definitions assets to the new User Asset Tag system."),
+							FSlateIcon(FAppStyle::GetAppStyleSetName(), "Icons.Tag"),
+							FExecuteAction::CreateLambda([InSelectedAssets = Context->SelectedAssets]()
+								{
+									TArray<UObject*> SelectedObjects;
+									SelectedObjects.Reserve(InSelectedAssets.Num());
+
+									for(const FAssetData& Asset : InSelectedAssets)
+									{
+										if(UObject* LoadedAsset = Asset.GetAsset())
+										{
+											SelectedObjects.Add(LoadedAsset);
+										}
+									}
+
+									FAssetRegistryModule& AssetRegistryModule = FModuleManager::Get().LoadModuleChecked<FAssetRegistryModule>("AssetRegistry");
+									FARFilter Filter;
+									Filter.ClassPaths.Add(UNiagaraEmitter::StaticClass()->GetClassPathName());
+									Filter.ClassPaths.Add(UNiagaraSystem::StaticClass()->GetClassPathName());
+									
+									for(UObject* SelectedObject : SelectedObjects)
+									{
+										UDEPRECATED_NiagaraAssetTagDefinitions* DefinitionsAsset = Cast<UDEPRECATED_NiagaraAssetTagDefinitions>(SelectedObject);
+										for(const FNiagaraAssetTagDefinition& Definition : DefinitionsAsset->GetAssetTagDefinitions())
+										{
+											Filter.TagsAndValues.Add(FName(Definition.GetGuidAsString()));
+										}
+									}
+
+									TArray<FAssetData> OutAssetData;
+									AssetRegistryModule.Get().GetAssets(Filter, OutAssetData);
+
+									TArray<UPackage*> PackagesToResave;
+									Algo::Transform(OutAssetData, PackagesToResave, [](const FAssetData& AssetData)
+									{
+										return AssetData.GetPackage();
+									});
+
+									for(UPackage* Package : PackagesToResave)
+									{
+										Package->MarkPackageDirty();
+									}
+									
+									if(PackagesToResave.Num() > 0)
+									{
+										UEditorLoadingAndSavingUtils::SavePackagesWithDialog(PackagesToResave, false);
+									}
+									else
+									{
+										UE_LOG(LogNiagaraEditor, Warning, TEXT("Didn't find any assets to resave"));
+									}
+								}));
+					}
+				}
+			}));
+	}));
+}
+
 FNiagaraEditorModule::FOnScriptApplied& FNiagaraEditorModule::OnScriptApplied()
 {
 	return OnScriptAppliedDelegate;
diff --git a/Engine/Plugins/FX/Niagara/Source/NiagaraEditor/Public/NiagaraEditorModule.h b/Engine/Plugins/FX/Niagara/Source/NiagaraEditor/Public/NiagaraEditorModule.h
--- a/Engine/Plugins/FX/Niagara/Source/NiagaraEditor/Public/NiagaraEditorModule.h
+++ b/Engine/Plugins/FX/Niagara/Source/NiagaraEditor/Public/NiagaraEditorModule.h
@@ -462,6 +462,10 @@ private:
 	void OnModulesChanged(FName Module, EModuleChangeReason Reason);
 
 	void RegisterDynamicCustomizations();
+	
+	void RegisterMenuExtensions();
+	
+	void RegisterMenuExtension_MigrateToUserAssetTags();
 
 private:
 	TSharedPtr<FExtensibilityManager> MenuExtensibilityManager;
