Fix for check firing if you modify a physics asset for a mesh while the mesh is being used by the rigid body DI
#tests Deleted a physics body from a skel mesh while a niagara simulation was running with actors using that skel mesh

diff --git a/Engine/Plugins/FX/Niagara/Source/Niagara/Private/NiagaraDataInterfaceRigidMeshCollisionQuery.cpp b/Engine/Plugins/FX/Niagara/Source/Niagara/Private/NiagaraDataInterfaceRigidMeshCollisionQuery.cpp
--- a/Engine/Plugins/FX/Niagara/Source/Niagara/Private/NiagaraDataInterfaceRigidMeshCollisionQuery.cpp
+++ b/Engine/Plugins/FX/Niagara/Source/Niagara/Private/NiagaraDataInterfaceRigidMeshCollisionQuery.cpp
@@ -576,12 +576,11 @@ void RemapPreviousTransforms(
 		{
 			const FNDIRigidMeshCollisionData::FComponentBodyCount& PreviousCount = PreviousCounts[PreviousIndex];
 
-			if (PreviousCount.ComponentHash == CurrentCount.ComponentHash)
+			if ((PreviousCount.ComponentHash == CurrentCount.ComponentHash)
+				&& (PreviousCount.BoxCount == CurrentCount.BoxCount)
+				&& (PreviousCount.SphereCount == CurrentCount.SphereCount)
+				&& (PreviousCount.CapsuleCount == CurrentCount.CapsuleCount))
 			{
-				check(PreviousCount.BoxCount == CurrentCount.BoxCount);
-				check(PreviousCount.SphereCount == CurrentCount.SphereCount);
-				check(PreviousCount.CapsuleCount == CurrentCount.CapsuleCount);
-
 				bPreviousValuesCopied = true;
 
 				// For the body in the current tick, we found a match in the previous tick transforms.  Copy the previous tick transforms (stored in CurrentTransform)
