Added TODO comments for a bad practice workaround

#rnx

diff --git a/Engine/Plugins/FX/Niagara/Source/NiagaraEditor/Private/NiagaraEmitterFactoryNew.cpp b/Engine/Plugins/FX/Niagara/Source/NiagaraEditor/Private/NiagaraEmitterFactoryNew.cpp
--- a/Engine/Plugins/FX/Niagara/Source/NiagaraEditor/Private/NiagaraEmitterFactoryNew.cpp
+++ b/Engine/Plugins/FX/Niagara/Source/NiagaraEditor/Private/NiagaraEmitterFactoryNew.cpp
@@ -71,6 +71,8 @@ bool UNiagaraEmitterFactoryNew::ConfigureProperties()
 			.OnAdditionalShouldFilterAsset_UObject(this, &UNiagaraEmitterFactoryNew::OnAdditionalShouldFilterAsset)
 			.OnGetExtraFrontendFilters_UObject(this, &UNiagaraEmitterFactoryNew::OnGetExtraFrontendFilters)
 			.OnExtendAddFilterMenu_UObject(this, &UNiagaraEmitterFactoryNew::OnExtendAddFilterMenu)
+			// @TODO (ME): This is not ideal. Ideally the factory would have access to creation path and/or property handle (if created from property asset picker -> create menu)
+			// then we wouldn't need to use GetLastDirectory which is relatively unsafe
 			.AdditionalReferencingAssets({ FAssetData("TmpAsset", FName(FEditorDirectories::Get().GetLastDirectory(ELastDirectory::NEW_ASSET)), "TmpAsset", UObject::StaticClass()->GetClassPathName())});
 
 		SWindow::FArguments WindowArgs;
diff --git a/Engine/Plugins/FX/Niagara/Source/NiagaraEditor/Private/NiagaraSystemFactoryNew.cpp b/Engine/Plugins/FX/Niagara/Source/NiagaraEditor/Private/NiagaraSystemFactoryNew.cpp
--- a/Engine/Plugins/FX/Niagara/Source/NiagaraEditor/Private/NiagaraSystemFactoryNew.cpp
+++ b/Engine/Plugins/FX/Niagara/Source/NiagaraEditor/Private/NiagaraSystemFactoryNew.cpp
@@ -66,6 +66,8 @@ bool UNiagaraSystemFactoryNew::ConfigureProperties()
 			.RecentAndFavoritesList(FNiagaraEditorModule::Get().GetRecentsManager()->GetRecentEmitterAndSystemsList())
 			.OnGetExtraFrontendFilters_UObject(this, &UNiagaraSystemFactoryNew::OnGetExtraFrontendFilters)
 			.OnExtendAddFilterMenu_UObject(this, &UNiagaraSystemFactoryNew::OnExtendAddFilterMenu)
+			// @TODO (ME): This is not ideal. Ideally the factory would have access to creation path and/or property handle (if created from property asset picker -> create menu)
+			// then we wouldn't need to use GetLastDirectory which is relatively unsafe
 			.AdditionalReferencingAssets({ FAssetData("TmpAsset", FName(FEditorDirectories::Get().GetLastDirectory(ELastDirectory::NEW_ASSET)), "TmpAsset", UObject::StaticClass()->GetClassPathName())});
 
 		SWindow::FArguments WindowArgs;
