Provide asset owner context to the Niagara user parameters asset picker so it can do proper asset reference filtering.

diff --git a/Engine/Plugins/FX/Niagara/Source/NiagaraEditor/Private/Customizations/UserParameters/NiagaraComponentUserParametersBuilder.cpp b/Engine/Plugins/FX/Niagara/Source/NiagaraEditor/Private/Customizations/UserParameters/NiagaraComponentUserParametersBuilder.cpp
--- a/Engine/Plugins/FX/Niagara/Source/NiagaraEditor/Private/Customizations/UserParameters/NiagaraComponentUserParametersBuilder.cpp
+++ b/Engine/Plugins/FX/Niagara/Source/NiagaraEditor/Private/Customizations/UserParameters/NiagaraComponentUserParametersBuilder.cpp
@@ -437,6 +437,7 @@ void FNiagaraComponentUserParametersNodeBuilder::GenerateRowForUserParameterInte
 			.DisplayBrowse(true)
 			.DisplayThumbnail(true)
 			.NewAssetFactories(TArray<UFactory*>())
+			.OwnerAssetDataArray({ FAssetData(SystemAsset) })
 		];
 		
 		AddCustomMenuActionsForParameter(DetailWidgetRow, UserParameter);		
diff --git a/Engine/Plugins/FX/Niagara/Source/NiagaraEditor/Private/Customizations/UserParameters/NiagaraSystemUserParametersBuilder.cpp b/Engine/Plugins/FX/Niagara/Source/NiagaraEditor/Private/Customizations/UserParameters/NiagaraSystemUserParametersBuilder.cpp
--- a/Engine/Plugins/FX/Niagara/Source/NiagaraEditor/Private/Customizations/UserParameters/NiagaraSystemUserParametersBuilder.cpp
+++ b/Engine/Plugins/FX/Niagara/Source/NiagaraEditor/Private/Customizations/UserParameters/NiagaraSystemUserParametersBuilder.cpp
@@ -264,6 +264,7 @@ void FNiagaraSystemUserParameterBuilder::GenerateRowForUserParameterInternal(IDe
 			.DisplayBrowse(true)
 			.DisplayThumbnail(true)
 			.NewAssetFactories(TArray<UFactory*>())
+			.OwnerAssetDataArray({ FAssetData(SystemAsset) })
 		];
 
 		AddCustomMenuActionsForParameter(DetailWidgetRow, UserParameter);
