Remove logic which creates the blueprint action database before the blueprint editor is opened, the blueprint action database has overhead and causes high cost for asset discovery and PIE. This logic also does not need to register early, its actions can be discovered on demand. Confirmed by dumping the node graph database with and without the change.

#rnx

diff --git a/Engine/Plugins/FX/Niagara/Source/NiagaraBlueprintNodes/Private/NiagaraBlueprintNodesModule.cpp b/Engine/Plugins/FX/Niagara/Source/NiagaraBlueprintNodes/Private/NiagaraBlueprintNodesModule.cpp
--- a/Engine/Plugins/FX/Niagara/Source/NiagaraBlueprintNodes/Private/NiagaraBlueprintNodesModule.cpp
+++ b/Engine/Plugins/FX/Niagara/Source/NiagaraBlueprintNodes/Private/NiagaraBlueprintNodesModule.cpp
@@ -24,13 +24,7 @@ private:
 	/** Callback for changes to other modules. Allows us to refresh certain things to reflect new types or changes etc. */
 	void OnModulesChanged(FName Module, EModuleChangeReason Reason);
 
-	/** In response to certain events like modules being loaded, we must ensure our BP actions are properly setup. */
-	void UpdateBPHooks();
-
 	FDelegateHandle OnModulesChangedHandle;
-
-	/** Set of struts with custom BP nodes/logic that require a refresh of the BP action DB. */
-	TSet<TWeakObjectPtr<UScriptStruct>> InitializedCustomBPStructs;
 };
 
 
@@ -43,9 +37,6 @@ void FNiagaraBlueprintNodesModule::StartupModule()
  	PropertyModule.RegisterCustomClassLayout(
  		UK2Node_DataChannelAccessContextOperation::StaticClass()->GetFName(),
  		FOnGetDetailCustomizationInstance::CreateStatic(&FNDCAccessContextOperationNodeDetailsDetails::MakeInstance));
-		
-		
-	UpdateBPHooks();
 }
 
 void FNiagaraBlueprintNodesModule::ShutdownModule()
@@ -55,50 +46,6 @@ void FNiagaraBlueprintNodesModule::ShutdownModule()
 
 void FNiagaraBlueprintNodesModule::OnModulesChanged(FName Module, EModuleChangeReason Reason)
 {
-	//Refresh any BP actions to accommodate newly loaded classes.
-	if (Reason == EModuleChangeReason::ModuleLoaded)
-	{
-		UpdateBPHooks();
-	}
-}
-
-void FNiagaraBlueprintNodesModule::UpdateBPHooks()
-{
-	if (FModuleManager::Get().IsModuleLoaded(TEXT("Kismet")) && !IsRunningCommandlet())
-	{
-		for(auto It = InitializedCustomBPStructs.CreateIterator(); It; ++It)
-		{
-			if(It->IsValid() == false)
-			{
-				It.RemoveCurrent();
-			}
-		}
-
-		//Disable all make/break nodes for NDC Context classes
-		bool bRefresh = false;
-		for (TObjectIterator<UScriptStruct> It; It; ++It)
-		{
-			UScriptStruct* Struct = *It;
-			if (Struct->IsChildOf(FNDCAccessContextBase::StaticStruct()))
-			{
-				if(InitializedCustomBPStructs.Contains(Struct) == false)
-				{
-					bRefresh = true;	
-					InitializedCustomBPStructs.Add(Struct);
-				}		
-			}
-		}
-		
-		if(bRefresh)
-		{
-			FBlueprintActionDatabase& ActionDB = FBlueprintActionDatabase::Get();
-			ActionDB.RefreshClassActions(UK2Node_DataChannelAccessContextOperation::StaticClass());
-			ActionDB.RefreshClassActions(UK2Node_DataChannelAccessContext_Make::StaticClass());
-			ActionDB.RefreshClassActions(UK2Node_DataChannelAccessContext_GetMembers::StaticClass());
-			ActionDB.RefreshClassActions(UK2Node_DataChannelAccessContext_SetMembers::StaticClass());
-			UBlueprintEditorLibrary::RefreshAllOpenBlueprintEditors();
-		}
-	}
 }
 
 
