Fixed bug with garbage world managers being left in the FGCObjet lists.

diff --git a/Engine/Plugins/FX/Niagara/Source/Niagara/Public/NiagaraWorldManager.h b/Engine/Plugins/FX/Niagara/Source/Niagara/Public/NiagaraWorldManager.h
--- a/Engine/Plugins/FX/Niagara/Source/Niagara/Public/NiagaraWorldManager.h
+++ b/Engine/Plugins/FX/Niagara/Source/Niagara/Public/NiagaraWorldManager.h
@@ -98,7 +98,7 @@ class FNiagaraWorldManager : public FGCObject
 
 public:
 	FNiagaraWorldManager();
-	~FNiagaraWorldManager();
+	virtual ~FNiagaraWorldManager() override;
 
 	void Init(UWorld* InWorld);
 
