Speculative fix for lingering Cull Proxies after world cleanup

diff --git a/Engine/Plugins/FX/Niagara/Source/Niagara/Private/NiagaraWorldManager.cpp b/Engine/Plugins/FX/Niagara/Source/Niagara/Private/NiagaraWorldManager.cpp
--- a/Engine/Plugins/FX/Niagara/Source/Niagara/Private/NiagaraWorldManager.cpp
+++ b/Engine/Plugins/FX/Niagara/Source/Niagara/Private/NiagaraWorldManager.cpp
@@ -43,6 +43,20 @@ DECLARE_CYCLE_STAT(TEXT("Niagara Manager Wait Pre Garbage Collect [GT]"), STAT_N
 DECLARE_CYCLE_STAT(TEXT("Niagara Manager Refresh Owner Allows Scalability"), STAT_NiagaraWorldManRefreshOwnerAllowsScalability, STATGROUP_Niagara);
 DECLARE_CYCLE_STAT(TEXT("Niagara Manager Tick Parameter Collections [GT]"), STAT_NiagaraWorldManTickParamCollections, STATGROUP_Niagara);
 
+
+namespace NiagaraRollbackCVars
+{
+	//Rollback gate for new world cleanup flow.
+	//Scheduled for removal in UE-347574
+	bool NewWorldCleanup = true;
+	static FAutoConsoleVariableRef CVarRollback_NewWorldCleanup(
+		TEXT("fx.Niagara.Rollback.NewWorldCleanup"),
+		NewWorldCleanup,
+		TEXT("When true we'll use the new World Manager Cleanup flows. False will use old legacy code."),
+		ECVF_Default
+	);
+}
+
 bool GNiagaraWorldManagerUseTickPasses = false;
 static FAutoConsoleVariableRef CVarNiagaraWorldManagerUseTickPasses(
 	TEXT("fx.Niagara.WorldManager.UseTickPasses"),
@@ -566,7 +580,10 @@ FString FNiagaraWorldManager::GetReferencerName() const
 UNiagaraCullProxyComponent* FNiagaraWorldManager::GetCullProxy(UNiagaraComponent* Component)
 {
 	UNiagaraSystem* System = Component->GetAsset();
-	if (ensure(System) && GNiagaraAllowCullProxies)
+
+	//Don't create or return the cull proxy if the world is cleaning up.
+	bool bIsWorldCleaningUp = NiagaraRollbackCVars::NewWorldCleanup && (bIsTearingDown || bIsCleaningUp);
+	if (ensure(System) && GNiagaraAllowCullProxies && bIsWorldCleaningUp == false)
 	{
 		auto& CullProxy = CullProxyMap.FindOrAdd(System);
 
@@ -772,7 +789,104 @@ void FNiagaraWorldManager::OnComputeDispatchInterfaceDestroyed_Internal(FNiagara
 	DeferredDeletionQueue.Empty();
 }
 
+void FNiagaraWorldManager::CleanupTransientObjectsAndManagers()
+{
+	//TODO: Move as much of cleanup as possible in here.
+	//Some world cleanup flows don't call BeginTearDown or even OnWorldCleanup sometimes so anything that could be cleaned up immediately on teardown should go in here.
+	if (ComponentPool)
+	{
+		ComponentPool->Cleanup(World);
+	}
+	DataChannelManager->Cleanup();
+
+	for (auto it = CullProxyMap.CreateIterator(); it; ++it)
+	{
+		if (it->Value)
+		{
+			it->Value->DestroyComponent();
+		}
+	}
+	CullProxyMap.Empty();
+
+	ScalabilityManagers.Empty();
+}
+
 void FNiagaraWorldManager::OnWorldBeginTearDown()
+{
+	bIsTearingDown = true;
+	if(NiagaraRollbackCVars::NewWorldCleanup)
+	{
+		CleanupTransientObjectsAndManagers();
+	}
+	else
+	{
+		LegacyOnWorldBeginTearDown();
+	}
+}
+
+void FNiagaraWorldManager::OnWorldCleanup(bool bSessionEnded, bool bCleanupResources)
+{
+	bIsCleaningUp = true;
+	if(NiagaraRollbackCVars::NewWorldCleanup)
+	{
+		DeferredMethods.ExecuteAndClear();
+	
+		CleanupTransientObjectsAndManagers();
+
+		for (int TG = 0; TG < NiagaraNumTickGroups; ++TG)
+		{
+			for (int iTickPass = 0; iTickPass < int(ENiagaraWorldManagerTickPass::Num); ++iTickPass)
+			{
+				for (TPair<UNiagaraSystem*, FNiagaraSystemSimulationRef>& SimPair : SystemSimulations[TG][iTickPass])
+				{
+					SimPair.Value->Destroy();
+				}
+				SystemSimulations[TG][iTickPass].Empty();
+			}
+		}
+		CleanupParameterCollections();
+
+		SimulationsWithPostActorWork.Empty();
+		SimulationsWithEndOfFrameWait.Empty();
+
+		DeferredDeletionQueue.Empty();
+
+		DIGeneratedData.Empty();
+
+	#if WITH_EDITOR
+		for (FSystemDelegateMap::TConstIterator DelegateIt(SystemPostChangeDelegateHandles); DelegateIt; ++DelegateIt)
+		{
+			if (UNiagaraSystem* System = DelegateIt.Key().Get())
+			{
+				System->OnSystemPostEditChange().Remove(DelegateIt.Value());
+			}
+		}
+		SystemPostChangeDelegateHandles.Empty();
+	#endif
+	}
+	else
+	{
+		LegacyOnWorldCleanup(bSessionEnded, bCleanupResources);
+	}
+}
+
+void FNiagaraWorldManager::OnPostWorldCleanup(bool bSessionEnded, bool bCleanupResources)
+{	
+	if(NiagaraRollbackCVars::NewWorldCleanup)
+	{
+		//Ensure we've already cleaned everything up. In some cases OnWorldCleanup doesn't appear to be called by the world cleanup flows.
+		OnWorldCleanup(bSessionEnded, bCleanupResources);
+	}
+	else
+	{
+		LegacyOnPostWorldCleanup(bSessionEnded, bCleanupResources);
+	}
+}
+
+//////////////////////////////////////////////////////////////////////////
+//Legacy cleanup code.
+//Scheduled for removal in UE-347574
+void FNiagaraWorldManager::LegacyOnWorldBeginTearDown()
 {
 	bIsTearingDown = true;
 	if (ComponentPool)
@@ -782,7 +896,7 @@ void FNiagaraWorldManager::OnWorldBeginTearDown()
 	DataChannelManager->Cleanup();
 }
 
-void FNiagaraWorldManager::OnWorldCleanup(bool bSessionEnded, bool bCleanupResources)
+void FNiagaraWorldManager::LegacyOnWorldCleanup(bool bSessionEnded, bool bCleanupResources)
 {
 	bIsCleaningUp = true;
 
@@ -829,7 +943,7 @@ void FNiagaraWorldManager::OnWorldCleanup(bool bSessionEnded, bool bCleanupResou
 #endif
 }
 
-void FNiagaraWorldManager::OnPostWorldCleanup(bool bSessionEnded, bool bCleanupResources)
+void FNiagaraWorldManager::LegacyOnPostWorldCleanup(bool bSessionEnded, bool bCleanupResources)
 {
 	if (ComponentPool)
 	{
@@ -839,6 +953,8 @@ void FNiagaraWorldManager::OnPostWorldCleanup(bool bSessionEnded, bool bCleanupR
 	DataChannelManager->Cleanup();
 }
 
+//////////////////////////////////////////////////////////////////////////
+
 void FNiagaraWorldManager::PreGarbageCollect()
 {
 	if (GNiagaraWaitOnPreGC)
diff --git a/Engine/Plugins/FX/Niagara/Source/Niagara/Public/NiagaraWorldManager.h b/Engine/Plugins/FX/Niagara/Source/Niagara/Public/NiagaraWorldManager.h
--- a/Engine/Plugins/FX/Niagara/Source/Niagara/Public/NiagaraWorldManager.h
+++ b/Engine/Plugins/FX/Niagara/Source/Niagara/Public/NiagaraWorldManager.h
@@ -159,6 +159,11 @@ public:
 	void OnWorldBeginTearDown();
 	void OnWorldCleanup(bool bSessionEnded, bool bCleanupResources);
 	void OnPostWorldCleanup(bool bSessionEnded, bool bCleanupResources);
+	
+	//TODO: Remove - UE-347574
+	void LegacyOnWorldBeginTearDown();
+	void LegacyOnWorldCleanup(bool bSessionEnded, bool bCleanupResources);
+	void LegacyOnPostWorldCleanup(bool bSessionEnded, bool bCleanupResources);
 
 	void PreGarbageCollect();
 	void PostReachabilityAnalysis();
@@ -241,6 +246,8 @@ private:
 
 	bool HasActiveWorld() const;
 
+	void CleanupTransientObjectsAndManagers();
+
 public:
 	void SetDebugPlaybackMode(ENiagaraDebugPlaybackMode Mode) { RequestedDebugPlaybackMode = Mode; }
 	ENiagaraDebugPlaybackMode GetDebugPlaybackMode() const { return DebugPlaybackMode; }
