Speculative fix for scalabiltiy manager crash that cannot be reproduced.

diff --git a/Engine/Plugins/FX/Niagara/Source/Niagara/Private/NiagaraScalabilityManager.cpp b/Engine/Plugins/FX/Niagara/Source/Niagara/Private/NiagaraScalabilityManager.cpp
--- a/Engine/Plugins/FX/Niagara/Source/Niagara/Private/NiagaraScalabilityManager.cpp
+++ b/Engine/Plugins/FX/Niagara/Source/Niagara/Private/NiagaraScalabilityManager.cpp
@@ -298,6 +298,13 @@ void FNiagaraScalabilityManager::ProcessSignificance(UNiagaraSignificanceHandler
 
 	int32 EffectTypeActiveInstances = 0;
 
+	struct FCullProxyCreateDestroyTask
+	{
+		UNiagaraComponent* Component = nullptr;
+		bool bCreate = false;
+	};
+
+	TArray<FCullProxyCreateDestroyTask, TInlineAllocator<64>> CullProxyTasks;
 	for (int32 SortedIt = 0; SortedIt < Context.SignificanceIndices.Num(); ++SortedIt)
 	{
 		int32 SortedIdx = Context.SignificanceIndices[SortedIt];
@@ -328,9 +335,12 @@ void FNiagaraScalabilityManager::ProcessSignificance(UNiagaraSignificanceHandler
 			Context.bHasDirtyState |= CompState.IsDirty();
 		}
 
-		//Handle cull proxy creation here so that it can be done in significance order
+		//Gather list of cull proxies to create/destroy. Move outside of this loop as in some code paths we can have components register/unregister and invalidate the SortedIdx.
 		if (ScalabilitySettings.CullProxyMode != ENiagaraCullProxyMode::None)
 		{
+			FCullProxyCreateDestroyTask& NewCullProxyTask = CullProxyTasks.AddDefaulted_GetRef();			
+			NewCullProxyTask.Component = Component;			
+			NewCullProxyTask.bCreate = false;
 			if (CompState.bCulled)
 			{
 				bool bEnableCullProxy = SysData.CullProxyCount < ScalabilitySettings.MaxSystemProxies
@@ -339,17 +349,25 @@ void FNiagaraScalabilityManager::ProcessSignificance(UNiagaraSignificanceHandler
 				if (bEnableCullProxy)
 				{
 					++SysData.CullProxyCount;
-					Component->CreateCullProxy(true);
-				}
-				else
-				{
-					Component->DestroyCullProxy();
+					NewCullProxyTask.bCreate = true;
 				}
 			}
+		}
+	}
+
+	//Process any cull proxy creation / destruction here so as to not invalidate the ManagedComponents/SortedIdx as we process.
+	for(FCullProxyCreateDestroyTask& CullProxyTask : CullProxyTasks)
+	{
+		if(IsValid(CullProxyTask.Component))
+		{
+			if (CullProxyTask.bCreate)
+			{
+				CullProxyTask.Component->CreateCullProxy(true);
+			}
 			else
 			{
-				Component->DestroyCullProxy();
-			}		
+				CullProxyTask.Component->DestroyCullProxy();
+			}
 		}
 	}
 }
