A couple of fixes for GC issues with UWorlds

diff --git a/Engine/Plugins/FX/Niagara/Source/Niagara/Private/NiagaraDataChannelManager.cpp b/Engine/Plugins/FX/Niagara/Source/Niagara/Private/NiagaraDataChannelManager.cpp
--- a/Engine/Plugins/FX/Niagara/Source/Niagara/Private/NiagaraDataChannelManager.cpp
+++ b/Engine/Plugins/FX/Niagara/Source/Niagara/Private/NiagaraDataChannelManager.cpp
@@ -46,6 +46,7 @@ void FNiagaraDataChannelManager::RefreshDataChannels()
 
 void FNiagaraDataChannelManager::Init()
 {
+	bIsInitialized = true;
 	bIsCleanedUp = false;
 
 	//Initialize any existing data channels, more may be initialized later as they are loaded.
@@ -63,6 +64,7 @@ void FNiagaraDataChannelManager::Cleanup()
 	}
 	Channels.Empty();
 	bIsCleanedUp = true;
+	bIsInitialized = false;
 }
 
 void FNiagaraDataChannelManager::BeginFrame(float DeltaSeconds)
@@ -73,6 +75,11 @@ void FNiagaraDataChannelManager::BeginFrame(float DeltaSeconds)
 		SCOPE_CYCLE_COUNTER(STAT_DataChannelManager_BeginFrame);
 		CSV_SCOPED_TIMING_STAT(Particles, CoreSystems_NiagaraDataChannelManager);
 
+		if(NeedsInit())
+		{
+			Init();
+		}
+
 		//Tick all DataChannel channel handlers.
 		for (auto It = Channels.CreateIterator(); It; ++It)
 		{
@@ -167,6 +174,11 @@ UNiagaraDataChannelHandler* FNiagaraDataChannelManager::InitDataChannel(const UN
 	UWorld* World = GetWorld();
 	if(INiagaraModule::DataChannelsEnabled() && World && !World->IsNetMode(NM_DedicatedServer) && InChannel->IsValid())
 	{
+		if(NeedsInit())
+		{
+			Init();
+		}
+
 		TObjectPtr<UNiagaraDataChannelHandler>& Handler = Channels.FindOrAdd(InChannel);
 
 		if (bForce || Handler == nullptr)
diff --git a/Engine/Plugins/FX/Niagara/Source/Niagara/Private/NiagaraDataChannelManager.h b/Engine/Plugins/FX/Niagara/Source/Niagara/Private/NiagaraDataChannelManager.h
--- a/Engine/Plugins/FX/Niagara/Source/Niagara/Private/NiagaraDataChannelManager.h
+++ b/Engine/Plugins/FX/Niagara/Source/Niagara/Private/NiagaraDataChannelManager.h
@@ -49,7 +49,10 @@ public:
 private:
 	FNiagaraWorldManager* WorldMan = nullptr;
 
+	bool NeedsInit()const { return !bIsCleanedUp && !bIsInitialized; }
+
 	/** Runtime handlers for each DataChannel channel. */
 	TMap<TWeakObjectPtr<const UNiagaraDataChannel>, TObjectPtr<UNiagaraDataChannelHandler>> Channels;
 	bool bIsCleanedUp = false;
+	bool bIsInitialized = false;
 };
diff --git a/Engine/Plugins/FX/Niagara/Source/Niagara/Private/NiagaraWorldManager.cpp b/Engine/Plugins/FX/Niagara/Source/Niagara/Private/NiagaraWorldManager.cpp
--- a/Engine/Plugins/FX/Niagara/Source/Niagara/Private/NiagaraWorldManager.cpp
+++ b/Engine/Plugins/FX/Niagara/Source/Niagara/Private/NiagaraWorldManager.cpp
@@ -55,6 +55,27 @@ namespace NiagaraRollbackCVars
 		TEXT("When true we'll use the new World Manager Cleanup flows. False will use old legacy code."),
 		ECVF_Default
 	);
+	
+	//Rollback gate for world initialization changes.
+	//Scheduled for removal in UE-348784
+	bool WorldManagerRespectsWorldCreateFXSystem = false;
+	static FAutoConsoleVariableRef CVarRollback_WorldManagerRespectsWorldCreateFXSystem(
+		TEXT("fx.Niagara.Rollback.WorldManagerRespectsWorldCreateFXSystem"),
+		WorldManagerRespectsWorldCreateFXSystem,
+		TEXT("When true, Niagara world managers will only be created for worlds with bCreateFXSystem==true in it's initialization options."),
+		ECVF_Default
+	);
+	
+	//Rollback gate for delaying NDC init. Avoids some cases where UWorlds are created but not "real" UWorlds. 
+	//They don't need FX and having NDC managers created that themselves reference the world can cause some GC issues.
+	//Scheduled for removal in UE-348786
+	bool DelayNDCManagerInitialiation = true;
+	static FAutoConsoleVariableRef CVarRollback_DelayNDCManagerInitialiation(
+		TEXT("fx.Niagara.Rollback.DelayNDCManagerInitialiation"),
+		DelayNDCManagerInitialiation,
+		TEXT("When true, Niagara will not initialize it's NDC managers until the first tick of the world manager or first NDC use."),
+		ECVF_Default
+	);
 }
 
 bool GNiagaraWorldManagerUseTickPasses = false;
@@ -451,7 +472,10 @@ void FNiagaraWorldManager::Init(UWorld* InWorld)
 
 void FNiagaraWorldManager::OnPostWorldInit()
 {
-	DataChannelManager->Init();
+	if(NiagaraRollbackCVars::DelayNDCManagerInitialiation == false)
+	{
+		DataChannelManager->Init();
+	}
 }
 
 FNiagaraWorldManager::~FNiagaraWorldManager()
@@ -1084,12 +1108,15 @@ void FNiagaraWorldManager::WaitForAsyncWork()
 
 void FNiagaraWorldManager::OnWorldInit(UWorld* World, const UWorld::InitializationValues IVS)
 {
-	FNiagaraWorldManager*& NewManager = WorldManagers.FindOrAdd(World);
-	if (!NewManager)
+	if(IVS.bCreateFXSystem || !NiagaraRollbackCVars::WorldManagerRespectsWorldCreateFXSystem)
 	{
-		NewManager = new FNiagaraWorldManager();
+		FNiagaraWorldManager*& NewManager = WorldManagers.FindOrAdd(World);
+		if (!NewManager)
+		{
+			NewManager = new FNiagaraWorldManager();
+		}
+		NewManager->Init(World);
 	}
-	NewManager->Init(World);
 }
 
 void FNiagaraWorldManager::OnPostWorldInit(UWorld* World, const UWorld::InitializationValues IVS)
@@ -1184,22 +1211,34 @@ void FNiagaraWorldManager::DestroyAllSystemSimulations(class UNiagaraSystem* Sys
 
 void FNiagaraWorldManager::TickWorldStart(UWorld* World, ELevelTick TickType, float DeltaSeconds)
 {
-	Get(World)->TickStart(DeltaSeconds);
+	if(FNiagaraWorldManager* WorldMan = Get(World))
+	{
+		WorldMan->TickStart(DeltaSeconds);
+	}	
 }
 
 void FNiagaraWorldManager::TickWorldEnd(UWorld* World, ELevelTick TickType, float DeltaSeconds)
 {
-	Get(World)->TickEnd(DeltaSeconds);
+	if(FNiagaraWorldManager* WorldMan = Get(World))
+	{
+		WorldMan->TickEnd(DeltaSeconds);
+	}
 }
 
 void FNiagaraWorldManager::TickWorld(UWorld* World, ELevelTick TickType, float DeltaSeconds)
 {
-	Get(World)->PostActorTick(DeltaSeconds);
+	if(FNiagaraWorldManager* WorldMan = Get(World))
+	{
+		WorldMan->PostActorTick(DeltaSeconds);
+	}
 }
 
 void FNiagaraWorldManager::OnWorldPreActorTick(UWorld* InWorld, ELevelTick InLevelTick, float InDeltaSeconds)
 {
-	Get(InWorld)->PreActorTick(InLevelTick, InDeltaSeconds);
+	if(FNiagaraWorldManager* WorldMan = Get(InWorld))
+	{
+		WorldMan->PreActorTick(InLevelTick, InDeltaSeconds);
+	}
 }
 
 void FNiagaraWorldManager::OnPreGarbageCollect()
