Fix issue caused by mis-matched GT and RT ticks causing NDC data to be effectively lost from the POV of the GPU.

diff --git a/Engine/Plugins/FX/Niagara/Source/Niagara/Private/NiagaraDataChannelData.cpp b/Engine/Plugins/FX/Niagara/Source/Niagara/Private/NiagaraDataChannelData.cpp
--- a/Engine/Plugins/FX/Niagara/Source/Niagara/Private/NiagaraDataChannelData.cpp
+++ b/Engine/Plugins/FX/Niagara/Source/Niagara/Private/NiagaraDataChannelData.cpp
@@ -26,6 +26,19 @@ namespace NDCCVars
 	static FAutoConsoleVariableRef CVarEnableAccumulatedWrites(TEXT("fx.Niagara.DataChannels.EnableAccumulatedWrites"), EnableAccumulatedWrites, TEXT("True if we attempt to accumulate writes into larger shared buffers."), ECVF_Default);
 };
 
+namespace NiagaraRollbackCVars
+{
+	namespace NDC
+	{
+		bool bRetainGPUDataIfPendingCpuUpdate = true;
+		static FAutoConsoleVariableRef CVarRetainGPUFrameDataIfPendingCPUUpdate(
+		TEXT("fx.Niagara.Rollback.DataChannels.RetainGPUDataIfPendingCpuUpdate"), 
+		bRetainGPUDataIfPendingCpuUpdate, 
+		TEXT("If true, NDC data for the GPU will be kept if there has been no update from the CPU. Helps manager CPU/GPU sync issues where NDC writers on the CPU may not be updating at the same rate as the renderer."),
+		ECVF_Default | ECVF_RenderThreadSafe);
+	};
+};
+
 FNiagaraDataChannelData::FNiagaraDataChannelData()
 {
 
@@ -178,7 +191,7 @@ void FNiagaraDataChannelData::EndFrame(UNiagaraDataChannelHandler* Owner)
 			}
 		}
 
-		if(TotalGPUInstances > 0)
+		if(NiagaraRollbackCVars::NDC::bRetainGPUDataIfPendingCpuUpdate || TotalGPUInstances > 0)
 		{
 			ENQUEUE_RENDER_COMMAND(FDataChannelProxyEndFrame) (
 				[RT_Proxy = RTProxy, RT_BuffersForGPU = MoveTemp(BuffersForGPU)](FRHICommandListImmediate& CmdList)
@@ -534,7 +547,9 @@ FNiagaraDataChannelDataProxy::~FNiagaraDataChannelDataProxy()
 	PrevFrameData = nullptr;
 	bNeedsPrevFrameData = false;
 	PendingCPUBuffers.Empty();
+	CurrentCPUBuffers.Empty();
 	PendingGPUAllocations = 0;
+	CurrentGPUAllocation = 0;
 }
 
 
@@ -600,6 +615,26 @@ void FNiagaraDataChannelDataProxy::BeginFrame(FNiagaraGpuComputeDispatchInterfac
 	}
 	PendingGPUReadbackBuffers.Reset();
 
+	bool bPendingUpdatesFromCPU = PendingCPUBuffers.Num() > 0;
+	//We hold onto our most recent CPU buffers and reapply them each frame to help deal with mis-matched tick rate from the GT and RT.
+	if(NiagaraRollbackCVars::NDC::bRetainGPUDataIfPendingCpuUpdate)
+	{	
+		//If we have no update from the CPU and there have been (and will be) no writes from the GPU, we can just leave the frame data as is and wait for an update from the CPU.
+		if(!bPendingUpdatesFromCPU && CurrentGPUAllocation == 0 && PendingGPUAllocations == 0)
+		{
+			return;
+		}
+
+		if(bPendingUpdatesFromCPU)
+		{
+			CurrentCPUBuffers = MoveTemp(PendingCPUBuffers);		
+		}
+	}
+	else
+	{
+		CurrentCPUBuffers = MoveTemp(PendingCPUBuffers);
+	}
+	
 	if(this->bNeedsPrevFrameData)
 	{
 		Swap(PrevFrameData, CurrFrameData);
@@ -607,13 +642,14 @@ void FNiagaraDataChannelDataProxy::BeginFrame(FNiagaraGpuComputeDispatchInterfac
 
 	//The base num instances for writing from the GPU.
 	uint32 NumInstanceFromCPU = 0;
-	for (auto& Buffer : PendingCPUBuffers)
+	for (auto& Buffer : CurrentCPUBuffers)
 	{
 		NumInstanceFromCPU += Buffer->GetNumInstances();
 	}
 
 	//Allocate our GPU buffers we'll write into. Combine accumulated counts from writing DIs and the data coming in from the CPU.
 	uint32 InstancesToAllocate = NumInstanceFromCPU + PendingGPUAllocations;
+	CurrentGPUAllocation = PendingGPUAllocations;
 	PendingGPUAllocations = 0;
 
 	if(InstancesToAllocate == 0)
@@ -632,9 +668,13 @@ void FNiagaraDataChannelDataProxy::BeginFrame(FNiagaraGpuComputeDispatchInterfac
 	InstCountManager.FreeEntry(InstanceCountOffset);
 
 	CurrFrameData->AllocateGPU(RHICmdList, InstancesToAllocate, DispatchInterface->GetFeatureLevel(), GetDebugName());
-	CurrFrameData->PushCPUBuffersToGPU(PendingCPUBuffers, true, RHICmdList, DispatchInterface->GetFeatureLevel(), GetDebugName(), false);
+	CurrFrameData->PushCPUBuffersToGPU(CurrentCPUBuffers, true, RHICmdList, DispatchInterface->GetFeatureLevel(), GetDebugName(), false);
 	CurrFrameData->SetNumInstances(InstancesToAllocate);
-	PendingCPUBuffers.Reset();
+	
+	if(!NiagaraRollbackCVars::NDC::bRetainGPUDataIfPendingCpuUpdate)
+	{
+		CurrentCPUBuffers.Reset();		
+	}
 
 	InstanceCountOffset = InstCountManager.AcquireOrAllocateEntry(RHICmdList);	
 	CurrFrameData->SetGPUInstanceCountBufferOffset(InstanceCountOffset);
diff --git a/Engine/Plugins/FX/Niagara/Source/Niagara/Public/NiagaraDataChannelData.h b/Engine/Plugins/FX/Niagara/Source/Niagara/Public/NiagaraDataChannelData.h
--- a/Engine/Plugins/FX/Niagara/Source/Niagara/Public/NiagaraDataChannelData.h
+++ b/Engine/Plugins/FX/Niagara/Source/Niagara/Public/NiagaraDataChannelData.h
@@ -37,6 +37,9 @@ struct FNiagaraDataChannelDataProxy : public TSharedFromThis<FNiagaraDataChannel
 
 	//Buffers coming from the CPU that we're going to copy up for reading on the GPU
 	TArray<FNiagaraDataBufferRef> PendingCPUBuffers;
+	
+	//The buffers received from the CPU most recently. We keep these around until we get another update from the Cpu.
+	TArray<FNiagaraDataBufferRef> CurrentCPUBuffers;
 
 	//Buffers written from the GPU that we must send back to the CPU.
 	TArray<FNDCGpuReadbackInfo> PendingGPUReadbackBuffers;
@@ -44,6 +47,9 @@ struct FNiagaraDataChannelDataProxy : public TSharedFromThis<FNiagaraDataChannel
 	//Users that need space in this NDC Data add to this for each tick via AddGPUAllocationForNextTick().
 	int32 PendingGPUAllocations = 0;
 
+	//Allocation made for the most current frame for NDC writes from the GPU.
+	int32 CurrentGPUAllocation = 0;
+
 	//Track current read/write counts +ve for readers, -ve for writers. We cannot mix readers and writers in the same buffer in the same stage.
 	int32 CurrBufferAccessCounts = 0;
 
