Added off-by-default workaround for Niagara crash we hit in internal testing.

#[robomerge] -38.10
#rnx

diff --git a/Engine/Plugins/FX/Niagara/Source/Niagara/Private/NiagaraDataChannel_Map.cpp b/Engine/Plugins/FX/Niagara/Source/Niagara/Private/NiagaraDataChannel_Map.cpp
--- a/Engine/Plugins/FX/Niagara/Source/Niagara/Private/NiagaraDataChannel_Map.cpp
+++ b/Engine/Plugins/FX/Niagara/Source/Niagara/Private/NiagaraDataChannel_Map.cpp
@@ -2,6 +2,7 @@
 
 #include "NiagaraDataChannel_Map.h"
 
+#include "HAL/IConsoleManager.h"
 #include "NiagaraDataChannelCommon.h"
 #include "NiagaraWorldManager.h"
 #include "NiagaraComponent.h"
@@ -18,6 +19,12 @@ namespace NDCMapLocal
 
 	float EntryPoolUnusedCleanTime = 60.0f;
 	static FAutoConsoleVariableRef CVarEntryPoolUnusedCleanTime(TEXT("fx.Niagara.DataChannels.NDCMapBase.FreeUnusedTime"), EntryPoolUnusedCleanTime, TEXT("Time in seconds we keep unused entires in the pool."), ECVF_Default);
+
+	static bool CollectReferencesFromEntryPool = false;
+	static FAutoConsoleVariableRef CollectReferencesFromEntryPoolCVar(
+		TEXT("fx.Niagara.DataChannels.NDCMapBase.CollectReferencesFromEntryPool"),
+		CollectReferencesFromEntryPool,
+		TEXT("Whether to use a workaround to prevent components from being garbage collected while referenced by entry pool"));
 }
 
 //////////////////////////////////////////////////////////////////////////
@@ -142,6 +149,19 @@ void UNiagaraDataChannelHandler_MapBase::Tick(float DeltaTime, ETickingGroup Tic
 	}
 }
 
+void UNiagaraDataChannelHandler_MapBase::AddReferencedObjects(UObject* InThis, FReferenceCollector& Collector)
+{
+	Super::AddReferencedObjects(InThis, Collector);
+	UNiagaraDataChannelHandler_MapBase* This = static_cast<UNiagaraDataChannelHandler_MapBase*>(InThis);
+	if (NDCMapLocal::CollectReferencesFromEntryPool)
+	{
+		for (FNDCMapEntry& Entry : This->EntryPool)
+		{
+			Collector.AddPropertyReferencesWithStructARO(FNDCMapEntry::StaticStruct(), &Entry, This);
+		}
+	}
+}
+
 FNiagaraDataChannelDataPtr UNiagaraDataChannelHandler_MapBase::FindData(FNDCAccessContextInst& AccessContext, ENiagaraResourceAccess AccessType)
 {
 	FNDCMapEntry& Entry = FindOrAddEntry(AccessContext);
diff --git a/Engine/Plugins/FX/Niagara/Source/Niagara/Public/NiagaraDataChannel_Map.h b/Engine/Plugins/FX/Niagara/Source/Niagara/Public/NiagaraDataChannel_Map.h
--- a/Engine/Plugins/FX/Niagara/Source/Niagara/Public/NiagaraDataChannel_Map.h
+++ b/Engine/Plugins/FX/Niagara/Source/Niagara/Public/NiagaraDataChannel_Map.h
@@ -12,6 +12,8 @@
 class UNiagaraDataChannelHandler_MapBase;
 struct FNDCMapKey;
 
+#define UE_API NIAGARA_API
+
 /*
 Base class for Map entries for NDC Map. 
 Child classes of the NDC Map type should implement their own entry type and implement any needed overrides. 
@@ -235,6 +237,9 @@ class UNiagaraDataChannelHandler_MapBase : public UNiagaraDataChannelHandler
 	virtual void EndFrame(float DeltaTime)override;
 	virtual void Tick(float DeltaTime, ETickingGroup TickGroup) override;
 
+	// TEMP: workaround to avoid dangling components in pool
+	static UE_API void AddReferencedObjects(UObject* InThis, FReferenceCollector& Collector);
+
 	[[nodiscard]] virtual FNiagaraDataChannelDataPtr FindData(FNDCAccessContextInst& AccessContext, ENiagaraResourceAccess AccessType)override;
 
 	[[nodiscard]] UWorld* GetWorld()const { return OwningWorld->GetWorld(); }
@@ -268,4 +273,6 @@ protected:
 
 	UPROPERTY()
 	TObjectPtr<UObject> DefaultSystemToSpawn;
-};
\ No newline at end of file
+};
+
+#undef UE_API
