Fixed issue with lingering NDC data when there are updates with no data from the CPU.

diff --git a/Engine/Plugins/FX/Niagara/Source/Niagara/Private/NiagaraDataChannelData.cpp b/Engine/Plugins/FX/Niagara/Source/Niagara/Private/NiagaraDataChannelData.cpp
--- a/Engine/Plugins/FX/Niagara/Source/Niagara/Private/NiagaraDataChannelData.cpp
+++ b/Engine/Plugins/FX/Niagara/Source/Niagara/Private/NiagaraDataChannelData.cpp
@@ -615,17 +615,16 @@ void FNiagaraDataChannelDataProxy::BeginFrame(FNiagaraGpuComputeDispatchInterfac
 	}
 	PendingGPUReadbackBuffers.Reset();
 
-	bool bPendingUpdatesFromCPU = PendingCPUBuffers.Num() > 0;
 	//We hold onto our most recent CPU buffers and reapply them each frame to help deal with mis-matched tick rate from the GT and RT.
 	if(NiagaraRollbackCVars::NDC::bRetainGPUDataIfPendingCpuUpdate)
 	{	
 		//If we have no update from the CPU and there have been (and will be) no writes from the GPU, we can just leave the frame data as is and wait for an update from the CPU.
-		if(!bPendingUpdatesFromCPU && CurrentGPUAllocation == 0 && PendingGPUAllocations == 0)
+		if(!bHasPendingCpuUpdate && CurrentGPUAllocation == 0 && PendingGPUAllocations == 0)
 		{
 			return;
 		}
 
-		if(bPendingUpdatesFromCPU)
+		if(bHasPendingCpuUpdate)
 		{
 			CurrentCPUBuffers = MoveTemp(PendingCPUBuffers);		
 		}
@@ -634,7 +633,9 @@ void FNiagaraDataChannelDataProxy::BeginFrame(FNiagaraGpuComputeDispatchInterfac
 	{
 		CurrentCPUBuffers = MoveTemp(PendingCPUBuffers);
 	}
-	
+
+	bHasPendingCpuUpdate = false;
+
 	if(this->bNeedsPrevFrameData)
 	{
 		Swap(PrevFrameData, CurrFrameData);
@@ -902,6 +903,7 @@ FNiagaraDataBufferRef FNiagaraDataChannelDataProxy::AllocateBufferForCPU(FRDGBui
 void FNiagaraDataChannelDataProxy::AddBuffersFromCPU(const TArray<FNiagaraDataBufferRef>& BuffersFromCPU)
 {
 	PendingCPUBuffers.Append(BuffersFromCPU);
+	bHasPendingCpuUpdate = true;
 }
 
 void FNiagaraDataChannelDataProxy::AddGPUAllocationForNextTick(int32 AllocationCount)
diff --git a/Engine/Plugins/FX/Niagara/Source/Niagara/Public/NiagaraDataChannelData.h b/Engine/Plugins/FX/Niagara/Source/Niagara/Public/NiagaraDataChannelData.h
--- a/Engine/Plugins/FX/Niagara/Source/Niagara/Public/NiagaraDataChannelData.h
+++ b/Engine/Plugins/FX/Niagara/Source/Niagara/Public/NiagaraDataChannelData.h
@@ -31,6 +31,7 @@ struct FNiagaraDataChannelDataProxy : public TSharedFromThis<FNiagaraDataChannel
 	FNiagaraDataBufferRef CurrFrameData = nullptr;
 	FNiagaraDataBufferRef PrevFrameData = nullptr;
 	bool bNeedsPrevFrameData = false;
+	bool bHasPendingCpuUpdate = false;
 
 	//Keeping layout info ref to ensure lifetime for GPUDataSet.
 	FNiagaraDataChannelLayoutInfoPtr LayoutInfo;
