Fix for issue with access to freed Niagara Components during cleanup.

diff --git a/Engine/Plugins/FX/Niagara/Source/Niagara/Private/NiagaraDataChannel_Map.cpp b/Engine/Plugins/FX/Niagara/Source/Niagara/Private/NiagaraDataChannel_Map.cpp
--- a/Engine/Plugins/FX/Niagara/Source/Niagara/Private/NiagaraDataChannel_Map.cpp
+++ b/Engine/Plugins/FX/Niagara/Source/Niagara/Private/NiagaraDataChannel_Map.cpp
@@ -19,12 +19,6 @@ namespace NDCMapLocal
 
 	float EntryPoolUnusedCleanTime = 60.0f;
 	static FAutoConsoleVariableRef CVarEntryPoolUnusedCleanTime(TEXT("fx.Niagara.DataChannels.NDCMapBase.FreeUnusedTime"), EntryPoolUnusedCleanTime, TEXT("Time in seconds we keep unused entires in the pool."), ECVF_Default);
-
-	static bool CollectReferencesFromEntryPool = false;
-	static FAutoConsoleVariableRef CollectReferencesFromEntryPoolCVar(
-		TEXT("fx.Niagara.DataChannels.NDCMapBase.CollectReferencesFromEntryPool"),
-		CollectReferencesFromEntryPool,
-		TEXT("Whether to use a workaround to prevent components from being garbage collected while referenced by entry pool"));
 }
 
 //////////////////////////////////////////////////////////////////////////
@@ -149,19 +143,6 @@ void UNiagaraDataChannelHandler_MapBase::Tick(float DeltaTime, ETickingGroup Tic
 	}
 }
 
-void UNiagaraDataChannelHandler_MapBase::AddReferencedObjects(UObject* InThis, FReferenceCollector& Collector)
-{
-	Super::AddReferencedObjects(InThis, Collector);
-	UNiagaraDataChannelHandler_MapBase* This = static_cast<UNiagaraDataChannelHandler_MapBase*>(InThis);
-	if (NDCMapLocal::CollectReferencesFromEntryPool)
-	{
-		for (FNDCMapEntry& Entry : This->EntryPool)
-		{
-			Collector.AddPropertyReferencesWithStructARO(FNDCMapEntry::StaticStruct(), &Entry, This);
-		}
-	}
-}
-
 FNiagaraDataChannelDataPtr UNiagaraDataChannelHandler_MapBase::FindData(FNDCAccessContextInst& AccessContext, ENiagaraResourceAccess AccessType)
 {
 	FNDCMapEntry& Entry = FindOrAddEntry(AccessContext);
diff --git a/Engine/Plugins/FX/Niagara/Source/Niagara/Public/NiagaraDataChannel_Map.h b/Engine/Plugins/FX/Niagara/Source/Niagara/Public/NiagaraDataChannel_Map.h
--- a/Engine/Plugins/FX/Niagara/Source/Niagara/Public/NiagaraDataChannel_Map.h
+++ b/Engine/Plugins/FX/Niagara/Source/Niagara/Public/NiagaraDataChannel_Map.h
@@ -237,9 +237,6 @@ class UNiagaraDataChannelHandler_MapBase : public UNiagaraDataChannelHandler
 	virtual void EndFrame(float DeltaTime)override;
 	virtual void Tick(float DeltaTime, ETickingGroup TickGroup) override;
 
-	// TEMP: workaround to avoid dangling components in pool
-	static UE_API void AddReferencedObjects(UObject* InThis, FReferenceCollector& Collector);
-
 	[[nodiscard]] virtual FNiagaraDataChannelDataPtr FindData(FNDCAccessContextInst& AccessContext, ENiagaraResourceAccess AccessType)override;
 
 	[[nodiscard]] UWorld* GetWorld()const { return OwningWorld->GetWorld(); }
@@ -260,6 +257,7 @@ protected:
 	/** The index of the currently initializing entry to help handle re-entrant calls to FindData from spawned systems. */
 	FNDCMapEntry* CurrentlyInitializingEntry = nullptr;
 
+	UPROPERTY(transient)
 	TArray<FNDCMapEntry> EntryPool;
 	
 	/** Currently Active Entries in EntryPool. Map of Entries by FNDCKey that we generate from an Access Context inside GenerateKey. Child classes can generate the key however they need and so we can easily subdivide the NDC internal data with this single map. */
@@ -269,7 +267,8 @@ protected:
 	TArray<int32> FreeEntries;
 
 	/** Mapping from spawned handler systems back to their owning entry. Allows faster access when handler systems are accessing the NDC. */
-	TMap<UNiagaraComponent*, int32> SpawnedComponentsToActiveEntry;
+	UPROPERTY(transient)
+	TMap<TObjectPtr<UNiagaraComponent>, int32> SpawnedComponentsToActiveEntry;
 
 	UPROPERTY()
 	TObjectPtr<UObject> DefaultSystemToSpawn;
