Shader build fix

diff --git a/Engine/Plugins/FX/Niagara/Shaders/Private/NiagaraMeshVertexFactory.ush b/Engine/Plugins/FX/Niagara/Shaders/Private/NiagaraMeshVertexFactory.ush
--- a/Engine/Plugins/FX/Niagara/Shaders/Private/NiagaraMeshVertexFactory.ush
+++ b/Engine/Plugins/FX/Niagara/Shaders/Private/NiagaraMeshVertexFactory.ush
@@ -617,25 +617,25 @@ FVertexFactoryIntermediates GetVertexFactoryIntermediates(FVertexFactoryInput In
 			Intermediates.TexCoords[0] = Input.TexCoords0;
 		#elif NUM_MATERIAL_TEXCOORDS_VERTEX >= 2
 			Intermediates.TexCoords[0] = Input.TexCoords01.xy;
-			Intermediates.TexCoords[1] = NumFetchTexCoords > 1 ? Input.TexCoords01.zw: Result.TexCoords[0];
+			Intermediates.TexCoords[1] = NumFetchTexCoords > 1 ? Input.TexCoords01.zw: Intermediates.TexCoords[0];
 		#endif
 		#if NUM_MATERIAL_TEXCOORDS_VERTEX == 3
-			Intermediates.TexCoords[2] = NumFetchTexCoords > 2 ? Input.TexCoords2: Result.TexCoords[1];
+			Intermediates.TexCoords[2] = NumFetchTexCoords > 2 ? Input.TexCoords2: Intermediates.TexCoords[1];
 		#elif NUM_MATERIAL_TEXCOORDS_VERTEX >= 4
-			Intermediates.TexCoords[2] = NumFetchTexCoords > 2 ? Input.TexCoords23.xy : Result.TexCoords[1];
-			Intermediates.TexCoords[3] = NumFetchTexCoords > 3 ? Input.TexCoords23.zw : Result.TexCoords[2];
+			Intermediates.TexCoords[2] = NumFetchTexCoords > 2 ? Input.TexCoords23.xy : Intermediates.TexCoords[1];
+			Intermediates.TexCoords[3] = NumFetchTexCoords > 3 ? Input.TexCoords23.zw : Intermediates.TexCoords[2];
 		#endif
 		#if NUM_MATERIAL_TEXCOORDS_VERTEX == 5
-			Intermediates.TexCoords[4] = NumFetchTexCoords > 4 ? Input.TexCoords4 : Result.TexCoords[3];
+			Intermediates.TexCoords[4] = NumFetchTexCoords > 4 ? Input.TexCoords4 : Intermediates.TexCoords[3];
 		#elif NUM_MATERIAL_TEXCOORDS_VERTEX >= 6
-			Intermediates.TexCoords[4] = NumFetchTexCoords > 4 ? Input.TexCoords45.xy : Result.TexCoords[3];
-			Intermediates.TexCoords[5] = NumFetchTexCoords > 5 ? Input.TexCoords45.zw : Result.TexCoords[4];
+			Intermediates.TexCoords[4] = NumFetchTexCoords > 4 ? Input.TexCoords45.xy : Intermediates.TexCoords[3];
+			Intermediates.TexCoords[5] = NumFetchTexCoords > 5 ? Input.TexCoords45.zw : Intermediates.TexCoords[4];
 		#endif
 		#if NUM_MATERIAL_TEXCOORDS_VERTEX == 7
-			Intermediates.TexCoords[6] = NumFetchTexCoords > 6 ? Input.TexCoords6 : Result.TexCoords[5];
+			Intermediates.TexCoords[6] = NumFetchTexCoords > 6 ? Input.TexCoords6 : Intermediates.TexCoords[5];
 		#elif NUM_MATERIAL_TEXCOORDS_VERTEX >= 8
-			Intermediates.TexCoords[6] = NumFetchTexCoords > 6 ? Input.TexCoords67.xy : Result.TexCoords[5];
-			Intermediates.TexCoords[7] = NumFetchTexCoords > 7 ? Input.TexCoords67.zw : Result.TexCoords[6];
+			Intermediates.TexCoords[6] = NumFetchTexCoords > 6 ? Input.TexCoords67.xy : Intermediates.TexCoords[5];
+			Intermediates.TexCoords[7] = NumFetchTexCoords > 7 ? Input.TexCoords67.zw : Intermediates.TexCoords[6];
 		#endif
 	#endif // MANUAL_VERTEX_FETCH
 #endif // NUM_MATERIAL_TEXCOORDS_VERTEX
