- Add Debug HUD command for ShowRegisteredComponents

diff --git a/Engine/Plugins/FX/Niagara/Source/Niagara/Private/NiagaraDebugHud.cpp b/Engine/Plugins/FX/Niagara/Source/Niagara/Private/NiagaraDebugHud.cpp
--- a/Engine/Plugins/FX/Niagara/Source/Niagara/Private/NiagaraDebugHud.cpp
+++ b/Engine/Plugins/FX/Niagara/Source/Niagara/Private/NiagaraDebugHud.cpp
@@ -145,6 +145,7 @@ namespace NiagaraDebugLocal
 		MakeTuple(TEXT("SystemTextFont="), TEXT("Set the system display text font to use"), [](FString Arg) { Settings.SystemTextOptions.Font = StringToFont(Arg); }),
 		MakeTuple(TEXT("ParticleTextFont="), TEXT("Set the particle display text font to use"), [](FString Arg) { Settings.ParticleTextOptions.Font = StringToFont(Arg); }),
 
+		MakeTuple(TEXT("ShowRegisteredComponents="), TEXT("Should we show registered components in the overview."), [](FString Arg) {Settings.bShowRegisteredComponents = FCString::Atoi(*Arg) != 0; }),
 		MakeTuple(TEXT("ShowGlobalBudgetInfo="), TEXT("Shows global budget information"), [](FString Arg) {Settings.bShowGlobalBudgetInfo = FCString::Atoi(*Arg) != 0; }),
 
 		MakeTuple(TEXT("PerfGraphMode="), TEXT("Change the mode of the perf graph"), [](FString Arg) {Settings.PerfGraphMode = (ENiagaraDebugHUDPerfGraphMode)FCString::Atoi(*Arg); }),
