- Socker reader support for entities

diff --git a/Engine/Plugins/FX/Niagara/Source/Niagara/Internal/DataInterface/NiagaraDataInterfaceSocketReader.h b/Engine/Plugins/FX/Niagara/Source/Niagara/Internal/DataInterface/NiagaraDataInterfaceSocketReader.h
--- a/Engine/Plugins/FX/Niagara/Source/Niagara/Internal/DataInterface/NiagaraDataInterfaceSocketReader.h
+++ b/Engine/Plugins/FX/Niagara/Source/Niagara/Internal/DataInterface/NiagaraDataInterfaceSocketReader.h
@@ -74,11 +74,11 @@ public:
 	TObjectPtr<UObject> SourceAsset;
 
 	/** When looking for an attached parent component only accept this type of component. */
-	UPROPERTY(EditAnywhere, Category = "SocketReader", meta = (EditConditionHides, EditCondition = "SourceMode == ENDISocketReaderSourceMode::Default || SourceMode == ENDISocketReaderSourceMode::AttachedParentOnly", AllowedClasses="/Script/Engine.SceneComponent"))
+	UPROPERTY(EditAnywhere, Category = "SocketReader", meta = (EditConditionHides, EditCondition = "SourceMode == ENDISocketReaderSourceMode::Default || SourceMode == ENDISocketReaderSourceMode::AttachedParentOnly", AllowedClasses="/Script/Engine.SceneComponent,/EntityFramework/_Verse/VNI/Entity.component"))
 	TObjectPtr<UClass> AttachComponentClass;
 
 	/** When looking for an attached parent component it must have this tag to be considered. */
-	UPROPERTY(EditAnywhere, Category = "SocketReader", meta = (EditConditionHides, EditCondition = "SourceMode == ENDISocketReaderSourceMode::Default || SourceMode == ENDISocketReaderSourceMode::AttachedParentOnly", AllowedClasses = "/Script/Engine.SceneComponent"))
+	UPROPERTY(EditAnywhere, Category = "SocketReader", meta = (EditConditionHides, EditCondition = "SourceMode == ENDISocketReaderSourceMode::Default || SourceMode == ENDISocketReaderSourceMode::AttachedParentOnly"))
 	FName AttachComponentTag;
 
 	/**
diff --git a/Engine/Plugins/FX/Niagara/Source/Niagara/Private/DataInterface/NiagaraDataInterfaceSocketReader.cpp b/Engine/Plugins/FX/Niagara/Source/Niagara/Private/DataInterface/NiagaraDataInterfaceSocketReader.cpp
--- a/Engine/Plugins/FX/Niagara/Source/Niagara/Private/DataInterface/NiagaraDataInterfaceSocketReader.cpp
+++ b/Engine/Plugins/FX/Niagara/Source/Niagara/Private/DataInterface/NiagaraDataInterfaceSocketReader.cpp
@@ -7,6 +7,8 @@
 #include "NiagaraComponent.h"
 #include "NiagaraDataInterfaceUtilities.h"
 #include "NiagaraTypes.h"
+#include "NiagaraActorSceneComponentSocketUtils.h"
+#include "NiagaraSceneComponentUtils.h"
 #include "NiagaraShaderParametersBuilder.h"
 #include "NiagaraSystemInstance.h"
 #include "NiagaraWorldManager.h"
@@ -57,6 +59,8 @@ namespace NDISocketReaderLocal
 
 	struct FInstanceData_GameThread
 	{
+		TUniquePtr<FNiagaraSceneComponentSocketUtils>	SocketUtils;
+
 		FNiagaraParameterDirectBinding<UObject*>	UserParamBinding;
 		uint32										bIsDataValid : 1 = false;
 		uint32										bNeedsSocketRecache : 1 = true;
@@ -79,8 +83,6 @@ namespace NDISocketReaderLocal
 		TArray<FName>								SocketNames;
 		TArray<FTransform3f>						SocketTransforms;
 		TArray<FTransform3f>						PreviousSocketTransforms;
-
-		TWeakObjectPtr<UObject>						ResolvedObject;
 	};
 
 	struct FInstanceData_SharedData
@@ -223,142 +225,8 @@ namespace NDISocketReaderLocal
 
 	void ResolveObject(const UNiagaraDataInterfaceSocketReader* SocketDI, FInstanceData_GameThread* InstanceData, FNiagaraSystemInstance* SystemInstance)
 	{
-		UObject* ResolvedObject = nullptr;
-
-		const bool bTryParameterBinding = SocketDI->SourceMode == ENDISocketReaderSourceMode::Default || SocketDI->SourceMode == ENDISocketReaderSourceMode::ParameterBindingOnly;
-		const bool bTryAttachedParent	= SocketDI->SourceMode == ENDISocketReaderSourceMode::Default || SocketDI->SourceMode == ENDISocketReaderSourceMode::AttachedParentOnly;
-		const bool bTrySource			= SocketDI->SourceMode == ENDISocketReaderSourceMode::Default || SocketDI->SourceMode == ENDISocketReaderSourceMode::SourceOnly;
-
-		// Resolve Parameter Binding?
-		if (bTryParameterBinding)
-		{
-			ResolvedObject = InstanceData->UserParamBinding.GetValue();
-			if (ResolvedObject)
-			{
-				if (AActor* AsActor = Cast<AActor>(ResolvedObject))
-				{
-					ResolvedObject = AsActor->GetRootComponent();
-				}
-			}
-		}
-
-		// Resolve Attached Parent?
-		if (bTryAttachedParent && !ResolvedObject)
-		{
-			USceneComponent* AttachComponent = SystemInstance->GetAttachComponent();
-			if (AttachComponent->IsA<UNiagaraComponent>())
-			{
-				AttachComponent = AttachComponent->GetAttachParent();
-			}
-
-			// Look for the first viable component that satisfies the class and tag
-			if (SocketDI->AttachComponentClass || !SocketDI->AttachComponentTag.IsNone())
-			{
-				while (AttachComponent)
-				{
-					if ((!SocketDI->AttachComponentClass || AttachComponent->IsA(SocketDI->AttachComponentClass)) &&
-						(SocketDI->AttachComponentTag.IsNone() || AttachComponent->ComponentHasTag(SocketDI->AttachComponentTag)) )
-					{
-						break;
-					}
-					AttachComponent = AttachComponent->GetAttachParent();
-				}
-			}
-
-			ResolvedObject = AttachComponent;
-		}
-
-		// Resolve Source
-		if (bTrySource && !ResolvedObject)
-		{
-			ResolvedObject = SocketDI->SourceActor.Get();
-			if (ResolvedObject)
-			{
-				if (AActor* AsActor = Cast<AActor>(ResolvedObject))
-				{
-					ResolvedObject = AsActor->GetRootComponent();
-				}
-			}
-			if ( !ResolvedObject )
-			{
-				ResolvedObject = SocketDI->SourceAsset;
-			}
-		}
-
-	#if WITH_EDITORONLY_DATA
-		// Fallback for editor preview
-		if (!ResolvedObject && !SystemInstance->GetWorld()->IsGameWorld())
-		{
-			ResolvedObject = SocketDI->EditorPreviewAsset.LoadSynchronous();
-		}
-	#endif
-
-		// Store the resolved object
-		InstanceData->bIsDataValid			= ResolvedObject != nullptr;
-		InstanceData->bNeedsSocketRecache  |= ResolvedObject != InstanceData->ResolvedObject;
-		InstanceData->ResolvedObject		= ResolvedObject;
-
-	#if WITH_EDITOR
-		// When in the editor the socket counts can change so we might need to recache the socket list
-		if (!InstanceData->bNeedsSocketRecache && ResolvedObject)
-		{
-			int32 NewSocketCount = InstanceData->NumSockets;
-			if (USceneComponent* SceneComponent = Cast<USceneComponent>(ResolvedObject))
-			{
-				TArray<FComponentSocketDescription> SocketList;
-				SocketList.Reserve(InstanceData->NumSockets);
-				SceneComponent->QuerySupportedSockets(SocketList);
-				NewSocketCount = SocketList.Num();
-			}
-			else if (UStaticMesh* StaticMesh = Cast<UStaticMesh>(ResolvedObject))
-			{
-				NewSocketCount = StaticMesh->Sockets.Num();
-			}
-			else if (USkeletalMesh* SkeletalMesh = Cast<USkeletalMesh>(ResolvedObject))
-			{
-				NewSocketCount = SkeletalMesh->NumSockets();
-			}
-			InstanceData->bNeedsSocketRecache |= InstanceData->NumSockets != NewSocketCount;
-		}
-	#endif
-	}
-
-	void GetSocketNames(UObject* ResolvedObject, TArray<FName>& OutSocketNames)
-	{
-		if (USceneComponent* SceneComponent = Cast<USceneComponent>(ResolvedObject))
-		{
-			OutSocketNames = SceneComponent->GetAllSocketNames();
-		}
-		else if (UStaticMesh* StaticMesh = Cast<UStaticMesh>(ResolvedObject))
-		{
-			const int32 NumSockets = StaticMesh->Sockets.Num();
-			OutSocketNames.Reserve(StaticMesh->Sockets.Num());
-			for (UStaticMeshSocket* Socket : StaticMesh->Sockets)
-			{
-				check(Socket);
-				OutSocketNames.Add(Socket->SocketName);
-			}
-		}
-		else if (USkeletalMesh* SkeletalMesh = Cast<USkeletalMesh>(ResolvedObject))
-		{
-			const FReferenceSkeleton& RefSkeleton = SkeletalMesh->GetRefSkeleton();
-			const TArray<FMeshBoneInfo>& RefBoneInfo = RefSkeleton.GetRefBoneInfo();
-			const int32 NumBones = RefBoneInfo.Num();
-			const int32 NumSockets = SkeletalMesh->NumSockets();
-
-			OutSocketNames.Reserve(NumBones + NumSockets);
-			for (int i = 0; i < NumBones; ++i)
-			{
-				OutSocketNames.Add(RefBoneInfo[i].Name);
-			}
-
-			for (int i = 0; i < NumSockets; ++i)
-			{
-				USkeletalMeshSocket* Socket = SkeletalMesh->GetSocketByIndex(i);
-				check(Socket);
-				OutSocketNames.Add(Socket->SocketName);
-			}
-		}
+		InstanceData->bNeedsSocketRecache  |= InstanceData->SocketUtils->ResolveObject();
+		InstanceData->bIsDataValid			= InstanceData->SocketUtils->IsValid();
 	}
 
 	void UpdateSocketCache(const UNiagaraDataInterfaceSocketReader* SocketDI, FInstanceData_GameThread* InstanceData, FNiagaraSystemInstance* SystemInstance)
@@ -382,7 +250,7 @@ namespace NDISocketReaderLocal
 			InstanceData->SocketFilterUnfilteredIndex.Empty();
 			InstanceData->SocketNames.Empty();
 
-			GetSocketNames(InstanceData->ResolvedObject.Get(), InstanceData->SocketNames);
+			InstanceData->SocketNames = InstanceData->SocketUtils->GetSocketNames();
 
 			InstanceData->NumSockets = InstanceData->SocketNames.Num();
 			if (InstanceData->NumSockets > 0 )
@@ -436,81 +304,15 @@ namespace NDISocketReaderLocal
 				InstanceData->PreviousSocketTransforms = InstanceData->SocketTransforms;
 			}
 
-			if (USceneComponent* SceneComponent = Cast<USceneComponent>(InstanceData->ResolvedObject))
-			{
-				InstanceData->bNeedsRenderUpdate = true;
-
-				if (InstanceData->bReadsFilteredTransforms && InstanceData->bReadsUnfilteredTransforms)
-				{
-					for (int32 SocketIndex=0; SocketIndex < InstanceData->NumSockets; ++SocketIndex)
-					{
-						const FTransform SocketTransform = SceneComponent->GetSocketTransform(InstanceData->SocketNames[SocketIndex], ERelativeTransformSpace::RTS_Component);
-						InstanceData->SocketTransforms[SocketIndex] = FTransform3f(SocketTransform);
-					}
-				}
-				else
-				{
-					const int32 SocketTableStart = InstanceData->bReadsFilteredTransforms ? 0 : InstanceData->NumFilteredSockets;
-					const int32 SocketTableEnd   = SocketTableStart + (InstanceData->bReadsFilteredTransforms ? InstanceData->NumFilteredSockets : InstanceData->NumUnfilteredSockets);
-					for (int32 i=SocketTableStart; i < SocketTableEnd; ++i)
-					{
-						const int32 SocketIndex = InstanceData->SocketFilterUnfilteredIndex[i + 1];
-						if (SocketIndex >= 0 && SocketIndex < InstanceData->NumSockets)
-						{
-							const FTransform SocketTransform = SceneComponent->GetSocketTransform(InstanceData->SocketNames[SocketIndex], ERelativeTransformSpace::RTS_Component);
-							InstanceData->SocketTransforms[SocketIndex] = FTransform3f(SocketTransform);
-						}
-					}
-				}
-			}
-			else if (UStaticMesh* StaticMesh = Cast<UStaticMesh>(InstanceData->ResolvedObject))
+			// Read the sockets
+			TOptional<TConstArrayView<int32>> SocketIndices;
+			if (!InstanceData->bReadsFilteredTransforms || !InstanceData->bReadsUnfilteredTransforms)
 			{
-				InstanceData->bNeedsRenderUpdate = true;
-
-				for (int32 i = 0; i < InstanceData->NumSockets; ++i)
-				{
-					UStaticMeshSocket* Socket = StaticMesh->Sockets[i];
-					InstanceData->SocketTransforms[i].SetTranslation(FVector3f(Socket->RelativeLocation));
-					InstanceData->SocketTransforms[i].SetRotation(FQuat4f(Socket->RelativeRotation.Quaternion()));
-					InstanceData->SocketTransforms[i].SetScale3D(FVector3f(Socket->RelativeScale));
-				}
-			}
-			else if (USkeletalMesh* SkeletalMesh = Cast<USkeletalMesh>(InstanceData->ResolvedObject))
-			{
-				InstanceData->bNeedsRenderUpdate = true;
-
-				const FReferenceSkeleton& RefSkeleton = SkeletalMesh->GetRefSkeleton();
-				const TArray<FMeshBoneInfo>& RefBoneInfo = RefSkeleton.GetRefBoneInfo();
-				const TArray<FTransform>& RefBonePose = RefSkeleton.GetRefBonePose();
-				const int32 NumBones = RefBoneInfo.Num();
-				const int32 NumSockets = SkeletalMesh->NumSockets();
-
-				for (int i = 0; i < NumBones; ++i)
-				{
-					FTransform3f SocketTransform = FTransform3f(RefBonePose[i]);
-					const int32 ParentIndex = RefBoneInfo[i].ParentIndex;
-					if (ParentIndex >= 0 && ParentIndex < i)
-					{
-						SocketTransform = SocketTransform * InstanceData->SocketTransforms[ParentIndex];
-					}
-					InstanceData->SocketTransforms[i] = FTransform3f(SocketTransform);
-					InstanceData->PreviousSocketTransforms[i] = InstanceData->SocketTransforms[i];
-				}
-
-				for (int i = 0; i < NumSockets; ++i)
-				{
-					USkeletalMeshSocket* Socket = SkeletalMesh->GetSocketByIndex(i);
-					check(Socket);
-
-					FTransform3f SocketTransform(FTransform(Socket->RelativeRotation.Quaternion(), Socket->RelativeLocation, Socket->RelativeScale));
-					const int32 ParentIndex = RefSkeleton.FindBoneIndex(Socket->BoneName);
-					if (ParentIndex >= 0)
-					{
-						SocketTransform = SocketTransform * InstanceData->SocketTransforms[ParentIndex];
-					}
-					InstanceData->SocketTransforms[i + NumBones] = FTransform3f(SocketTransform);
-				}
+				const int32 FirstSocket = (InstanceData->bReadsFilteredTransforms ? 0 : InstanceData->NumFilteredSockets) + 1;
+				const int32 NumSockets = InstanceData->bReadsFilteredTransforms ? InstanceData->NumFilteredSockets : InstanceData->NumUnfilteredSockets;
+				SocketIndices = MakeArrayView<int32>(InstanceData->SocketFilterUnfilteredIndex.GetData() + FirstSocket, NumSockets);
 			}
+			InstanceData->bNeedsRenderUpdate = InstanceData->SocketUtils->GetSocketTransforms(InstanceData->SocketTransforms, SocketIndices);
 
 			// Transforms are being reset so copy current -> previous
 			if (bResetPreviousTransform)
@@ -519,28 +321,21 @@ namespace NDISocketReaderLocal
 			}
 		}
 
-		USceneComponent* SceneComponent = Cast<USceneComponent>(InstanceData->ResolvedObject);
-		if ( !SceneComponent )
+		// Update the transform
+		const FTransform ComponentToWorld = InstanceData->SocketUtils->GetLocalToWorldTransform();
+		if (bResetPreviousTransform)
 		{
-			SceneComponent = SystemInstance->GetAttachComponent();
+			InstanceData->ComponentToWorld			= ComponentToWorld;
+			InstanceData->PreviousComponentToWorld	= InstanceData->ComponentToWorld;
 		}
-
-		if ( SceneComponent )
+		else
 		{
-			if (bResetPreviousTransform)
-			{
-				InstanceData->ComponentToWorld			= SceneComponent->GetComponentToWorld();
-				InstanceData->PreviousComponentToWorld	= InstanceData->ComponentToWorld;
-			}
-			else
-			{
-				InstanceData->PreviousComponentToWorld	= InstanceData->ComponentToWorld;
-				InstanceData->ComponentToWorld			= SceneComponent->GetComponentToWorld();
-			}
-			const FNiagaraLWCConverter LWCConverter = SystemInstance->GetLWCConverter();
-			InstanceData->ComponentToTranslatedWorld = LWCConverter.ConvertWorldToSimulationTransform(InstanceData->ComponentToWorld);
-			InstanceData->PreviousComponentToTranslatedWorld = LWCConverter.ConvertWorldToSimulationTransform(InstanceData->PreviousComponentToWorld);
+			InstanceData->PreviousComponentToWorld	= InstanceData->ComponentToWorld;
+			InstanceData->ComponentToWorld			= ComponentToWorld;
 		}
+		const FNiagaraLWCConverter LWCConverter = SystemInstance->GetLWCConverter();
+		InstanceData->ComponentToTranslatedWorld = LWCConverter.ConvertWorldToSimulationTransform(InstanceData->ComponentToWorld);
+		InstanceData->PreviousComponentToTranslatedWorld = LWCConverter.ConvertWorldToSimulationTransform(InstanceData->PreviousComponentToWorld);
 	}
 
 	//////////////////////////////////////////////////////////////////////////
@@ -832,6 +627,28 @@ bool UNiagaraDataInterfaceSocketReader::InitPerInstanceData(void* PerInstanceDat
 	using namespace NDISocketReaderLocal;
 
 	FInstanceData_GameThread* InstanceData = new (PerInstanceData) FInstanceData_GameThread;
+	PRAGMA_DISABLE_DEPRECATION_WARNINGS
+	InstanceData->SocketUtils.Reset(FNiagaraSystemInstance::GetSceneComponentUtils(SystemInstance).CreateSocketUtils());
+	PRAGMA_ENABLE_DEPRECATION_WARNINGS
+	check(InstanceData->SocketUtils.IsValid());
+	if ( SourceMode == ENDISocketReaderSourceMode::Default || SourceMode == ENDISocketReaderSourceMode::ParameterBindingOnly )
+	{
+		InstanceData->SocketUtils->SetAllowParameterBinding(SystemInstance, ObjectParameterBinding.Parameter);
+	}
+
+	if ( SourceMode == ENDISocketReaderSourceMode::Default || SourceMode == ENDISocketReaderSourceMode::AttachedParentOnly )
+	{
+		InstanceData->SocketUtils->SetAllowAttachParent(AttachComponentClass, AttachComponentTag);
+	}
+
+	if ( SourceMode == ENDISocketReaderSourceMode::Default || SourceMode == ENDISocketReaderSourceMode::SourceOnly )
+	{
+		InstanceData->SocketUtils->SetAllowSource(SourceActor, SourceAsset);
+	}
+#if WITH_EDITOR
+	InstanceData->SocketUtils->SetEditorPreviewAsset(EditorPreviewAsset);
+#endif
+
 	InstanceData->UserParamBinding.Init(SystemInstance->GetInstanceParameters(), ObjectParameterBinding.Parameter);
 	ResolveObject(this, InstanceData, SystemInstance);
 
@@ -916,10 +733,11 @@ void UNiagaraDataInterfaceSocketReader::ProvidePerInstanceDataForRenderThread(vo
 ETickingGroup UNiagaraDataInterfaceSocketReader::CalculateTickGroup(const void* PerInstanceData) const
 {
 	using namespace NDISocketReaderLocal;
-	if ( bRequireCurrentFrameData )
+	if ( bRequireCurrentFrameData && PerInstanceData)
 	{
 		const FInstanceData_GameThread* InstanceData = static_cast<const FInstanceData_GameThread*>(PerInstanceData);
-		const USceneComponent* SceneComponent = InstanceData ? Cast<USceneComponent>(InstanceData->ResolvedObject.Get()) : nullptr;
+		//-TODO: Not sure how we will handle tick groups with Entity vs Actor?
+		const USceneComponent* SceneComponent = Cast<USceneComponent>(InstanceData->SocketUtils->GetResolvedObject());
 		if (SceneComponent)
 		{
 			ETickingGroup FinalTickGroup = FMath::Max(SceneComponent->PrimaryComponentTick.TickGroup, SceneComponent->PrimaryComponentTick.EndTickGroup);
@@ -994,7 +812,7 @@ void UNiagaraDataInterfaceSocketReader::DrawDebugHud(FNDIDrawDebugHudContext& De
 		return;
 	}
 
-	UObject* ResolvedObject = InstanceData_GT->ResolvedObject.Get();
+	UObject* ResolvedObject = InstanceData_GT->SocketUtils->GetResolvedObject();
 	DebugHudContext.GetOutputString().Appendf(TEXT("ResolvedObject(%s)"), *GetNameSafe(ResolvedObject));
 
 	UCanvas* Canvas = DebugHudContext.GetCanvas();
@@ -1022,7 +840,10 @@ TArray<FName> UNiagaraDataInterfaceSocketReader::GetEditorSocketNames() const
 	TArray<FName> SocketNames;
 	if (UObject* EditorAsset = EditorPreviewAsset.LoadSynchronous())
 	{
-		NDISocketReaderLocal::GetSocketNames(EditorAsset, SocketNames);
+		FNiagaraActorSceneComponentSocketUtils SocketUtils(nullptr);
+		SocketUtils.SetAllowSource(nullptr, EditorAsset);
+		SocketUtils.ResolveObject();
+		SocketNames.Append(SocketUtils.GetSocketNames());
 	}
 	return SocketNames;
 }
diff --git a/Engine/Plugins/FX/Niagara/Source/Niagara/Private/NiagaraActorSceneComponentSocketUtils.h b/Engine/Plugins/FX/Niagara/Source/Niagara/Private/NiagaraActorSceneComponentSocketUtils.h
new file mode 100644
--- /dev/null
+++ b/Engine/Plugins/FX/Niagara/Source/Niagara/Private/NiagaraActorSceneComponentSocketUtils.h
@@ -0,0 +1,105 @@
+// Copyright Epic Games, Inc. All Rights Reserved.
+
+#include "NiagaraSceneComponentSocketUtils.h"
+#include "NiagaraComponent.h"
+#include "NiagaraSystemInstance.h"
+
+#include "Components/SceneComponent.h"
+#include "Engine/Canvas.h"
+#include "Engine/SkeletalMesh.h"
+#include "Engine/SkeletalMeshSocket.h"
+#include "Engine/StaticMesh.h"
+#include "Engine/StaticMeshSocket.h"
+
+class FNiagaraActorSceneComponentSocketUtils final : public FNiagaraSceneComponentSocketUtils
+{
+public:
+	explicit FNiagaraActorSceneComponentSocketUtils(TWeakObjectPtr<UObject> InWeakOwnerComponent)
+		: FNiagaraSceneComponentSocketUtils(InWeakOwnerComponent)
+	{
+	}
+
+	virtual FTransform GetLocalToWorldTransform() const override
+	{
+		USceneComponent* SceneComponent = Cast<USceneComponent>(WeakResolvedObject.Get());
+		if (SceneComponent == nullptr)
+		{
+			SceneComponent = Cast<USceneComponent>(WeakOwnerComponent.Get());
+		}
+		return SceneComponent ? SceneComponent->GetComponentToWorld() : FTransform::Identity;
+	}
+
+	virtual UObject* ResolveAttachParentInternal() const override
+	{
+		USceneComponent* AttachComponent = Cast<USceneComponent>(WeakOwnerComponent.Get());
+		if (AttachComponent && AttachComponent->IsA<UNiagaraComponent>())
+		{
+			AttachComponent = AttachComponent->GetAttachParent();
+		}
+
+		// Look for the first viable component that satisfies the class and tag
+		UClass* AttachComponentClass = WeakAttachComponentClass.Get();
+		if (AttachComponentClass || !AttachComponentTag.IsNone())
+		{
+			while (AttachComponent)
+			{
+				if ((!AttachComponentClass || AttachComponent->IsA(AttachComponentClass)) &&
+					(AttachComponentTag.IsNone() || AttachComponent->ComponentHasTag(AttachComponentTag)))
+				{
+					break;
+				}
+				AttachComponent = AttachComponent->GetAttachParent();
+			}
+		}
+
+		return AttachComponent;
+	}
+
+	virtual UObject* ResolveSourceInternal() const override
+	{
+		UObject* ResolvedObject = SourceActor.Get();
+		if (ResolvedObject)
+		{
+			if (AActor* AsActor = Cast<AActor>(ResolvedObject))
+			{
+				ResolvedObject = AsActor->GetRootComponent();
+			}
+		}
+		if (!ResolvedObject)
+		{
+			ResolvedObject = SourceAsset.Get();
+		}
+		return ResolvedObject;
+	}
+
+	virtual bool CacheSocketNamesInternal(UObject* ResolvedObject) override
+	{
+		if (USceneComponent* SceneComponent = Cast<USceneComponent>(ResolvedObject))
+		{
+			SocketNames = SceneComponent->GetAllSocketNames();
+			return true;
+		}
+		return false;
+	}
+
+	virtual bool GetSocketTransformsInternal(UObject* ResolvedObject, TArray<FTransform3f>&OutSocketTransforms, TOptional<TConstArrayView<int32>> SocketIndices) const override
+	{
+		if (USceneComponent* SceneComponent = Cast<USceneComponent>(ResolvedObject))
+		{
+			const bool bIndirect = SocketIndices.IsSet();
+			const int32 NumIndicies = bIndirect ? SocketIndices.GetValue().Num() : SocketNames.Num();
+			for (int iIndex = 0; iIndex < NumIndicies; ++iIndex)
+			{
+				const int32 iSocketIndex = bIndirect ? SocketIndices.GetValue()[iIndex] : iIndex;
+				if (SocketNames.IsValidIndex(iSocketIndex))
+				{
+					const FTransform SocketTransform = SceneComponent->GetSocketTransform(SocketNames[iSocketIndex], ERelativeTransformSpace::RTS_Component);
+					OutSocketTransforms[iSocketIndex] = FTransform3f(SocketTransform);
+				}
+			}
+
+			return true;
+		}
+		return false;
+	}
+};
diff --git a/Engine/Plugins/FX/Niagara/Source/Niagara/Private/NiagaraSceneComponentUtils.cpp b/Engine/Plugins/FX/Niagara/Source/Niagara/Private/NiagaraSceneComponentUtils.cpp
--- a/Engine/Plugins/FX/Niagara/Source/Niagara/Private/NiagaraSceneComponentUtils.cpp
+++ b/Engine/Plugins/FX/Niagara/Source/Niagara/Private/NiagaraSceneComponentUtils.cpp
@@ -3,6 +3,8 @@
 #include "NiagaraSceneComponentUtils.h"
 #include "NiagaraComponent.h"
 #include "NiagaraSystemInstance.h"
+#include "NiagaraSceneComponentSocketUtils.h"
+#include "NiagaraActorSceneComponentSocketUtils.h"
 
 #include "Components/StaticMeshComponent.h"
 #include "Components/InstancedStaticMeshComponent.h"
@@ -164,3 +166,8 @@ FVector FNiagaraActorSceneComponentUtils::GetPhysicsLinearVelocity(UObject* Comp
 	UPrimitiveComponent* PrimitiveComponent = Cast<UPrimitiveComponent>(Component);
 	return PrimitiveComponent ? PrimitiveComponent->GetPhysicsLinearVelocity() : FVector::ZeroVector;
 }
+
+FNiagaraSceneComponentSocketUtils* FNiagaraActorSceneComponentUtils::CreateSocketUtils() const
+{
+	return new FNiagaraActorSceneComponentSocketUtils(WeakOwnerComponent);
+}
diff --git a/Engine/Plugins/FX/Niagara/Source/Niagara/Public/NiagaraSceneComponentSocketUtils.h b/Engine/Plugins/FX/Niagara/Source/Niagara/Public/NiagaraSceneComponentSocketUtils.h
new file mode 100644
--- /dev/null
+++ b/Engine/Plugins/FX/Niagara/Source/Niagara/Public/NiagaraSceneComponentSocketUtils.h
@@ -0,0 +1,321 @@
+// Copyright Epic Games, Inc. All Rights Reserved.
+
+#pragma once
+
+#include "Math/MathFwd.h"
+#include "UObject/Object.h"
+#include "Engine/World.h"
+
+#include "NiagaraComponent.h"
+#include "NiagaraSystemInstance.h"
+
+#include "Components/SceneComponent.h"
+#include "Engine/Canvas.h"
+#include "Engine/SkeletalMesh.h"
+#include "Engine/SkeletalMeshSocket.h"
+#include "Engine/StaticMesh.h"
+#include "Engine/StaticMeshSocket.h"
+
+// Utilities to help reading sockets
+// DO NOT USE THIS IN EXTERNAL CODE as it is subject to change
+class FNiagaraSceneComponentSocketUtils
+{
+public:
+	FNiagaraSceneComponentSocketUtils(TWeakObjectPtr<UObject> InWeakOwnerComponent)
+		: WeakOwnerComponent(InWeakOwnerComponent)
+	{
+	}
+	virtual ~FNiagaraSceneComponentSocketUtils() = default;
+
+	// Enable resolving via a parameter binding
+	void SetAllowParameterBinding(FNiagaraSystemInstance* SystemInstance, const FNiagaraVariable& BoundParameter)
+	{
+		ParameterBinding.Init(SystemInstance->GetInstanceParameters(), BoundParameter);
+		bTryResolveBinding = ParameterBinding.IsBound();
+	}
+
+	// Enable resolving via attach parent
+	void SetAllowAttachParent(UClass* ComponentClass, FName ComponentTag)
+	{
+		bTryResolveAttachParent = true;
+		WeakAttachComponentClass = ComponentClass;
+		AttachComponentTag = ComponentTag;
+	}
+
+	// Enable resolving via a source object
+	void SetAllowSource(TLazyObjectPtr<UObject> InSourceActor, UObject* InSourceAsset)
+	{
+		SourceActor = InSourceActor;
+		SourceAsset = InSourceAsset;
+		bTryResolveSource = true;
+	}
+
+#if WITH_EDITOR
+	// Enable resolving via an editor only asset
+	void SetEditorPreviewAsset(TSoftObjectPtr<UObject>	InEditorPreviewAsset)
+	{
+		EditorPreviewAsset = InEditorPreviewAsset;
+	}
+#endif
+
+	// Are we bound to something valid?
+	bool IsValid() const
+	{
+		return WeakResolvedObject.Get() != nullptr;
+	}
+
+	// Resolve the object we are bound to
+	// Returns true if the object changed, otherwise false
+	bool ResolveObject()
+	{
+		UObject* ResolvedObject = nullptr;
+
+		// Resolve using the parameter binding?
+		if (bTryResolveBinding)
+		{
+			ResolvedObject = ParameterBinding.GetValue();
+			if (ResolvedObject)
+			{
+				if (AActor* AsActor = Cast<AActor>(ResolvedObject))
+				{
+					ResolvedObject = AsActor->GetRootComponent();
+				}
+			}
+		}
+
+		// Resolve using attach parent?
+		if (bTryResolveAttachParent && !ResolvedObject)
+		{
+			ResolvedObject = ResolveAttachParentInternal();
+		}
+
+		// Resolve Source?
+		if (bTryResolveSource && !ResolvedObject)
+		{
+			ResolvedObject = ResolveSourceInternal();
+		}
+
+#if WITH_EDITOR
+		// In editor fallback to the editor preview asset
+		if (!ResolvedObject)
+		{
+			USceneComponent* OwnerComponent = Cast<USceneComponent>(WeakOwnerComponent.Get());
+			UWorld* World = OwnerComponent ? OwnerComponent->GetWorld() : nullptr;
+			if (World && !World->IsGameWorld())
+			{
+				ResolvedObject = EditorPreviewAsset.LoadSynchronous();
+			}
+		}
+#endif
+
+		if (WeakResolvedObject.Get() == ResolvedObject)
+		{
+#if WITH_EDITOR
+			// When in the editor the socket counts can change so we might need to recache the socket list
+			TArray<FName> PrevSocketNames;
+			Swap(PrevSocketNames, SocketNames);
+			CacheSocketNames(ResolvedObject);
+
+			if (PrevSocketNames != SocketNames)
+			{
+				return true;
+			}
+#endif
+
+			return false;
+		}
+
+		WeakResolvedObject = ResolvedObject;
+		CacheSocketNames(ResolvedObject);
+
+		return true;
+	}
+
+	// Get the number of sockets available
+	int32 GetNumSockets() const
+	{
+		return SocketNames.Num();
+	}
+
+	// Get all of the available socket names
+	virtual TConstArrayView<FName> GetSocketNames() const
+	{
+		return MakeArrayView(SocketNames);
+	}
+
+	// Get socket transforms
+	// Optionally reading from only SocketIndices
+	// Returns true if we read any data, otherwise false
+	bool GetSocketTransforms(TArray<FTransform3f>& OutSocketTransforms, TOptional<TConstArrayView<int32>> SocketIndices = {}) const
+	{
+		UObject* ResolvedObject = WeakResolvedObject.Get();
+
+		if (GetSocketTransformsInternal(ResolvedObject, OutSocketTransforms, SocketIndices))
+		{
+			return true;
+		}
+
+		if (UStaticMesh* StaticMesh = Cast<UStaticMesh>(ResolvedObject))
+		{
+			GetSocketTransformsStaticMesh(StaticMesh, OutSocketTransforms, SocketIndices);
+			return true;
+		}
+
+		if (USkeletalMesh* SkeletalMesh = Cast<USkeletalMesh>(ResolvedObject))
+		{
+			GetSocketTransformsSkeletalMesh(SkeletalMesh, OutSocketTransforms, SocketIndices);
+			return true;
+		}
+
+		return false;
+	}
+
+	// Get the current local to world for the socket owner
+	virtual FTransform GetLocalToWorldTransform() const = 0;
+
+	// Get the resolve object
+	// Note: Do not make any assumptions about where this comes from
+	UObject* GetResolvedObject() const
+	{
+		return WeakResolvedObject.Get();
+	}
+
+protected:
+	void CacheSocketNames(UObject* ResolvedObject)
+	{
+		SocketNames.Reset();
+		if (CacheSocketNamesInternal(ResolvedObject))
+		{
+			return;
+		}
+		
+		if (UStaticMesh* StaticMesh = Cast<UStaticMesh>(ResolvedObject))
+		{
+			CacheSocketNamesStaticMesh(StaticMesh);
+			return;
+		}
+		
+		if (USkeletalMesh* SkeletalMesh = Cast<USkeletalMesh>(ResolvedObject))
+		{
+			CacheSocketNamesSkeletalMesh(SkeletalMesh);
+			return;
+		}
+		SocketNames.Empty();
+	}
+
+	void CacheSocketNamesStaticMesh(UStaticMesh* StaticMesh)
+	{
+		const int32 NumSockets = StaticMesh->Sockets.Num();
+		SocketNames.Reserve(StaticMesh->Sockets.Num());
+		for (UStaticMeshSocket* Socket : StaticMesh->Sockets)
+		{
+			check(Socket);
+			SocketNames.Add(Socket->SocketName);
+		}
+	}
+
+	void CacheSocketNamesSkeletalMesh(USkeletalMesh* SkeletalMesh)
+	{
+		const FReferenceSkeleton& RefSkeleton = SkeletalMesh->GetRefSkeleton();
+		const TArray<FMeshBoneInfo>& RefBoneInfo = RefSkeleton.GetRefBoneInfo();
+		const int32 NumBones = RefBoneInfo.Num();
+		const int32 NumSockets = SkeletalMesh->NumSockets();
+
+		SocketNames.Reserve(NumBones + NumSockets);
+		for (int i = 0; i < NumBones; ++i)
+		{
+			SocketNames.Add(RefBoneInfo[i].Name);
+		}
+
+		for (int i = 0; i < NumSockets; ++i)
+		{
+			USkeletalMeshSocket* Socket = SkeletalMesh->GetSocketByIndex(i);
+			check(Socket);
+			SocketNames.Add(Socket->SocketName);
+		}
+	}
+
+	void GetSocketTransformsStaticMesh(UStaticMesh* StaticMesh, TArray<FTransform3f>& OutSocketTransforms, TOptional<TConstArrayView<int32>> SocketIndices) const
+	{
+		const bool bIndirect = SocketIndices.IsSet();
+		const int32 NumIndicies = bIndirect ? SocketIndices.GetValue().Num() : SocketNames.Num();
+		for (int iIndex = 0; iIndex < NumIndicies; ++iIndex)
+		{
+			const int32 iSocketIndex = bIndirect ? SocketIndices.GetValue()[iIndex] : iIndex;
+			if (SocketNames.IsValidIndex(iSocketIndex) && StaticMesh->Sockets.IsValidIndex(iSocketIndex))
+			{
+				UStaticMeshSocket* Socket = StaticMesh->Sockets[iSocketIndex];
+				OutSocketTransforms[iSocketIndex].SetTranslation(FVector3f(Socket->RelativeLocation));
+				OutSocketTransforms[iSocketIndex].SetRotation(FQuat4f(Socket->RelativeRotation.Quaternion()));
+				OutSocketTransforms[iSocketIndex].SetScale3D(FVector3f(Socket->RelativeScale));
+			}
+		}
+	}
+	
+	void GetSocketTransformsSkeletalMesh(USkeletalMesh* SkeletalMesh, TArray<FTransform3f>& OutSocketTransforms, TOptional<TConstArrayView<int32>> SocketIndices) const
+	{
+		// For SKM we read everything all the time, as there is a hierarchy to the data that needs to be resolved
+		const FReferenceSkeleton& RefSkeleton = SkeletalMesh->GetRefSkeleton();
+		const TArray<FMeshBoneInfo>& RefBoneInfo = RefSkeleton.GetRefBoneInfo();
+		const TArray<FTransform>& RefBonePose = RefSkeleton.GetRefBonePose();
+		const int32 NumBones = RefBoneInfo.Num();
+		const int32 NumSockets = SkeletalMesh->NumSockets();
+
+		for (int i = 0; i < NumBones; ++i)
+		{
+			FTransform3f SocketTransform = FTransform3f(RefBonePose[i]);
+			const int32 ParentIndex = RefBoneInfo[i].ParentIndex;
+			if (ParentIndex >= 0 && ParentIndex < i)
+			{
+				SocketTransform = SocketTransform * OutSocketTransforms[ParentIndex];
+			}
+			OutSocketTransforms[i] = FTransform3f(SocketTransform);
+		}
+
+		for (int i = 0; i < NumSockets; ++i)
+		{
+			USkeletalMeshSocket* Socket = SkeletalMesh->GetSocketByIndex(i);
+			check(Socket);
+
+			FTransform3f SocketTransform(FTransform(Socket->RelativeRotation.Quaternion(), Socket->RelativeLocation, Socket->RelativeScale));
+			const int32 ParentIndex = RefSkeleton.FindBoneIndex(Socket->BoneName);
+			if (ParentIndex >= 0)
+			{
+				SocketTransform = SocketTransform * OutSocketTransforms[ParentIndex];
+			}
+			OutSocketTransforms[i + NumBones] = FTransform3f(SocketTransform);
+		}
+	}
+
+	// Resolve from attach parent Actor / Component / Entity
+	virtual UObject* ResolveAttachParentInternal() const = 0;
+	// Resolve from source Actor / Component / Entity
+	virtual UObject* ResolveSourceInternal() const = 0;
+
+	// Cache socket names from Actor / Component / Entity if valid
+	virtual bool CacheSocketNamesInternal(UObject* ResolvedObject) = 0;
+	// Read socket transforms from Actor / Component / Entity if valid
+	virtual bool GetSocketTransformsInternal(UObject* ResolvedObject, TArray<FTransform3f>& OutSocketTransforms, TOptional<TConstArrayView<int32>> SocketIndices) const = 0;
+
+protected:
+	bool										bTryResolveBinding = false;
+	bool										bTryResolveAttachParent = false;
+	bool										bTryResolveSource = false;
+
+	TWeakObjectPtr<UObject>						WeakOwnerComponent;
+	TWeakObjectPtr<UObject>						WeakResolvedObject;
+
+	FNiagaraParameterDirectBinding<UObject*>	ParameterBinding;
+
+	TWeakObjectPtr<UClass>						WeakAttachComponentClass;
+	FName										AttachComponentTag;
+
+	TLazyObjectPtr<UObject>						SourceActor;
+	TWeakObjectPtr<UObject>						SourceAsset;
+
+#if WITH_EDITOR
+	TSoftObjectPtr<UObject>						EditorPreviewAsset;
+#endif
+
+	TArray<FName>								SocketNames;
+};
diff --git a/Engine/Plugins/FX/Niagara/Source/Niagara/Public/NiagaraSceneComponentUtils.h b/Engine/Plugins/FX/Niagara/Source/Niagara/Public/NiagaraSceneComponentUtils.h
--- a/Engine/Plugins/FX/Niagara/Source/Niagara/Public/NiagaraSceneComponentUtils.h
+++ b/Engine/Plugins/FX/Niagara/Source/Niagara/Public/NiagaraSceneComponentUtils.h
@@ -11,6 +11,8 @@ class FColorVertexBuffer;
 class UNiagaraComponent;
 class UStaticMesh;
 
+class FNiagaraSceneComponentSocketUtils;
+
 // Helper class to abstract how we need to search the scene for components in various data interfaces
 // This is temporary until scene graph is folded into the core engine, or we have more official APIs to abstract Actors/Entity/Desc
 // DO NOT USE THIS IN EXTERNAL CODE as it is subject to change
@@ -25,35 +27,34 @@ public:
 	virtual void ResolveStaticMesh(UObject* ObjectFrom, bool bRecurseParents, UObject*& OutComponent, UStaticMesh*& OutStaticMesh) const = 0;
 	// Get the component transform and any ISM instance transforms
 	virtual bool GetStaticMeshTransforms(UObject* Component, FTransform& OutComponentTransform, TArray<FTransform>& OutInstanceTransforms) const = 0;
-	// Get the statich mesh override vertex colors, if any
+	// Get the static mesh override vertex colors, if any
 	virtual FColorVertexBuffer* GetStaticMeshOverrideColors(UObject* Component, int32 LODIndex) const = 0;
 
 	// Get the FPrimitiveComponentId for the provided component
 	virtual FPrimitiveComponentId GetPrimitiveSceneId(UObject* Component) const = 0;
-	// Get the physics lineary velocit for the provided component
+	// Get the physics lineary velociy for the provided component
 	virtual FVector GetPhysicsLinearVelocity(UObject* Component) const = 0;
+
+	// Create a socket utils
+	virtual FNiagaraSceneComponentSocketUtils* CreateSocketUtils() const = 0;
 };
 
 // Implementation for AActor & UActorComponents
 class FNiagaraActorSceneComponentUtils final : public INiagaraSceneComponentUtils
 {
 public:
-	FNiagaraActorSceneComponentUtils(UNiagaraComponent* OwnerComponent);
+	explicit FNiagaraActorSceneComponentUtils(UNiagaraComponent* OwnerComponent);
 
-	// Resolve the static mesh from what the interface's owner object
 	virtual void ResolveStaticMesh(bool bRecurseParents, UObject*& OutComponent, UStaticMesh*& OutStaticMesh) const override;
-	// Resolve the static mesh from the provided object which could be a static mesh / component / actor / entity / etc
 	virtual void ResolveStaticMesh(UObject* ObjectFrom, bool bRecurseParents, UObject*& OutComponent, UStaticMesh*& OutStaticMesh) const override;
-	// Get the component transform and any ISM instance transforms
 	virtual bool GetStaticMeshTransforms(UObject* Component, FTransform& OutComponentTransform, TArray<FTransform>& OutInstanceTransforms) const override;
-	// Get the statich mesh override vertex colors, if any
 	virtual FColorVertexBuffer* GetStaticMeshOverrideColors(UObject* Component, int32 LODIndex) const override;
 
-	// Get the FPrimitiveComponentId for the provided component
 	virtual FPrimitiveComponentId GetPrimitiveSceneId(UObject* Component) const override;
-	// Get the physics lineary velocit for the provided component
 	virtual FVector GetPhysicsLinearVelocity(UObject* Component) const override;
 
+	virtual FNiagaraSceneComponentSocketUtils* CreateSocketUtils() const override;
+
 private:
 	TWeakObjectPtr<USceneComponent> WeakOwnerComponent;
 };
