- Add missing pragma once

#rnx

diff --git a/Engine/Plugins/FX/Niagara/Source/Niagara/Private/NiagaraActorSceneComponentSocketUtils.h b/Engine/Plugins/FX/Niagara/Source/Niagara/Private/NiagaraActorSceneComponentSocketUtils.h
--- a/Engine/Plugins/FX/Niagara/Source/Niagara/Private/NiagaraActorSceneComponentSocketUtils.h
+++ b/Engine/Plugins/FX/Niagara/Source/Niagara/Private/NiagaraActorSceneComponentSocketUtils.h
@@ -1,5 +1,7 @@
 // Copyright Epic Games, Inc. All Rights Reserved.
 
+#pragma once
+
 #include "NiagaraSceneComponentSocketUtils.h"
 #include "NiagaraComponent.h"
 #include "NiagaraSystemInstance.h"
