[Niagara Validation] Add two new Niagara validation rules:
- "Require Latest Parent Emitter Version" will require that all emitters with a parent emitter set must be using the latest version of that parent emitter
- "Require Parent Emitter(s)" will require that all emitters have a parent emitter be from a set of allowed emitters (while optionally allowing no parent emitter to be set)

diff --git a/Engine/Plugins/FX/Niagara/Source/NiagaraEditor/Private/NiagaraValidationRules.cpp b/Engine/Plugins/FX/Niagara/Source/NiagaraEditor/Private/NiagaraValidationRules.cpp
--- a/Engine/Plugins/FX/Niagara/Source/NiagaraEditor/Private/NiagaraValidationRules.cpp
+++ b/Engine/Plugins/FX/Niagara/Source/NiagaraEditor/Private/NiagaraValidationRules.cpp
@@ -8,6 +8,7 @@
 #include "NiagaraDataInterfaceSkeletalMesh.h"
 #include "NiagaraDataInterfaceUtilities.h"
 #include "NiagaraEditorSettings.h"
+#include "NiagaraEditorUtilities.h"
 #include "NiagaraNodeFunctionCall.h"
 #include "NiagaraMeshRendererProperties.h"
 #include "NiagaraRibbonRendererProperties.h"
@@ -81,6 +82,22 @@ bool NiagaraValidation::HasValidationRules(UNiagaraSystem* NiagaraSystem)
 	return false;
 }
 
+static void ForEachEmitterInContext(const FNiagaraValidationContext& Context, TFunctionRef<void(const FVersionedNiagaraEmitter&, UNiagaraStackEmitterPropertiesItem*)> Callback)
+{
+	TArray<TSharedRef<FNiagaraEmitterHandleViewModel>> EmitterHandleViewModels = Context.ViewModel->GetEmitterHandleViewModels();
+	for (TSharedRef<FNiagaraEmitterHandleViewModel> EmitterHandleViewModel : EmitterHandleViewModels)
+	{
+		const FVersionedNiagaraEmitter VersionedEmitter = EmitterHandleViewModel->GetEmitterHandle()->GetInstance();
+		if (IsValid(VersionedEmitter.Emitter) == false)
+		{
+			continue;
+		}
+
+		UNiagaraStackEmitterPropertiesItem* EmitterProperties = NiagaraValidation::GetStackEntry<UNiagaraStackEmitterPropertiesItem>(EmitterHandleViewModel->GetEmitterStackViewModel());
+		Callback(VersionedEmitter, EmitterProperties);
+	}
+}
+
 // --------------------------------------------------------------------------------------------------------------------------------------------
 
 void NiagaraValidation::ValidateAllRulesInSystem(TSharedPtr<FNiagaraSystemViewModel> SysViewModel, TFunction<void(const FNiagaraValidationResult& Result)> ResultCallback)
@@ -1769,4 +1786,97 @@ void UNiagaraValidationRule_MaterialUsage::CheckValidity(const FNiagaraValidatio
 	}
 }
 
+void UNiagaraValidationRule_RequireLatestParentEmitterVersion::CheckValidity(const FNiagaraValidationContext& Context, TArray<FNiagaraValidationResult>& OutResults) const
+{
+	ForEachEmitterInContext(Context, [&, this](const FVersionedNiagaraEmitter& VersionedEmitter, UNiagaraStackEmitterPropertiesItem* EmitterProperties)
+	{
+		VersionedEmitter.Emitter->ForEachVersionData([&, this](const FVersionedNiagaraEmitterData& VersionedEmitterData)
+		{
+			const UNiagaraEmitter* ParentEmitter = VersionedEmitterData.VersionedParent.Emitter.Get();
+			if (IsValid(ParentEmitter) == false)
+			{
+				return;
+			}
+
+			const FGuid UsedParentVersionGuid = VersionedEmitterData.VersionedParent.Version;
+			const FVersionedNiagaraEmitterData* UsedParentVersionData = ParentEmitter->GetEmitterData(UsedParentVersionGuid);
+			if (ensureMsgf(UsedParentVersionData != nullptr, TEXT("Using unknown or invalid parent emitter version")) == false)
+			{
+				return;
+			}
+
+			if (const FVersionedNiagaraEmitterData* LatestParentVersionData = VersionedEmitterData.VersionedParent.Emitter->GetLatestEmitterData())
+			{
+				if (LatestParentVersionData->Version <= UsedParentVersionData->Version)
+				{
+					return;
+				}
+
+				FNiagaraValidationResult& Result = OutResults.AddDefaulted_GetRef();
+				Result.Severity = Severity;
+				Result.SummaryText = FText::Format(LOCTEXT("NotLatestParentVersion_Title", "Not using latest version of parent emitter '{0}'"), FText::FromName(*GetPathNameSafe(ParentEmitter)));
+				Result.Description = FText::Format(LOCTEXT("NotLatestParentVersion_Description", "Using version '{0}.{1}', but latest is '{2}.{3}'. Please consider updating to the latest version"),
+					UsedParentVersionData->Version.MajorVersion, UsedParentVersionData->Version.MinorVersion,
+					LatestParentVersionData->Version.MajorVersion, LatestParentVersionData->Version.MinorVersion);
+				Result.SourceObject = TWeakObjectPtr { EmitterProperties };
+
+				FNiagaraValidationFix& UpdateToLatestParentFix = Result.Fixes.AddDefaulted_GetRef();
+				UpdateToLatestParentFix.Description = LOCTEXT("UpgradeToLatestParentVersion", "Upgrade to latest parent version.");
+				UpdateToLatestParentFix.FixDelegate = FNiagaraValidationFixDelegate::CreateLambda([EmitterProperties, LatestParentVersionGuid = LatestParentVersionData->Version.VersionGuid]()
+				{
+					TSharedRef<FNiagaraEmitterViewModel> EmitterViewModel = EmitterProperties->GetEmitterViewModel().ToSharedRef();
+					TSharedRef<FNiagaraSystemViewModel> SystemViewModel = EmitterProperties->GetSystemViewModel();
+					FNiagaraEditorUtilities::SwitchParentEmitterVersion(EmitterViewModel, SystemViewModel, LatestParentVersionGuid);
+				});
+			}
+		});
+	});
+}
+
+template<class ObjectType>
+static bool ObjectIsInSoftObjectArray(const TArray<TSoftObjectPtr<ObjectType>>& ObjectArray, const ObjectType* ObjectToFind)
+{
+	const TSoftObjectPtr<const ObjectType> SoftObjectToFindPtr { ObjectToFind };
+	return ObjectArray.Contains(SoftObjectToFindPtr);
+}
+
+void UNiagaraValidationRule_RequireParentEmitter::CheckValidity(const FNiagaraValidationContext& Context, TArray<FNiagaraValidationResult>& OutResults) const
+{
+	if (AllowedParentEmitters.IsEmpty())
+	{
+		return;
+	}
+
+	ForEachEmitterInContext(Context, [&, this](const FVersionedNiagaraEmitter& VersionedEmitter, UNiagaraStackEmitterPropertiesItem* EmitterProperties)
+	{
+		VersionedEmitter.Emitter->ForEachVersionData([&, this](const FVersionedNiagaraEmitterData& VersionedEmitterData)
+		{
+			const UNiagaraEmitter* ParentEmitter = VersionedEmitterData.VersionedParent.Emitter.Get();
+			if (IsValid(ParentEmitter) == false)
+			{
+				if (bMustHaveParentEmitter)
+				{
+					FNiagaraValidationResult& Result = OutResults.AddDefaulted_GetRef();
+					Result.Severity = Severity;
+					Result.SummaryText = LOCTEXT("NoParentEmitter_Title", "No parent emitter set");
+					Result.Description = LOCTEXT("NoParentEmitter_Description", "The emitter does not have a parent emitter. Please set one");
+					Result.SourceObject = TWeakObjectPtr { EmitterProperties };
+				}
+				return;
+			}
+
+			if (ObjectIsInSoftObjectArray(AllowedParentEmitters, ParentEmitter))
+			{
+				return;
+			}
+
+			FNiagaraValidationResult& Result = OutResults.AddDefaulted_GetRef();
+			Result.Severity = Severity;
+			Result.SummaryText = FText::Format(LOCTEXT("InvalidParentEmitter_Title", "Parent emitter '{0}' is invalid"), FText::FromName(*GetPathNameSafe(ParentEmitter)));
+			Result.Description = LOCTEXT("InvalidParentEmitter_Description", "The emitter is not explicitly allowed. Please consider using a different parent emitter");
+			Result.SourceObject = TWeakObjectPtr { EmitterProperties };
+		});
+	});
+}
+
 #undef LOCTEXT_NAMESPACE
diff --git a/Engine/Plugins/FX/Niagara/Source/NiagaraEditor/Public/NiagaraValidationRules.h b/Engine/Plugins/FX/Niagara/Source/NiagaraEditor/Public/NiagaraValidationRules.h
--- a/Engine/Plugins/FX/Niagara/Source/NiagaraEditor/Public/NiagaraValidationRules.h
+++ b/Engine/Plugins/FX/Niagara/Source/NiagaraEditor/Public/NiagaraValidationRules.h
@@ -578,3 +578,39 @@ public:
 
 	virtual void CheckValidity(const FNiagaraValidationContext& Context, TArray<FNiagaraValidationResult>& OutResults) const override;
 };
+
+/** Checks that all emitters with a parent emitter are using the latest version of the parent. */
+UCLASS(Category = "Validation", DisplayName = "Require Latest Parent Emitter Version")
+class UNiagaraValidationRule_RequireLatestParentEmitterVersion : public UNiagaraValidationRule
+{
+	GENERATED_BODY()
+
+public:
+
+	UPROPERTY(EditAnywhere, Category = Validation)
+	ENiagaraValidationSeverity Severity = ENiagaraValidationSeverity::Warning;
+
+	virtual void CheckValidity(const FNiagaraValidationContext& Context, TArray<FNiagaraValidationResult>& OutResults) const override;
+};
+
+/** Checks that all emitters with a parent emitter are using one of a set of allowed parent emitters. */
+UCLASS(Category = "Validation", DisplayName = "Require Parent Emitter(s)")
+class UNiagaraValidationRule_RequireParentEmitter : public UNiagaraValidationRule
+{
+	GENERATED_BODY()
+
+public:
+
+	/** The collection of allowed parent emitters */
+	UPROPERTY(EditAnywhere, Category = Validation)
+	TArray<TSoftObjectPtr<UNiagaraEmitter>> AllowedParentEmitters;
+
+	/** True if emitters must have a parent emitter, false to allow no parent emitter */
+	UPROPERTY(EditAnywhere, Category = Validation)
+	bool bMustHaveParentEmitter = true;
+
+	UPROPERTY(EditAnywhere, Category = Validation)
+	ENiagaraValidationSeverity Severity = ENiagaraValidationSeverity::Error;
+
+	virtual void CheckValidity(const FNiagaraValidationContext& Context, TArray<FNiagaraValidationResult>& OutResults) const override;
+};
