[Niagara Validation] Fix potential GC issue with capturing raw pointer in fix callback

diff --git a/Engine/Plugins/FX/Niagara/Source/NiagaraEditor/Private/NiagaraValidationRules.cpp b/Engine/Plugins/FX/Niagara/Source/NiagaraEditor/Private/NiagaraValidationRules.cpp
--- a/Engine/Plugins/FX/Niagara/Source/NiagaraEditor/Private/NiagaraValidationRules.cpp
+++ b/Engine/Plugins/FX/Niagara/Source/NiagaraEditor/Private/NiagaraValidationRules.cpp
@@ -1822,7 +1822,7 @@ void UNiagaraValidationRule_RequireLatestParentEmitterVersion::CheckValidity(con
 
 				FNiagaraValidationFix& UpdateToLatestParentFix = Result.Fixes.AddDefaulted_GetRef();
 				UpdateToLatestParentFix.Description = LOCTEXT("UpgradeToLatestParentVersion", "Upgrade to latest parent version.");
-				UpdateToLatestParentFix.FixDelegate = FNiagaraValidationFixDelegate::CreateLambda([EmitterProperties, LatestParentVersionGuid = LatestParentVersionData->Version.VersionGuid]()
+				UpdateToLatestParentFix.FixDelegate = FNiagaraValidationFixDelegate::CreateWeakLambda(EmitterProperties, [EmitterProperties, LatestParentVersionGuid = LatestParentVersionData->Version.VersionGuid]()
 				{
 					TSharedRef<FNiagaraEmitterViewModel> EmitterViewModel = EmitterProperties->GetEmitterViewModel().ToSharedRef();
 					TSharedRef<FNiagaraSystemViewModel> SystemViewModel = EmitterProperties->GetSystemViewModel();
