- Fix crash in data set readback if a GPU emitter was disabled by scalability (for example) and a readback was requested

#tests Enable validation in Debug HUD, add a GPU emitter, disable GPU via scalability

diff --git a/Engine/Plugins/FX/Niagara/Source/Niagara/Private/NiagaraDataSetReadback.cpp b/Engine/Plugins/FX/Niagara/Source/Niagara/Private/NiagaraDataSetReadback.cpp
--- a/Engine/Plugins/FX/Niagara/Source/Niagara/Private/NiagaraDataSetReadback.cpp
+++ b/Engine/Plugins/FX/Niagara/Source/Niagara/Private/NiagaraDataSetReadback.cpp
@@ -52,7 +52,10 @@ void FNiagaraDataSetReadback::EnqueueReadback(FNiagaraEmitterInstance* EmitterIn
 		}
 		else
 		{
-			DataSet.CopyFromGPUReadback(nullptr, nullptr, nullptr, 0, 0, 0, 0, 0);
+			FNiagaraDataBuffer& DestBuffer = DataSet.BeginSimulate();
+			DestBuffer.SetNumInstances(0);
+			DestBuffer.Allocate(0);
+			DataSet.EndSimulate();
 		}
 		ParameterStore = StatefulEmitterInstance->GetUpdateExecutionContext().Parameters;
 	}
