- Improve tooltip information for the actor / component DI

#rnx

diff --git a/Engine/Plugins/FX/Niagara/Source/Niagara/Private/DataInterface/NiagaraDataInterfaceActorComponent.cpp b/Engine/Plugins/FX/Niagara/Source/Niagara/Private/DataInterface/NiagaraDataInterfaceActorComponent.cpp
--- a/Engine/Plugins/FX/Niagara/Source/Niagara/Private/DataInterface/NiagaraDataInterfaceActorComponent.cpp
+++ b/Engine/Plugins/FX/Niagara/Source/Niagara/Private/DataInterface/NiagaraDataInterfaceActorComponent.cpp
@@ -201,7 +201,7 @@ void UNiagaraDataInterfaceActorComponent::GetFunctionsInternal(TArray<FNiagaraFu
 		FunctionSignature.Name = GetMatrixName;
 		FunctionSignature.Outputs.Emplace(FNiagaraTypeDefinition::GetBoolDef(), TEXT("IsValid"));
 		FunctionSignature.Outputs.Emplace(FNiagaraTypeDefinition::GetMatrix4Def(), TEXT("Matrix"));
-		FunctionSignature.SetDescription(LOCTEXT("GetMatrix", "Returns the current matrix for the component if valid."));
+		FunctionSignature.SetDescription(LOCTEXT("GetMatrix", "Returns the actor/component to world matrix if valid, otherwise identity."));
 	}
 	{
 		FNiagaraFunctionSignature& FunctionSignature = OutFunctions.Add_GetRef(DefaultSig);
@@ -210,14 +210,14 @@ void UNiagaraDataInterfaceActorComponent::GetFunctionsInternal(TArray<FNiagaraFu
 		FunctionSignature.Outputs.Emplace(FNiagaraTypeDefinition::GetPositionDef(), TEXT("Position"));
 		FunctionSignature.Outputs.Emplace(FNiagaraTypeDefinition::GetQuatDef(), TEXT("Rotation"));
 		FunctionSignature.Outputs.Emplace(FNiagaraTypeDefinition::GetVec3Def(), TEXT("Scale"));
-		FunctionSignature.SetDescription(LOCTEXT("GetTransform", "Returns the current transform for the component if valid."));
+		FunctionSignature.SetDescription(LOCTEXT("GetTransform", "Returns the actor/component to world transform if valid, otherwise identity."));
 	}
 	{
 		FNiagaraFunctionSignature& FunctionSignature = OutFunctions.Add_GetRef(DefaultSig);
 		FunctionSignature.Name = GetVelocityName;
 		FunctionSignature.Outputs.Emplace(FNiagaraTypeDefinition::GetBoolDef(), TEXT("IsValid"));
 		FunctionSignature.Outputs.Emplace(FNiagaraTypeDefinition::GetVec3Def(), TEXT("Velocity"));
-		FunctionSignature.SetDescription(LOCTEXT("GetVelocityDesc", "Returns the velocity from the component or actor if valid."));
+		FunctionSignature.SetDescription(LOCTEXT("GetVelocityDesc", "Returns the actor/component world space velocity if valid, otherwise zero."));
 	}
 }
 #endif
