- Add more visible information about LWE performance

#tests

diff --git a/Engine/Plugins/FX/Niagara/Source/Niagara/Classes/NiagaraSystem.h b/Engine/Plugins/FX/Niagara/Source/Niagara/Classes/NiagaraSystem.h
--- a/Engine/Plugins/FX/Niagara/Source/Niagara/Classes/NiagaraSystem.h
+++ b/Engine/Plugins/FX/Niagara/Source/Niagara/Classes/NiagaraSystem.h
@@ -850,7 +850,7 @@ public:
 #endif
 
 	/** Access the code system state data. */
-	const FNiagaraSystemStateData& GetSystemStateData() const;
+	NIAGARA_API const FNiagaraSystemStateData& GetSystemStateData() const;
 
 	/** Used for debug HUD / viewport to convey what mode we are running in */
 	NIAGARA_API const TCHAR* GetSystemStateModeString() const;
diff --git a/Engine/Plugins/FX/Niagara/Source/NiagaraEditorWidgets/Private/SNiagaraOverviewGraphTitleBar.cpp b/Engine/Plugins/FX/Niagara/Source/NiagaraEditorWidgets/Private/SNiagaraOverviewGraphTitleBar.cpp
--- a/Engine/Plugins/FX/Niagara/Source/NiagaraEditorWidgets/Private/SNiagaraOverviewGraphTitleBar.cpp
+++ b/Engine/Plugins/FX/Niagara/Source/NiagaraEditorWidgets/Private/SNiagaraOverviewGraphTitleBar.cpp
@@ -5,6 +5,7 @@
 #include "NiagaraEditorWidgetsStyle.h"
 #include "NiagaraEditorSettings.h"
 #include "NiagaraEditorUtilities.h"
+#include "NiagaraSystemEmitterState.h"
 #include "ViewModels/NiagaraOverviewGraphViewModel.h"
 #include "SNiagaraScalabilityPreviewSettings.h"
 #include "SNiagaraSystemEffectTypeBar.h"
@@ -129,6 +130,27 @@ void SNiagaraOverviewGraphTitleBar::RebuildWidget()
 		];
 	}
 
+	// Clear information about LWE & performance
+	if (SystemViewModel.IsValid() && SystemViewModel.Pin()->GetEditMode() == ENiagaraSystemViewModelEditMode::SystemAsset)
+	{
+		ContainerWidget->AddSlot()
+		.HAlign(HAlign_Fill)
+		.AutoHeight()
+		[
+			SNew(SBorder)
+			.BorderImage(FAppStyle::GetBrush(TEXT("Graph.TitleBackground")))
+			.ColorAndOpacity(this, &SNiagaraOverviewGraphTitleBar::GetStatelessInfoHeaderColor)
+			.Visibility(this, &SNiagaraOverviewGraphTitleBar::GetStatelessInfoVisibility)
+			.HAlign(HAlign_Fill)
+			[
+				SNew(STextBlock)
+				.Text(this, &SNiagaraOverviewGraphTitleBar::GetStatelessInfoText)
+				.TextStyle(FAppStyle::Get(), TEXT("GraphBreadcrumbButtonText"))
+				.Justification(ETextJustify::Center)
+			]
+		];
+	}
+
 	if(bScalabilityModeActive)
 	{
 		ContainerWidget->AddSlot()
@@ -209,9 +231,9 @@ EVisibility SNiagaraOverviewGraphTitleBar::GetSystemCompileModeHeaderVisibility(
 	if (SystemViewModel.IsValid())
 	{
 		const UNiagaraSystem& NiagaraSystem = SystemViewModel.Pin()->GetSystem();
-		return NiagaraSystem.GetCompileForEdit() ? EVisibility::Hidden : EVisibility::Visible;
+		return NiagaraSystem.GetCompileForEdit() ? EVisibility::Collapsed : EVisibility::Visible;
 	}
-	return EVisibility::Hidden;
+	return EVisibility::Collapsed;
 }
 
 FText SNiagaraOverviewGraphTitleBar::GetSystemCompileModeHeaderText() const
@@ -219,6 +241,64 @@ FText SNiagaraOverviewGraphTitleBar::GetSystemCompileModeHeaderText() const
 	return LOCTEXT("SystemInCompiledForPerformance", "Note: System is compiled for performance.  Edits may result in more compilation.");
 }
 
+FLinearColor SNiagaraOverviewGraphTitleBar::GetStatelessInfoHeaderColor() const
+{
+	if (SystemViewModel.IsValid())
+	{
+		const UNiagaraSystem& NiagaraSystem = SystemViewModel.Pin()->GetSystem();
+		const FNiagaraSystemStateData& SystemStateData = NiagaraSystem.GetSystemStateData();
+		if (SystemStateData.bRunUpdateScript)
+		{
+			return FLinearColor::Red;
+		}
+		else if (SystemStateData.bRunSpawnScript)
+		{
+			return FLinearColor::Yellow;
+		}
+		return FLinearColor::Green;
+	}
+	return FLinearColor::White;
+}
+
+EVisibility SNiagaraOverviewGraphTitleBar::GetStatelessInfoVisibility() const
+{
+	if (SystemViewModel.IsValid())
+	{
+		const UNiagaraSystem& NiagaraSystem = SystemViewModel.Pin()->GetSystem();
+		if (NiagaraSystem.GetEmitterHandles().Num() > 0)
+		{
+			for (const FNiagaraEmitterHandle& EmitterHandle : NiagaraSystem.GetEmitterHandles())
+			{
+				if (EmitterHandle.GetIsEnabled() && EmitterHandle.GetEmitterMode() != ENiagaraEmitterMode::Stateless)
+				{
+					return EVisibility::Collapsed;
+				}
+			}
+			return EVisibility::Visible;
+		}
+	}
+	return EVisibility::Collapsed;
+}
+
+FText SNiagaraOverviewGraphTitleBar::GetStatelessInfoText() const
+{
+	if (SystemViewModel.IsValid())
+	{
+		const UNiagaraSystem& NiagaraSystem = SystemViewModel.Pin()->GetSystem();
+		const FNiagaraSystemStateData& SystemStateData = NiagaraSystem.GetSystemStateData();
+		if (SystemStateData.bRunUpdateScript)
+		{
+			return LOCTEXT("LWEInfo_WorstPerf", "LWE: Spawn & update scripts will run (Worst Perf)");
+		}
+		else if (SystemStateData.bRunSpawnScript)
+		{
+			return LOCTEXT("LWEInfo_GoodPerf", "LWE: Spawn script only will run (Good Perf)");
+		}
+		return LOCTEXT("LWEInfo_BestPerf", "LWE: No scripts will run (Best Perf)");
+	}
+	return FText::GetEmpty();
+}
+
 FLinearColor SNiagaraOverviewGraphTitleBar::GetSystemCompileModeHeaderColor() const
 {
 	return FNiagaraEditorWidgetsStyle::Get().GetColor("NiagaraEditor.SystemOverview.AffectedAssetsWarningColor");
diff --git a/Engine/Plugins/FX/Niagara/Source/NiagaraEditorWidgets/Private/SNiagaraOverviewGraphTitleBar.h b/Engine/Plugins/FX/Niagara/Source/NiagaraEditorWidgets/Private/SNiagaraOverviewGraphTitleBar.h
--- a/Engine/Plugins/FX/Niagara/Source/NiagaraEditorWidgets/Private/SNiagaraOverviewGraphTitleBar.h
+++ b/Engine/Plugins/FX/Niagara/Source/NiagaraEditorWidgets/Private/SNiagaraOverviewGraphTitleBar.h
@@ -48,6 +48,10 @@ private:
 	FText GetSystemCompileModeHeaderText() const;
 	FLinearColor GetSystemCompileModeHeaderColor() const;
 
+	FLinearColor GetStatelessInfoHeaderColor() const;
+	EVisibility GetStatelessInfoVisibility() const;
+	FText GetStatelessInfoText() const;
+
 	int32 GetEmitterAffectedAssets() const;
 	bool IsUsingDeprecatedEmitter() const;
 	void ResetAssetCount(const FAssetData& InAssetData);
