Fixed operator inconsistency in cache comparison and possible divide-by-zero error in read function

#tests ran all the existing cache engine tests as verification

diff --git a/Engine/Plugins/FX/Niagara/Source/Niagara/Private/NiagaraSimCache.cpp b/Engine/Plugins/FX/Niagara/Source/Niagara/Private/NiagaraSimCache.cpp
--- a/Engine/Plugins/FX/Niagara/Source/Niagara/Private/NiagaraSimCache.cpp
+++ b/Engine/Plugins/FX/Niagara/Source/Niagara/Private/NiagaraSimCache.cpp
@@ -236,7 +236,7 @@ bool FNiagaraSimCacheDataBuffers::operator==(const FNiagaraSimCacheDataBuffers&
 	bMatches &= ThisDataBuffer.Num() == OtherDataBuffer.Num();
 	if (bMatches)
 	{
-		bMatches &= FMemory::Memcmp(ThisDataBuffer.GetData(), OtherDataBuffer.GetData(), ThisDataBuffer.Num()) != 0;
+		bMatches &= FMemory::Memcmp(ThisDataBuffer.GetData(), OtherDataBuffer.GetData(), ThisDataBuffer.Num()) == 0;
 	}
 
 	return bMatches;
@@ -1050,7 +1050,8 @@ bool UNiagaraSimCache::Read(float TimeSeconds, FNiagaraSystemInstance* SystemIns
 	{
 		const float SimulationAgeA = CacheFrames[FrameIndex].SimulationAge;
 		const float SimulationAgeB = CacheFrames[FMath::Min(FrameIndex + 1, NumFramesMinusOne)].SimulationAge;
-		FrameFraction = (TimeSeconds - SimulationAgeA) / (SimulationAgeB - SimulationAgeA);
+		const float AgeDiff = SimulationAgeB - SimulationAgeA;
+		FrameFraction = FMath::IsNearlyEqual(AgeDiff, 0.0f) ? 0.0f : (TimeSeconds - SimulationAgeA) / AgeDiff;
 		if (!CreateParameters.bAllowInterpolation && !CreateParameters.bAllowVelocityExtrapolation && FrameFraction > 0.5f)
 		{
 			FrameFraction = 0.0f;
