- SimCache: Change Cache Buffer selection style to look more consistent with the component selection

#rnx

diff --git a/Engine/Plugins/FX/Niagara/Source/NiagaraEditor/Private/Widgets/SNiagaraSimCacheOverview.cpp b/Engine/Plugins/FX/Niagara/Source/NiagaraEditor/Private/Widgets/SNiagaraSimCacheOverview.cpp
--- a/Engine/Plugins/FX/Niagara/Source/NiagaraEditor/Private/Widgets/SNiagaraSimCacheOverview.cpp
+++ b/Engine/Plugins/FX/Niagara/Source/NiagaraEditor/Private/Widgets/SNiagaraSimCacheOverview.cpp
@@ -155,20 +155,8 @@ void SNiagaraSimCacheOverview::RebuildMainWidget()
 
 TSharedRef<ITableRow> SNiagaraSimCacheOverview::OnGenerateRowForItem(TSharedRef<FNiagaraSimCacheOverviewItem> Item, const TSharedRef<STableViewBase>& Owner)
 {
-	static const char* ItemStyles[] =
-	{
-		"NiagaraEditor.SimCache.SystemItem",
-		"NiagaraEditor.SimCache.EmitterItem",
-		"NiagaraEditor.SimCache.ComponentItem",
-		"NiagaraEditor.SimCache.DataInterfaceItem",
-		"NiagaraEditor.SimCache.DebugData",
-	};
-	static_assert(UE_ARRAY_COUNT(ItemStyles) == int(ENiagaraSimCacheOverviewItemType::MAX), "Mismatch on style count");
-
-	ENiagaraSimCacheOverviewItemType StyleType = Item->GetType();
-	
 	return SNew(STableRow<TSharedRef<FNiagaraSimCacheOverviewItem>>, Owner)
-	.Style(FNiagaraEditorStyle::Get(), ItemStyles[static_cast<int32>(StyleType)])
+	.Style(FNiagaraEditorStyle::Get(), "NiagaraEditor.SimCache.ComponentItem")
 	.Padding(1.0f)
 	[
 		SNew(SNiagaraSimCacheBufferItem)
