- Add missing reference to OutFunctionCalls

#rnx

diff --git a/Engine/Plugins/FX/Niagara/Source/NiagaraEditor/Private/ViewModels/Stack/NiagaraStackGraphUtilities.cpp b/Engine/Plugins/FX/Niagara/Source/NiagaraEditor/Private/ViewModels/Stack/NiagaraStackGraphUtilities.cpp
--- a/Engine/Plugins/FX/Niagara/Source/NiagaraEditor/Private/ViewModels/Stack/NiagaraStackGraphUtilities.cpp
+++ b/Engine/Plugins/FX/Niagara/Source/NiagaraEditor/Private/ViewModels/Stack/NiagaraStackGraphUtilities.cpp
@@ -2409,7 +2409,7 @@ void ConnectModuleNode(UNiagaraNodeFunctionCall& ModuleNode, UNiagaraNodeOutput&
 	FNiagaraStackGraphUtilities::ConnectStackNodeGroup(ModuleGroup, TargetInsertPreviousGroup, TargetInsertGroup);
 }
 
-bool FNiagaraStackGraphUtilities::FindScriptModulesInStack(FAssetData ModuleScriptAsset, UNiagaraNodeOutput& TargetOutputNode, TArray<UNiagaraNodeFunctionCall*> OutFunctionCalls)
+bool FNiagaraStackGraphUtilities::FindScriptModulesInStack(FAssetData ModuleScriptAsset, UNiagaraNodeOutput& TargetOutputNode, TArray<UNiagaraNodeFunctionCall*>& OutFunctionCalls)
 {
 	UNiagaraGraph* Graph = TargetOutputNode.GetNiagaraGraph();
 	TArray<UNiagaraNode*> Nodes;
diff --git a/Engine/Plugins/FX/Niagara/Source/NiagaraEditor/Public/ViewModels/Stack/NiagaraStackGraphUtilities.h b/Engine/Plugins/FX/Niagara/Source/NiagaraEditor/Public/ViewModels/Stack/NiagaraStackGraphUtilities.h
--- a/Engine/Plugins/FX/Niagara/Source/NiagaraEditor/Public/ViewModels/Stack/NiagaraStackGraphUtilities.h
+++ b/Engine/Plugins/FX/Niagara/Source/NiagaraEditor/Public/ViewModels/Stack/NiagaraStackGraphUtilities.h
@@ -287,7 +287,7 @@ namespace FNiagaraStackGraphUtilities
 
 	NIAGARAEDITOR_API UNiagaraNodeFunctionCall* AddScriptModuleToStack(UNiagaraScript* ModuleScript, UNiagaraNodeOutput& TargetOutputNode, int32 TargetIndex = INDEX_NONE, FString SuggestedName = FString(), const FGuid& VersionGuid = FGuid());
 	
-	bool FindScriptModulesInStack(FAssetData ModuleScriptAsset, UNiagaraNodeOutput& TargetOutputNode, TArray<UNiagaraNodeFunctionCall*> OutFunctionCalls);
+	bool FindScriptModulesInStack(FAssetData ModuleScriptAsset, UNiagaraNodeOutput& TargetOutputNode, TArray<UNiagaraNodeFunctionCall*>& OutFunctionCalls);
 
 	NIAGARAEDITOR_API UNiagaraNodeAssignment* AddParameterModuleToStack(const TArray<FNiagaraVariable>& ParameterVariables, UNiagaraNodeOutput& TargetOutputNode, int32 TargetIndex, const TArray<FString>& InDefaultValues);
 		
