- Change LWC tile update behavior to rebase
- When an instance crosses too many tiles we used to reset, this will rebase all the positions
- There may be cases where data interfaces store internal data and do not handle this

diff --git a/Engine/Plugins/FX/Niagara/Source/Niagara/Public/NiagaraSettings.h b/Engine/Plugins/FX/Niagara/Source/Niagara/Public/NiagaraSettings.h
--- a/Engine/Plugins/FX/Niagara/Source/Niagara/Public/NiagaraSettings.h
+++ b/Engine/Plugins/FX/Niagara/Source/Niagara/Public/NiagaraSettings.h
@@ -151,7 +151,7 @@ class UNiagaraSettings : public UDeveloperSettings
 
 	/** Controls how we respond to a system instance crossing large world coordinate tiles. Can be overriden on individual assets. */
 	UPROPERTY(config, EditAnywhere, Category = Niagara, meta = (EditCondition = "bSystemsSupportLargeWorldCoordinates", EditConditionHides))
-	ENiagaraLwcTileUpdateMode LargeWorldCoordinateTileUpdateMode = ENiagaraLwcTileUpdateMode::ResetSimulation;
+	ENiagaraLwcTileUpdateMode LargeWorldCoordinateTileUpdateMode = ENiagaraLwcTileUpdateMode::Rebase;
 
 	/** The maximum tiles we can cross before forcing a simulation to reset, if the tile update mode is set to RebaseOrResetSimulation. */
 	UPROPERTY(config, EditAnywhere, Category = Niagara, meta = (EditCondition = "bSystemsSupportLargeWorldCoordinates", EditConditionHides))
