- LWE: Add Noise Field Offset to curl noise
- Allows you to emit from a single point and have random noise

#virtualized

diff --git a/Engine/Plugins/FX/Niagara/Shaders/Private/Stateless/Modules/NiagaraStatelessModule_SolveVelocitiesAndForces.ush b/Engine/Plugins/FX/Niagara/Shaders/Private/Stateless/Modules/NiagaraStatelessModule_SolveVelocitiesAndForces.ush
--- a/Engine/Plugins/FX/Niagara/Shaders/Private/Stateless/Modules/NiagaraStatelessModule_SolveVelocitiesAndForces.ush
+++ b/Engine/Plugins/FX/Niagara/Shaders/Private/Stateless/Modules/NiagaraStatelessModule_SolveVelocitiesAndForces.ush
@@ -29,6 +29,8 @@ float3	SolveVelocitiesAndForces_PointOrigin;
 uint				SolveVelocitiesAndForces_NoiseEnabled;
 float				SolveVelocitiesAndForces_NoiseStrength;
 float				SolveVelocitiesAndForces_NoiseFrequency;
+float3				SolveVelocitiesAndForces_NoiseFieldOffsetScale;
+float3				SolveVelocitiesAndForces_NoiseFieldOffsetBias;
 float				SolveVelocitiesAndForces_NoiseLUTNormalize;
 uint				SolveVelocitiesAndForces_NoiseLUTRowMask;
 uint				SolveVelocitiesAndForces_NoiseLUTRowWidth;
@@ -43,9 +45,9 @@ struct FStatelessModule_SolveVelocitiesAndForces
 	float3 Acceleration;
 };
 
-void SolveVelocitiesAndForces_EvaluateNoise(in FStatelessModule_SolveVelocitiesAndForces ModuleData, float Lifetime, float NormalizedAge, float Age, inout float3 Position)
+void SolveVelocitiesAndForces_EvaluateNoise(in FStatelessModule_SolveVelocitiesAndForces ModuleData, float Lifetime, float NormalizedAge, float Age, float3 NoiseFieldOffset, inout float3 Position)
 {
-	const float3 SamplePosition	= Position * SolveVelocitiesAndForces_NoiseFrequency;
+	const float3 SamplePosition	= (Position + NoiseFieldOffset) * SolveVelocitiesAndForces_NoiseFrequency;
 	
 	const float fRow	= dot(abs(SamplePosition), float3(1.0f / 8.0f, 2.0f / 8.0f, 5.0f / 8.0f));
 	const uint RowA		= uint(floor(abs(fRow))) & SolveVelocitiesAndForces_NoiseLUTRowMask;
@@ -128,8 +130,9 @@ void SolveVelocitiesAndForces_Simulate(inout FStatelessParticle Particle)
 
 	if ( SolveVelocitiesAndForces_NoiseEnabled  != 0 )
 	{
-		SolveVelocitiesAndForces_EvaluateNoise(ModuleData, Particle.Lifetime, Particle.NormalizedAge, Particle.Age, Particle.Position);
-		SolveVelocitiesAndForces_EvaluateNoise(ModuleData, Particle.Lifetime, Particle.PreviousNormalizedAge, Particle.PreviousAge, Particle.PreviousPosition);
+		const float3 NosieFieldBias = RandomScaleBiasFloat(10, SolveVelocitiesAndForces_NoiseFieldOffsetScale, SolveVelocitiesAndForces_NoiseFieldOffsetBias);
+		SolveVelocitiesAndForces_EvaluateNoise(ModuleData, Particle.Lifetime, Particle.NormalizedAge, Particle.Age, NosieFieldBias, Particle.Position);
+		SolveVelocitiesAndForces_EvaluateNoise(ModuleData, Particle.Lifetime, Particle.PreviousNormalizedAge, Particle.PreviousAge, NosieFieldBias, Particle.PreviousPosition);
 	}
 	
 	// Simulate
diff --git a/Engine/Plugins/FX/Niagara/Source/Niagara/Internal/Stateless/Modules/NiagaraStatelessModuleCommon.h b/Engine/Plugins/FX/Niagara/Source/Niagara/Internal/Stateless/Modules/NiagaraStatelessModuleCommon.h
--- a/Engine/Plugins/FX/Niagara/Source/Niagara/Internal/Stateless/Modules/NiagaraStatelessModuleCommon.h
+++ b/Engine/Plugins/FX/Niagara/Source/Niagara/Internal/Stateless/Modules/NiagaraStatelessModuleCommon.h
@@ -43,5 +43,6 @@ namespace NiagaraStateless
 		bool							bNoiseEnabled = false;
 		float							NoiseStrength = 0.0f;
 		float							NoiseFrequency = 0.0f;
+		FNiagaraStatelessRangeVector3	NoiseFieldOffset = FNiagaraStatelessRangeVector3(FVector3f::ZeroVector);
 	};
 }
diff --git a/Engine/Plugins/FX/Niagara/Source/Niagara/Internal/Stateless/Modules/NiagaraStatelessModule_CurlNoiseForce.h b/Engine/Plugins/FX/Niagara/Source/Niagara/Internal/Stateless/Modules/NiagaraStatelessModule_CurlNoiseForce.h
--- a/Engine/Plugins/FX/Niagara/Source/Niagara/Internal/Stateless/Modules/NiagaraStatelessModule_CurlNoiseForce.h
+++ b/Engine/Plugins/FX/Niagara/Source/Niagara/Internal/Stateless/Modules/NiagaraStatelessModule_CurlNoiseForce.h
@@ -21,6 +21,12 @@ public:
 	UPROPERTY(EditAnywhere, Category = "Parameters")
 	float NoiseFrequency = 50.0f;
 
+	UPROPERTY(EditAnywhere, Category = "Parameters", meta = (InlineEditConditionToggle))
+	uint32 bBiasNoiseFieldEnabled : 1 = false;
+
+	UPROPERTY(EditAnywhere, Category = "Parameters", meta = (EditCondition = "bBiasNoiseFieldEnabled"))
+	FNiagaraDistributionRangeVector3 BiasNoiseField = FNiagaraDistributionRangeVector3(FVector3f::ZeroVector);
+
 	virtual void BuildEmitterData(const FNiagaraStatelessEmitterDataBuildContext& BuildContext) const override;
 
 #if WITH_EDITOR
diff --git a/Engine/Plugins/FX/Niagara/Source/Niagara/Private/Stateless/Modules/NiagaraStatelessModule_CurlNoiseForce.cpp b/Engine/Plugins/FX/Niagara/Source/Niagara/Private/Stateless/Modules/NiagaraStatelessModule_CurlNoiseForce.cpp
--- a/Engine/Plugins/FX/Niagara/Source/Niagara/Private/Stateless/Modules/NiagaraStatelessModule_CurlNoiseForce.cpp
+++ b/Engine/Plugins/FX/Niagara/Source/Niagara/Private/Stateless/Modules/NiagaraStatelessModule_CurlNoiseForce.cpp
@@ -12,7 +12,8 @@ void UNiagaraStatelessModule_CurlNoiseForce::BuildEmitterData(const FNiagaraStat
 	}
 
 	NiagaraStateless::FPhysicsBuildData& PhysicsBuildData = BuildContext.GetTransientBuildData<NiagaraStateless::FPhysicsBuildData>();
-	PhysicsBuildData.bNoiseEnabled = true;
-	PhysicsBuildData.NoiseStrength = NoiseStrength;
-	PhysicsBuildData.NoiseFrequency = NoiseFrequency;
+	PhysicsBuildData.bNoiseEnabled		= true;
+	PhysicsBuildData.NoiseStrength		= NoiseStrength;
+	PhysicsBuildData.NoiseFrequency		= NoiseFrequency;
+	PhysicsBuildData.NoiseFieldOffset	= bBiasNoiseFieldEnabled ? BiasNoiseField.CalculateRange() : FNiagaraStatelessRangeVector3(FVector3f::ZeroVector);
 }
diff --git a/Engine/Plugins/FX/Niagara/Source/Niagara/Private/Stateless/Modules/NiagaraStatelessModule_SolveVelocitiesAndForces.cpp b/Engine/Plugins/FX/Niagara/Source/Niagara/Private/Stateless/Modules/NiagaraStatelessModule_SolveVelocitiesAndForces.cpp
--- a/Engine/Plugins/FX/Niagara/Source/Niagara/Private/Stateless/Modules/NiagaraStatelessModule_SolveVelocitiesAndForces.cpp
+++ b/Engine/Plugins/FX/Niagara/Source/Niagara/Private/Stateless/Modules/NiagaraStatelessModule_SolveVelocitiesAndForces.cpp
@@ -50,9 +50,9 @@ namespace NiagaraStatelessModuleSolveVelocitiesAndForcesPrivate
 		return Age * (Velocity + Wind) + 0.5f * Acceleration * Age * Age;
 	}
 
-	static void EvaluateNoise(TConstArrayView<FVector3f> NoiseLUT, const UNiagaraStatelessModule_SolveVelocitiesAndForces::FParameters* Parameters, float Lifetime, float NormalizedAge, FVector3f& Position)
+	void EvaluateNoise(TConstArrayView<FVector3f> NoiseLUT, const UNiagaraStatelessModule_SolveVelocitiesAndForces::FParameters* Parameters, float Lifetime, float NormalizedAge, const FVector3f& NoiseFieldOffset, FVector3f& Position)
 	{
-		const FVector3f SamplePosition = Position * Parameters->SolveVelocitiesAndForces_NoiseFrequency;
+		const FVector3f SamplePosition = (Position + NoiseFieldOffset) * Parameters->SolveVelocitiesAndForces_NoiseFrequency;
 
 		const float fRow = FVector3f::DotProduct(SamplePosition.GetAbs(), FVector3f(1.0f / 8.0f, 2.0f / 8.0f, 5.0f / 8.0f));
 		const uint32 RowA = uint32(FMath::FloorToInt(FMath::Abs(fRow))) & Parameters->SolveVelocitiesAndForces_NoiseLUTRowMask;
@@ -85,6 +85,7 @@ namespace NiagaraStatelessModuleSolveVelocitiesAndForcesPrivate
 		const float* PreviousAgeData			= ParticleSimulationContext.GetParticlePreviousAge();
 		const float* NormalizedAgeData			= ParticleSimulationContext.GetParticleNormalizedAge();
 		const float* PreviousNormalizedAgeData	= ParticleSimulationContext.GetParticlePreviousNormalizedAge();
+		const int32* UniqueIDData				= ParticleSimulationContext.GetParticleUniqueID();
 		const float InvDeltaTime				= ParticleSimulationContext.GetInvDeltaTime();
 
 		TConstArrayView<FVector3f> NoiseLUT = NiagaraStateless::FNoiseLUT::GetGlobalLUT().GetCpuData();
@@ -133,8 +134,9 @@ namespace NiagaraStatelessModuleSolveVelocitiesAndForcesPrivate
 
 			if (Parameters->SolveVelocitiesAndForces_NoiseEnabled != 0)
 			{
-				EvaluateNoise(NoiseLUT, Parameters, LifetimeData[i], NormalizedAgeData[i], Position);
-				EvaluateNoise(NoiseLUT, Parameters, LifetimeData[i], PreviousNormalizedAgeData[i], PreviousPosition);
+				const FVector3f NoiseFieldOffset = ParticleSimulationContext.RandomScaleBiasFloat(i, 10, Parameters->SolveVelocitiesAndForces_NoiseFieldOffsetScale, Parameters->SolveVelocitiesAndForces_NoiseFieldOffsetBias);
+				EvaluateNoise(NoiseLUT, Parameters, LifetimeData[i], NormalizedAgeData[i], NoiseFieldOffset, Position);
+				EvaluateNoise(NoiseLUT, Parameters, LifetimeData[i], PreviousNormalizedAgeData[i], NoiseFieldOffset, PreviousPosition);
 			}
 
 			Position += IntegratePosition(FMath::Max(AgeData[i], UE_KINDA_SMALL_NUMBER), Drag, InitialVelocity, Wind, Acceleration);
@@ -231,13 +233,14 @@ void UNiagaraStatelessModule_SolveVelocitiesAndForces::SetShaderParameters(const
 		const NiagaraStateless::FNoiseLUT& NoiseLUT = NiagaraStateless::FNoiseLUT::GetGlobalLUT();
 		checkf(FMath::IsPowerOfTwo(NoiseLUT.GetNumRows()), TEXT("Noise LUT must be PoT"));
 
-		Parameters->SolveVelocitiesAndForces_NoiseEnabled		= PhysicsData.bNoiseEnabled ? 1 : 0;
-		Parameters->SolveVelocitiesAndForces_NoiseStrength		= PhysicsData.NoiseStrength;
-		Parameters->SolveVelocitiesAndForces_NoiseFrequency		= PhysicsData.NoiseFrequency / GFrequencyNormalize;
-		Parameters->SolveVelocitiesAndForces_NoiseLUTNormalize	= (PhysicsData.NoiseFrequency / GFrequencyNormalize) * GFrequencyLUTNormalize;
-		Parameters->SolveVelocitiesAndForces_NoiseLUTRowMask	= NoiseLUT.GetNumRows() - 1;
-		Parameters->SolveVelocitiesAndForces_NoiseLUTRowWidth	= NoiseLUT.GetRowWidth();
-		Parameters->SolveVelocitiesAndForces_NoiseLUT			= NoiseLUT.GetGpuData();
+		Parameters->SolveVelocitiesAndForces_NoiseEnabled			= PhysicsData.bNoiseEnabled ? 1 : 0;
+		Parameters->SolveVelocitiesAndForces_NoiseStrength			= PhysicsData.NoiseStrength;
+		Parameters->SolveVelocitiesAndForces_NoiseFrequency			= PhysicsData.NoiseFrequency / GFrequencyNormalize;
+		SetShaderParameterContext.ConvertRangeToScaleBias(PhysicsData.NoiseFieldOffset, Parameters->SolveVelocitiesAndForces_NoiseFieldOffsetScale, Parameters->SolveVelocitiesAndForces_NoiseFieldOffsetBias);
+		Parameters->SolveVelocitiesAndForces_NoiseLUTNormalize		= (PhysicsData.NoiseFrequency / GFrequencyNormalize) * GFrequencyLUTNormalize;
+		Parameters->SolveVelocitiesAndForces_NoiseLUTRowMask		= NoiseLUT.GetNumRows() - 1;
+		Parameters->SolveVelocitiesAndForces_NoiseLUTRowWidth		= NoiseLUT.GetRowWidth();
+		Parameters->SolveVelocitiesAndForces_NoiseLUT				= NoiseLUT.GetGpuData();
 	}
 }
 
diff --git a/Engine/Plugins/FX/Niagara/Source/NiagaraEditor/Private/Widgets/SNiagaraDistributionCurveEditor.cpp b/Engine/Plugins/FX/Niagara/Source/NiagaraEditor/Private/Widgets/SNiagaraDistributionCurveEditor.cpp
--- a/Engine/Plugins/FX/Niagara/Source/NiagaraEditor/Private/Widgets/SNiagaraDistributionCurveEditor.cpp
+++ b/Engine/Plugins/FX/Niagara/Source/NiagaraEditor/Private/Widgets/SNiagaraDistributionCurveEditor.cpp
@@ -815,13 +815,7 @@ void SNiagaraDistributionCurveEditor::Construct(const FArguments& InArgs, TShare
 				+ SVerticalBox::Slot()
 				.Padding(0, 0, 0, 5)
 				[
-					SNew(SHorizontalBox)
-					+ SHorizontalBox::Slot()
-					.HAlign(HAlign_Fill)
-					.FillWidth(1.0f)
-					[
-						CurveEditorPanel.ToSharedRef()
-					]
+					CurveEditorPanel.ToSharedRef()
 				]
 				+ SVerticalBox::Slot()
 				.AutoHeight()
diff --git a/Engine/Plugins/FX/Niagara/Source/NiagaraShader/Internal/Stateless/NiagaraStatelessModuleShaderParameters.h b/Engine/Plugins/FX/Niagara/Source/NiagaraShader/Internal/Stateless/NiagaraStatelessModuleShaderParameters.h
--- a/Engine/Plugins/FX/Niagara/Source/NiagaraShader/Internal/Stateless/NiagaraStatelessModuleShaderParameters.h
+++ b/Engine/Plugins/FX/Niagara/Source/NiagaraShader/Internal/Stateless/NiagaraStatelessModuleShaderParameters.h
@@ -178,6 +178,8 @@ namespace NiagaraStateless
 		SHADER_PARAMETER(uint32,				SolveVelocitiesAndForces_NoiseEnabled)
 		SHADER_PARAMETER(float,					SolveVelocitiesAndForces_NoiseStrength)
 		SHADER_PARAMETER(float,					SolveVelocitiesAndForces_NoiseFrequency)
+		SHADER_PARAMETER(FVector3f,				SolveVelocitiesAndForces_NoiseFieldOffsetScale)
+		SHADER_PARAMETER(FVector3f,				SolveVelocitiesAndForces_NoiseFieldOffsetBias)
 		SHADER_PARAMETER(float,					SolveVelocitiesAndForces_NoiseLUTNormalize)
 		SHADER_PARAMETER(uint32,				SolveVelocitiesAndForces_NoiseLUTRowMask)
 		SHADER_PARAMETER(uint32,				SolveVelocitiesAndForces_NoiseLUTRowWidth)
