Niagara - Add a new feature for module inputs to make shaving values easier.  Converting an input to a shared input will copy the current value to a shared parameter, but will still allow the parameter to be edited inline.  Shared linked parameters are highlighted with a different color background and also have an additional drop-down button on hover for operations specific to shared linked inputs.

diff --git a/Engine/Plugins/FX/Niagara/Source/Niagara/Classes/NiagaraConstants.h b/Engine/Plugins/FX/Niagara/Source/Niagara/Classes/NiagaraConstants.h
--- a/Engine/Plugins/FX/Niagara/Source/Niagara/Classes/NiagaraConstants.h
+++ b/Engine/Plugins/FX/Niagara/Source/Niagara/Classes/NiagaraConstants.h
@@ -225,6 +225,7 @@ struct FNiagaraConstants
 	static NIAGARA_API const FName InitialNamespace;
 	static NIAGARA_API const FName PreviousNamespace;
 	static NIAGARA_API const FName OwnerNamespace;
+	static NIAGARA_API const FName SharedNamespace;
 
 	static NIAGARA_API const FName EngineOwnerScopeName;
 	static NIAGARA_API const FName EngineSystemScopeName;
diff --git a/Engine/Plugins/FX/Niagara/Source/Niagara/Private/NiagaraConstants.cpp b/Engine/Plugins/FX/Niagara/Source/Niagara/Private/NiagaraConstants.cpp
--- a/Engine/Plugins/FX/Niagara/Source/Niagara/Private/NiagaraConstants.cpp
+++ b/Engine/Plugins/FX/Niagara/Source/Niagara/Private/NiagaraConstants.cpp
@@ -44,6 +44,7 @@ const FName FNiagaraConstants::LocalNamespace(TEXT("Local"));
 const FName FNiagaraConstants::InitialNamespace(TEXT("Initial"));
 const FName FNiagaraConstants::PreviousNamespace(TEXT("Previous"));
 const FName FNiagaraConstants::OwnerNamespace(TEXT("Owner"));
+const FName FNiagaraConstants::SharedNamespace(TEXT("Shared"));
 
 const FName FNiagaraConstants::EngineOwnerScopeName(TEXT("EngineOwner"));
 const FName FNiagaraConstants::EngineSystemScopeName(TEXT("EngineSystem"));
diff --git a/Engine/Plugins/FX/Niagara/Source/NiagaraEditor/Private/NiagaraEditorSettings.cpp b/Engine/Plugins/FX/Niagara/Source/NiagaraEditor/Private/NiagaraEditorSettings.cpp
--- a/Engine/Plugins/FX/Niagara/Source/NiagaraEditor/Private/NiagaraEditorSettings.cpp
+++ b/Engine/Plugins/FX/Niagara/Source/NiagaraEditor/Private/NiagaraEditorSettings.cpp
@@ -297,6 +297,11 @@ void UNiagaraEditorSettings::SetupNamespaceMetadata()
 			.SetDescription(LOCTEXT("OwnerDescription", "A namespace modifier which specifies that an engine provided parameter comes from the owner, or component."))
 			.SetBackgroundColor(FLinearColor(FColor(170, 170, 170)))
 			.SetForegroundStyle("NiagaraEditor.ParameterName.NamespaceTextDark"),
+		FNiagaraNamespaceMetadata({FNiagaraConstants::SharedNamespace})
+			.SetDisplayName(LOCTEXT("SharedDisplayName", "Shared"))
+			.SetDescription(LOCTEXT("SharedDescription", "A namespace modifier which specifies that a parameter is being used to 'share' a value across multiple inputs and the UI should display it accordingly."))
+			.SetBackgroundColor(FLinearColor(FColor(170, 170, 170)))
+			.SetForegroundStyle("NiagaraEditor.ParameterName.NamespaceTextDark"),
 	};
 }
 
diff --git a/Engine/Plugins/FX/Niagara/Source/NiagaraEditor/Private/NiagaraEditorStyle.cpp b/Engine/Plugins/FX/Niagara/Source/NiagaraEditor/Private/NiagaraEditorStyle.cpp
--- a/Engine/Plugins/FX/Niagara/Source/NiagaraEditor/Private/NiagaraEditorStyle.cpp
+++ b/Engine/Plugins/FX/Niagara/Source/NiagaraEditor/Private/NiagaraEditorStyle.cpp
@@ -284,6 +284,7 @@ void FNiagaraEditorStyle::InitParameters()
 	
 	Set("NiagaraEditor.StaticInputIcon", new IMAGE_BRUSH_SVG("Icons/StaticInput", Icon12x12));
 	Set("NiagaraEditor.StaticIcon", new IMAGE_BRUSH("Icons/staticpill_16x", Icon16x16));
+	Set("NiagaraEditor.LinkedSharedIcon", new IMAGE_BRUSH_SVG("Icons/LinkedShared", Icon16x16));
 	Set("NiagaraEditor.Pins.StaticConnected", new IMAGE_BRUSH("Icons/StaticPin_Connected", Icon16x16));
 	Set("NiagaraEditor.Pins.StaticDisconnected", new IMAGE_BRUSH("Icons/StaticPin_Disconnected", Icon16x16));
 	Set("NiagaraEditor.Pins.StaticConnectedHovered", new IMAGE_BRUSH("Icons/StaticPin_Connected", Icon16x16, FLinearColor(0.8f, 0.8f, 0.8f)));
diff --git a/Engine/Plugins/FX/Niagara/Source/NiagaraEditor/Private/NiagaraValidationRules.cpp b/Engine/Plugins/FX/Niagara/Source/NiagaraEditor/Private/NiagaraValidationRules.cpp
--- a/Engine/Plugins/FX/Niagara/Source/NiagaraEditor/Private/NiagaraValidationRules.cpp
+++ b/Engine/Plugins/FX/Niagara/Source/NiagaraEditor/Private/NiagaraValidationRules.cpp
@@ -1202,7 +1202,8 @@ void UNiagaraValidationRule_LWC::CheckValidity(const FNiagaraValidationContext&
 				}
 
 				// check if the linked input variable is a vector
-				if (Input->GetValueMode() == UNiagaraStackFunctionInput::EValueMode::Linked && Settings->bEnforceStrictStackTypes)
+				if ((Input->GetValueMode() == UNiagaraStackFunctionInput::EValueMode::Linked ||
+					Input->GetValueMode() == UNiagaraStackFunctionInput::EValueMode::LinkedShared) && Settings->bEnforceStrictStackTypes)
 				{
 					FNiagaraVariable VectorVar(FNiagaraTypeDefinition::GetVec3Def(), Input->GetLinkedParameterValue().GetName());
 					const UNiagaraGraph* NiagaraGraph = Input->GetInputFunctionCallNode().GetNiagaraGraph();
diff --git a/Engine/Plugins/FX/Niagara/Source/NiagaraEditor/Private/ViewModels/NiagaraSystemViewModel.cpp b/Engine/Plugins/FX/Niagara/Source/NiagaraEditor/Private/ViewModels/NiagaraSystemViewModel.cpp
--- a/Engine/Plugins/FX/Niagara/Source/NiagaraEditor/Private/ViewModels/NiagaraSystemViewModel.cpp
+++ b/Engine/Plugins/FX/Niagara/Source/NiagaraEditor/Private/ViewModels/NiagaraSystemViewModel.cpp
@@ -1226,7 +1226,7 @@ void FNiagaraSystemViewModel::SetToolkitCommands(const TSharedRef<FUICommandList
 	ToolkitCommands = InToolkitCommands;
 }
 
-const TArray<FNiagaraStackModuleData>& FNiagaraSystemViewModel::GetStackModuleDataByModuleEntry(UNiagaraStackEntry* ModuleEntry)
+const TArray<FNiagaraStackModuleData>& FNiagaraSystemViewModel::GetStackModuleDataByModuleEntry(const UNiagaraStackEntry* ModuleEntry)
 {
 	FGuid EmitterHandleId = FGuid();
 	FVersionedNiagaraEmitter Emitter;
diff --git a/Engine/Plugins/FX/Niagara/Source/NiagaraEditor/Private/ViewModels/Stack/NiagaraStackFunctionInput.cpp b/Engine/Plugins/FX/Niagara/Source/NiagaraEditor/Private/ViewModels/Stack/NiagaraStackFunctionInput.cpp
--- a/Engine/Plugins/FX/Niagara/Source/NiagaraEditor/Private/ViewModels/Stack/NiagaraStackFunctionInput.cpp
+++ b/Engine/Plugins/FX/Niagara/Source/NiagaraEditor/Private/ViewModels/Stack/NiagaraStackFunctionInput.cpp
@@ -272,9 +272,16 @@ bool UNiagaraStackFunctionInput::GetShouldShowInStack() const
 	// Using the FilterOnlyModified child filter, we let inputs pass that are at default value but have modified children inputs
 	// If we removed them via filter, the modified children inputs also wouldn't be displayed.
 	// Instead, we remove these inputs from the display
-	if(IsFinalized() == false && GetStackEditorData().GetShowOnlyModified())
+	if (IsFinalized() == false)
 	{
-		if(CanReset() == false)
+		if (GetStackEditorData().GetShowOnlyModified())
+		{
+			if (CanReset() == false)
+			{
+				return false;
+			}
+		}
+		if (bIsSharedLinkedValueInput)
 		{
 			return false;
 		}
@@ -694,6 +701,7 @@ FText UNiagaraStackFunctionInput::GetCollapsedStateText() const
 			}
 			break;
 		case EValueMode::Linked:
+		case EValueMode::LinkedShared:
 			CollapsedTextCache = FText::FromString(InputValues.LinkedParameter.GetName().ToString());
 			break;
 		case EValueMode::Expression:
@@ -777,6 +785,7 @@ FText UNiagaraStackFunctionInput::GetValueToolTip() const
 			ValueToolTipCache = LOCTEXT("InvalidOverrideToolTip", "The script is in an invalid and unrecoverable state for this\ninput.  Resetting to default may fix this issue.");
 			break;
 		case EValueMode::Linked:
+		case EValueMode::LinkedShared:
 		{
 			FNiagaraVariable InputVariable = InputValues.LinkedParameter;
 			if (FNiagaraConstants::IsNiagaraConstant(InputVariable))
@@ -1076,6 +1085,15 @@ void UNiagaraStackFunctionInput::RefreshChildrenInternal(const TArray<UNiagaraSt
 		NewChildren.Add(ValueObjectEntry);
 	}
 
+	if (InputValues.Mode == EValueMode::LinkedShared && InputValues.SharedSourceAssignmentNode.IsValid())
+	{
+		NewChildren.Add(GetSharedLinkedValueInput());
+	}
+	else
+	{
+		SharedLinkedInput = nullptr;
+	}
+
 	DisplayNameOverride.Reset();
 
 	if (InputMetaData)
@@ -1347,7 +1365,7 @@ FString UNiagaraStackFunctionInput::ResolveDisplayNameArgument(const FString& In
 
 void UNiagaraStackFunctionInput::RefreshValues()
 {
-	if (ensureMsgf(IsStaticParameter() || InputParameterHandle.IsModuleHandle(), TEXT("Function inputs can only be generated for module paramters.")) == false)
+	if (ensureMsgf(IsStaticParameter() || InputParameterHandle.IsModuleHandle(), TEXT("Function inputs can only be generated for module parameters.")) == false)
 	{
 		return;
 	}
@@ -3065,6 +3083,7 @@ const UNiagaraClipboardFunctionInput* UNiagaraStackFunctionInput::ToClipboardFun
 		break;
 	}
 	case EValueMode::Linked:
+	case EValueMode::LinkedShared:
 		ClipboardInput = UNiagaraClipboardFunctionInput::CreateLinkedValue(InOuter, InputName, InputType, bEditConditionValue, InputValues.LinkedParameter);
 		break;
 	case EValueMode::Data:
@@ -3682,6 +3701,78 @@ void UNiagaraStackFunctionInput::UpdateValuesFromScriptDefaults(FInputValues& In
 	}
 }
 
+namespace FNiagaraStackFunctionInputLinkedSharedUtilities
+{
+	bool IsSharedParameterName(FName InParameterName)
+	{
+		FNiagaraParameterHandle ParameterHandle = FNiagaraParameterHandle(InParameterName);
+		return ParameterHandle.GetHandleParts().Num() > 2 && ParameterHandle.GetHandleParts()[1] == FNiagaraConstants::SharedNamespace;
+	}
+
+	FName MakeSharedParameterName(TOptional<ENiagaraScriptUsage> TargetScriptUsage, const FString& FunctionName, FName InputName)
+	{
+		TArray<FString> NameParts;
+		if (TargetScriptUsage == ENiagaraScriptUsage::SystemSpawnScript || TargetScriptUsage == ENiagaraScriptUsage::SystemUpdateScript)
+		{
+			NameParts.Add(FNiagaraConstants::SystemNamespaceString);
+		}
+		else if (TargetScriptUsage == ENiagaraScriptUsage::EmitterSpawnScript || TargetScriptUsage == ENiagaraScriptUsage::EmitterUpdateScript)
+		{
+			NameParts.Add(FNiagaraConstants::EmitterNamespaceString);
+		}
+		else if (TargetScriptUsage == ENiagaraScriptUsage::ParticleSpawnScript || TargetScriptUsage == ENiagaraScriptUsage::ParticleUpdateScript)
+		{
+			NameParts.Add(FNiagaraConstants::ParticleAttributeNamespaceString);
+		}
+		else
+		{
+			NameParts.Add(FNiagaraConstants::TransientNamespace.ToString());
+		}
+
+		NameParts.Add(FNiagaraConstants::SharedNamespace.ToString());
+		NameParts.Add(FString::Printf(TEXT("%s_%s"), *FunctionName, *InputName.ToString()));
+
+		return *FString::Join(NameParts, TEXT("."));
+	}
+
+	void GatherSharedParametersAndSharedParameterAssignmentNodes(UNiagaraScript& TargetScript, TSet<FNiagaraVariableBase>& OutSharedParameters, TArray<UNiagaraNodeAssignment*>& OutSharedParameterAssignmentNodes)
+	{
+		UNiagaraScriptSource* TargetScriptSource = CastChecked<UNiagaraScriptSource>(TargetScript.GetLatestSource());
+		UNiagaraGraph* TargetGraph = TargetScriptSource->NodeGraph;
+		UNiagaraNodeOutput* TargetOutputNode = TargetGraph->FindEquivalentOutputNode(TargetScript.GetUsage());
+		if (TargetOutputNode == nullptr)
+		{
+			return;
+		}
+
+		TArray<FNiagaraStackGraphUtilities::FStackNodeGroup> StackGroups;
+		FNiagaraStackGraphUtilities::GetStackNodeGroups(*TargetOutputNode, StackGroups);
+		for (const FNiagaraStackGraphUtilities::FStackNodeGroup& StackGroup : StackGroups)
+		{
+			UNiagaraNodeAssignment* AssignmentNode = Cast<UNiagaraNodeAssignment>(StackGroup.EndNode);
+			if (AssignmentNode == nullptr)
+			{
+				continue;
+			}
+
+			bool bSharedParameterFound = false;
+			for (const FNiagaraVariable& AssignmentTarget : AssignmentNode->GetAssignmentTargets())
+			{
+				if (FNiagaraStackFunctionInputLinkedSharedUtilities::IsSharedParameterName(AssignmentTarget.GetName()))
+				{
+					OutSharedParameters.Add(AssignmentTarget);
+					bSharedParameterFound = true;
+				}
+			}
+
+			if (bSharedParameterFound)
+			{
+				OutSharedParameterAssignmentNodes.Add(AssignmentNode);
+			}
+		}
+	}
+}
+
 void UNiagaraStackFunctionInput::UpdateValuesFromOverridePin(const FInputValues& OldInputValues, FInputValues& NewInputValues, UEdGraphPin& InOverridePin) const
 {
 	NewInputValues.Mode = EValueMode::InvalidOverride;
@@ -3731,8 +3822,34 @@ void UNiagaraStackFunctionInput::UpdateValuesFromOverridePin(const FInputValues&
 		else if (LinkedNode->IsA<UNiagaraNodeParameterMapGet>())
 		{
 			// Parameter map get nodes handle linked values.
-			NewInputValues.Mode = EValueMode::Linked;
 			NewInputValues.LinkedParameter = UEdGraphSchema_Niagara::PinToNiagaraVariable(InOverridePin.LinkedTo[0]);
+			bool bCanBeLinkedShared = FNiagaraStackFunctionInputLinkedSharedUtilities::IsSharedParameterName(NewInputValues.LinkedParameter.GetName());
+			if (bCanBeLinkedShared)
+			{
+				// Parameters with the shared namespace modifier can be displayed inline if they're setup correctly, so check that.
+				const TArray<FNiagaraStackModuleData>& StackModuleData = GetSystemViewModel()->GetStackModuleDataByModuleEntry(this);
+				int32 OwningModuleIndex = StackModuleData.IndexOfByPredicate([this](const FNiagaraStackModuleData& StackModuleDataItem)
+					{ return StackModuleDataItem.WeakModuleNode.IsValid() && StackModuleDataItem.WeakModuleNode.Get() == OwningModuleNode.Get(); });
+				UNiagaraNodeAssignment* SharedSourceAssignmentNode = nullptr;
+				for (int32 StackModuleDataIndex = OwningModuleIndex - 1; StackModuleDataIndex >= 0 && SharedSourceAssignmentNode == nullptr; StackModuleDataIndex--)
+				{
+					UNiagaraNodeAssignment* AssignmentNode = Cast<UNiagaraNodeAssignment>(StackModuleData[StackModuleDataIndex].WeakModuleNode.Get());
+					if (AssignmentNode != nullptr && AssignmentNode->GetAssignmentTargets().Contains(NewInputValues.LinkedParameter))
+					{
+						SharedSourceAssignmentNode = AssignmentNode;
+					}
+				}
+				if (SharedSourceAssignmentNode != nullptr)
+				{
+					NewInputValues.Mode = EValueMode::LinkedShared;
+					NewInputValues.SharedSourceAssignmentNode = SharedSourceAssignmentNode;
+				}
+			}
+
+			if (NewInputValues.Mode != EValueMode::LinkedShared)
+			{
+				NewInputValues.Mode = EValueMode::Linked;
+			}
 		}
 		else if (LinkedNode->IsA<UNiagaraNodeCustomHlsl>())
 		{
@@ -3809,4 +3926,179 @@ bool UNiagaraStackFunctionInput::SupportsCustomExpressions() const
 	return NiagaraEditorSettings->IsVisibleClass(UNiagaraNodeCustomHlsl::StaticClass());
 }
 
+void UNiagaraStackFunctionInput::GetAvailableScriptUsagesForSharedParameters(TArray<ENiagaraScriptUsage>& OutUsages) const
+{
+	if (bIsSharedLinkedValueInput)
+	{
+		// Inputs which are already shared linked values can not be linked to additional shared parameters.
+		return;
+	}
+
+	static const TArray<ENiagaraScriptUsage> OrderedUsages =
+	{
+		ENiagaraScriptUsage::SystemSpawnScript,
+		ENiagaraScriptUsage::SystemUpdateScript,
+		ENiagaraScriptUsage::EmitterSpawnScript,
+		ENiagaraScriptUsage::EmitterUpdateScript,
+		ENiagaraScriptUsage::ParticleSpawnScript,
+		ENiagaraScriptUsage::ParticleUpdateScript,
+	};
+
+	ENiagaraScriptUsage SourceUsage = SourceScript->GetUsage();
+	int32 SourceUsageIndex = OrderedUsages.IndexOfByPredicate([SourceUsage](const ENiagaraScriptUsage& Usage)
+	{
+		if (UNiagaraScript::IsEquivalentUsage(Usage, SourceUsage))
+		{
+			return true;
+		}
+		return Usage == ENiagaraScriptUsage::ParticleUpdateScript &&
+			(SourceUsage == ENiagaraScriptUsage::ParticleEventScript || SourceUsage == ENiagaraScriptUsage::ParticleSimulationStageScript);
+	});
+
+	for (int32 UsageIndex = 0; UsageIndex <= SourceUsageIndex; UsageIndex++)
+	{
+		OutUsages.Add(OrderedUsages[UsageIndex]);
+	}
+}
+
+UNiagaraStackFunctionInput* UNiagaraStackFunctionInput::GetSharedLinkedValueInput()
+{
+	if (InputValues.Mode == EValueMode::LinkedShared)
+	{
+		FNiagaraParameterHandle AssignmentInputHandle = FNiagaraParameterHandle(FNiagaraConstants::ModuleNamespace, InputValues.LinkedParameter.GetName());
+		if (SharedLinkedInput == nullptr ||
+			&SharedLinkedInput->GetInputFunctionCallNode() != InputValues.SharedSourceAssignmentNode.Get() ||
+			SharedLinkedInput->GetInputParameterHandle() != AssignmentInputHandle)
+		{
+			SharedLinkedInput = NewObject<UNiagaraStackFunctionInput>(this);
+			SharedLinkedInput->Initialize(CreateDefaultChildRequiredData(),
+				*InputValues.SharedSourceAssignmentNode.Get(),
+				*InputValues.SharedSourceAssignmentNode.Get(),
+				AssignmentInputHandle.GetParameterHandleString(),
+				InputValues.LinkedParameter.GetType(),
+				EStackParameterBehavior::Dynamic,
+				GetOwnerStackItemEditorDataKey());
+			SharedLinkedInput->SetIsSharedLinkedValueInput(true);
+		}
+		return SharedLinkedInput;
+	}
+	return nullptr;
+}
+
+void UNiagaraStackFunctionInput::MoveValueToLinkedSharedParameter(TOptional<ENiagaraScriptUsage> LinkedSharedParameterTargetUsage)
+{
+	// Gather the existing shared parameters and shared parameter assignment nodes.
+	UNiagaraScript* TargetScript = nullptr;
+	UNiagaraScript* SameNamespaceScript = nullptr;
+
+	if (LinkedSharedParameterTargetUsage.IsSet())
+	{
+		switch (LinkedSharedParameterTargetUsage.GetValue())
+		{
+		case ENiagaraScriptUsage::SystemSpawnScript:
+			TargetScript = GetSystemViewModel()->GetSystem().GetSystemSpawnScript();
+			SameNamespaceScript = GetSystemViewModel()->GetSystem().GetSystemUpdateScript();
+			break;
+		case ENiagaraScriptUsage::SystemUpdateScript:
+			TargetScript = GetSystemViewModel()->GetSystem().GetSystemUpdateScript();
+			SameNamespaceScript = GetSystemViewModel()->GetSystem().GetSystemSpawnScript();
+			break;
+		case ENiagaraScriptUsage::EmitterSpawnScript:
+			TargetScript = GetEmitterViewModel()->GetEmitter().GetEmitterData()->EmitterSpawnScriptProps.Script;
+			SameNamespaceScript = GetEmitterViewModel()->GetEmitter().GetEmitterData()->EmitterUpdateScriptProps.Script;
+			break;
+		case ENiagaraScriptUsage::EmitterUpdateScript:
+			TargetScript = GetEmitterViewModel()->GetEmitter().GetEmitterData()->EmitterUpdateScriptProps.Script;
+			SameNamespaceScript = GetEmitterViewModel()->GetEmitter().GetEmitterData()->EmitterSpawnScriptProps.Script;
+			break;
+		case ENiagaraScriptUsage::ParticleSpawnScript:
+		case ENiagaraScriptUsage::ParticleSpawnScriptInterpolated:
+			TargetScript = GetEmitterViewModel()->GetEmitter().GetEmitterData()->SpawnScriptProps.Script;
+			SameNamespaceScript = GetEmitterViewModel()->GetEmitter().GetEmitterData()->UpdateScriptProps.Script;
+			break;
+		case ENiagaraScriptUsage::ParticleUpdateScript:
+			TargetScript = GetEmitterViewModel()->GetEmitter().GetEmitterData()->UpdateScriptProps.Script;
+			SameNamespaceScript = GetEmitterViewModel()->GetEmitter().GetEmitterData()->SpawnScriptProps.Script;
+			break;
+		}
+	}
+	else
+	{
+		TargetScript = SourceScript.Get();
+	}
+
+	if (TargetScript == nullptr)
+	{
+		return;
+	}
+
+	TSet<FNiagaraVariableBase> SharedParameters;
+	TArray<UNiagaraNodeAssignment*> TargetAssignmentNodes;
+	FNiagaraStackFunctionInputLinkedSharedUtilities::GatherSharedParametersAndSharedParameterAssignmentNodes(*TargetScript, SharedParameters, TargetAssignmentNodes);
+	if (SameNamespaceScript != nullptr)
+	{
+		TArray<UNiagaraNodeAssignment*> SameNamespaceAssignmentNodes;
+		FNiagaraStackFunctionInputLinkedSharedUtilities::GatherSharedParametersAndSharedParameterAssignmentNodes(*SameNamespaceScript, SharedParameters, SameNamespaceAssignmentNodes);
+	}
+
+	FName ParameterName = FNiagaraStackFunctionInputLinkedSharedUtilities::MakeSharedParameterName(LinkedSharedParameterTargetUsage, OwningFunctionCallNode->GetFunctionName(), InputParameterHandle.GetName());
+	TSet<FName> SharedParameterNames;
+	for (const FNiagaraVariableBase& SharedParameter : SharedParameters)
+	{
+		SharedParameterNames.Add(SharedParameter.GetName());
+	}
+	FName UniqueParameterName = FNiagaraUtilities::GetUniqueName(ParameterName, SharedParameterNames);
+	FNiagaraVariableBase SharedParameter = FNiagaraVariableBase(InputType, UniqueParameterName);
+
+	// Copy the current value of the input before making any edits so it can be pasted to the shared input.
+	UNiagaraClipboardContent* ClipboardContent = UNiagaraClipboardContent::Create();
+	Copy(ClipboardContent);
+
+	// Add the linked shared input to an existing or new assignment node.
+	FScopedTransaction ScopedTransaction(LOCTEXT("MoveToSharedTransaction", "Move input value to shared parameter."));
+	UNiagaraNodeAssignment* TargetAssignmentNode;
+	if (TargetAssignmentNodes.Num() > 0)
+	{
+		TargetAssignmentNode = TargetAssignmentNodes[0];
+		TargetAssignmentNode->AddAssignmentTarget(SharedParameter);
+		TargetAssignmentNode->RefreshFromExternalChanges();
+	}
+	else
+	{
+		UNiagaraScriptSource* TargetScriptSource = CastChecked<UNiagaraScriptSource>(TargetScript->GetLatestSource());
+		UNiagaraGraph* TargetGraph = TargetScriptSource->NodeGraph;
+		UNiagaraNodeOutput* TargetOutputNode = TargetGraph->FindEquivalentOutputNode(TargetScript->GetUsage(), TargetScript->GetUsageId());
+		if (TargetOutputNode == nullptr)
+		{
+			return;
+		}
+		TargetAssignmentNode = FNiagaraStackGraphUtilities::AddParameterModuleToStack({ FNiagaraVariable(InputType, UniqueParameterName) }, *TargetOutputNode, 0, { FString() });
+	}
+	SetLinkedParameterValue(SharedParameter);
+
+	// Refresh children to make the shared linked input available and then paste the copied value.
+	RefreshChildren();
+	UNiagaraStackFunctionInput* SharedLinkedValueInput = GetSharedLinkedValueInput();
+	if (SharedLinkedValueInput != nullptr)
+	{
+		FText UnusedMessage;
+		SharedLinkedValueInput->Paste(ClipboardContent, UnusedMessage);
+	}
+}
+
+void UNiagaraStackFunctionInput::UnlinkFromSharedParameter()
+{
+	UNiagaraStackFunctionInput* SharedLinkedValueInput = GetSharedLinkedValueInput();
+	if (SharedLinkedValueInput != nullptr)
+	{
+		UNiagaraClipboardContent* ClipboardContent = UNiagaraClipboardContent::Create();
+		SharedLinkedValueInput->Copy(ClipboardContent);
+
+		FScopedTransaction ScopedTransaction(LOCTEXT("UnlinkFromSharedParameter", "Unlink from shared parameter."));
+		Reset();
+		FText UnusedMessage;
+		Paste(ClipboardContent, UnusedMessage);
+	}
+}
+
 #undef LOCTEXT_NAMESPACE
diff --git a/Engine/Plugins/FX/Niagara/Source/NiagaraEditor/Public/ViewModels/NiagaraSystemViewModel.h b/Engine/Plugins/FX/Niagara/Source/NiagaraEditor/Public/ViewModels/NiagaraSystemViewModel.h
--- a/Engine/Plugins/FX/Niagara/Source/NiagaraEditor/Public/ViewModels/NiagaraSystemViewModel.h
+++ b/Engine/Plugins/FX/Niagara/Source/NiagaraEditor/Public/ViewModels/NiagaraSystemViewModel.h
@@ -418,7 +418,7 @@ public:
 	void SetToolkitCommands(const TSharedRef<FUICommandList>& InToolkitCommands);
 
 	/** Gets the stack module data for the provided module stack entry, for use in determining dependencies. */
-	const TArray<FNiagaraStackModuleData>& GetStackModuleDataByModuleEntry(UNiagaraStackEntry* ModuleEntry);
+	const TArray<FNiagaraStackModuleData>& GetStackModuleDataByModuleEntry(const UNiagaraStackEntry* ModuleEntry);
 
 	/** Gets the stack module data for the provided emitter handled id, for use in determining dependencies. */
 	const TArray<FNiagaraStackModuleData>& GetStackModuleDataByEmitterHandleId(FGuid EmitterHandleId);
diff --git a/Engine/Plugins/FX/Niagara/Source/NiagaraEditor/Public/ViewModels/Stack/NiagaraStackFunctionInput.h b/Engine/Plugins/FX/Niagara/Source/NiagaraEditor/Public/ViewModels/Stack/NiagaraStackFunctionInput.h
--- a/Engine/Plugins/FX/Niagara/Source/NiagaraEditor/Public/ViewModels/Stack/NiagaraStackFunctionInput.h
+++ b/Engine/Plugins/FX/Niagara/Source/NiagaraEditor/Public/ViewModels/Stack/NiagaraStackFunctionInput.h
@@ -43,6 +43,8 @@ public:
 		Local,
 		/** The value is linked to a parameter defined outside of this function. */
 		Linked,
+		/** The value is linked to a shared parameter defined in the top level script. */
+		LinkedShared,
 		/** The value is provided by a secondary dynamic input function. */
 		Dynamic,
 		/** The value is provided by a data interface object. */
@@ -350,6 +352,22 @@ public:
 
 	NIAGARAEDITOR_API bool SupportsCustomExpressions() const;
 
+	/** Gets available target script usages which can be used to create shared linked parameters. */
+	NIAGARAEDITOR_API void GetAvailableScriptUsagesForSharedParameters(TArray<ENiagaraScriptUsage>& OutUsages) const;
+
+	/** Gets a stack function input to display and edit a linked value, when this input is linked to a shared linked parameter, otherwise returns null */
+	NIAGARAEDITOR_API UNiagaraStackFunctionInput* GetSharedLinkedValueInput();
+
+	/** Gets whether or not this input was created by another input to display and edit a shared linked value. */
+	NIAGARAEDITOR_API bool GetIsSharedLinkedValueInput() const { return bIsSharedLinkedValueInput; }
+
+	/** Moves the current value of this input, to a shared parameter and then changes the value of this input to link to that new parameter.  If a target usage is provided
+		the shared input parameter will be added to that script, otherwise the shared parameter will be added to a transient parameter in the current script. */
+	NIAGARAEDITOR_API void MoveValueToLinkedSharedParameter(TOptional<ENiagaraScriptUsage> LinkedSharedParameterTargetUsage);
+
+	/** Breaks the link to a shared parameter and copies the value of the shared parameter this input. */
+	NIAGARAEDITOR_API void UnlinkFromSharedParameter();
+
 protected:
 	//~ UNiagaraStackEntry interface
 	NIAGARAEDITOR_API virtual void FinalizeInternal() override;
@@ -360,11 +378,13 @@ protected:
 	NIAGARAEDITOR_API FString ResolveDisplayNameArgument(const FString& InArg) const;
 	NIAGARAEDITOR_API FStackIssueFixDelegate GetUpgradeDynamicInputVersionFix();
 
+	/** Sets whether or not this input was created by another input to display and edit a shared linked value. */
+	NIAGARAEDITOR_API void SetIsSharedLinkedValueInput(bool bInIsSharedLinkedValueInput) { bIsSharedLinkedValueInput = bInIsSharedLinkedValueInput; }
+
 private:
 	struct FInputValues
 	{
 		FInputValues()
-			: Mode(EValueMode::None)
 		{
 		}
 
@@ -373,9 +393,10 @@ private:
 			return Mode != EValueMode::None && Mode != EValueMode::InvalidOverride && Mode != EValueMode::UnsupportedDefault;
 		}
 
-		EValueMode Mode;
+		EValueMode Mode = EValueMode::None;
 		TSharedPtr<FStructOnScope> LocalStruct;
 		FNiagaraVariableBase LinkedParameter;
+		TWeakObjectPtr<UNiagaraNodeAssignment> SharedSourceAssignmentNode;
 		TWeakObjectPtr<UNiagaraNodeFunctionCall> DynamicNode;
 		TWeakObjectPtr<UNiagaraNodeCustomHlsl> ExpressionNode;
 		TWeakObjectPtr<UNiagaraDataInterface> DataObject;
@@ -585,6 +606,12 @@ private:
 	// If true then this stack entry is the semantic child of another stack entry
 	bool bIsSemanticChild = false;
 
+	/** A stack function input to display and edit a linked value, when this input is linked to a shared linked parameter. */
+	UPROPERTY()
+	TObjectPtr<UNiagaraStackFunctionInput> SharedLinkedInput;
+	/** Whether or not this input was created by another input to display and edit a shared linked value. */
+	bool bIsSharedLinkedValueInput = false;
+
 	// The value of the change id for the graph which owns the function call that owns this input as of the last refresh.
 	TOptional<FGuid> LastOwningGraphChangeId;
 
diff --git a/Engine/Plugins/FX/Niagara/Source/NiagaraEditorWidgets/Private/NiagaraEditorWidgetsStyle.cpp b/Engine/Plugins/FX/Niagara/Source/NiagaraEditorWidgets/Private/NiagaraEditorWidgetsStyle.cpp
--- a/Engine/Plugins/FX/Niagara/Source/NiagaraEditorWidgets/Private/NiagaraEditorWidgetsStyle.cpp
+++ b/Engine/Plugins/FX/Niagara/Source/NiagaraEditorWidgets/Private/NiagaraEditorWidgetsStyle.cpp
@@ -376,6 +376,13 @@ void FNiagaraEditorWidgetsStyle::InitButtonStyles()
 	CompactComboButton.DownArrowPadding = FMargin(1.0f, 0.0f);
 	Set("NiagaraEditor.Stack.CompactComboButton", CompactComboButton);
 
+	FComboButtonStyle CompactComboButtonWithBorder = FAppStyle::Get().GetWidgetStyle<FComboButtonStyle>("ComboButton");
+	CompactComboButtonWithBorder.DownArrowPadding = FMargin(1.0f, 0.0f);
+	CompactComboButtonWithBorder.ContentPadding = FMargin(0.0f, 2.0f);
+	CompactComboButtonWithBorder.ButtonStyle.NormalPadding = FMargin(6, 1, 2, 2);
+	CompactComboButtonWithBorder.ButtonStyle.PressedPadding = FMargin(6, 2, 2, 1);
+	Set("NiagaraEditor.Stack.CompactComboButtonWithBorder", CompactComboButtonWithBorder);
+
 	FButtonStyle InlineParameterButtonStyle = FButtonStyle()
 		.SetNormal(FSlateRoundedBoxBrush(FStyleColors::AccentGray, 4.0f, FStyleColors::Input, 1.f))
 		.SetHovered(FSlateRoundedBoxBrush(FStyleColors::White, 4.0f, FStyleColors::Input, 1.f))
diff --git a/Engine/Plugins/FX/Niagara/Source/NiagaraEditorWidgets/Private/NiagaraEditorWidgetsUtilities.cpp b/Engine/Plugins/FX/Niagara/Source/NiagaraEditorWidgets/Private/NiagaraEditorWidgetsUtilities.cpp
--- a/Engine/Plugins/FX/Niagara/Source/NiagaraEditorWidgets/Private/NiagaraEditorWidgetsUtilities.cpp
+++ b/Engine/Plugins/FX/Niagara/Source/NiagaraEditorWidgets/Private/NiagaraEditorWidgetsUtilities.cpp
@@ -174,6 +174,7 @@ FText FNiagaraStackEditorWidgetsUtilities::GetIconTextForInputMode(UNiagaraStack
 	switch (InputValueMode)
 	{
 	case UNiagaraStackFunctionInput::EValueMode::Linked:
+	case UNiagaraStackFunctionInput::EValueMode::LinkedShared:
 		return FEditorFontGlyphs::Link;
 	case UNiagaraStackFunctionInput::EValueMode::Data:
 		return FEditorFontGlyphs::Database;
@@ -201,6 +202,8 @@ FText FNiagaraStackEditorWidgetsUtilities::GetIconToolTipForInputMode(UNiagaraSt
 	{
 	case UNiagaraStackFunctionInput::EValueMode::Linked:
 		return LOCTEXT("LinkInputIconToolTip", "Linked Value");
+	case UNiagaraStackFunctionInput::EValueMode::LinkedShared:
+		return LOCTEXT("LinkSharedInputIconToolTip", "Linked Shared Value");
 	case UNiagaraStackFunctionInput::EValueMode::Data:
 		return LOCTEXT("DataInterfaceInputIconToolTip", "Data Value");
 	case UNiagaraStackFunctionInput::EValueMode::ObjectAsset:
@@ -225,6 +228,7 @@ FName FNiagaraStackEditorWidgetsUtilities::GetIconColorNameForInputMode(UNiagara
 	switch (InputValueMode)
 	{
 	case UNiagaraStackFunctionInput::EValueMode::Linked:
+	case UNiagaraStackFunctionInput::EValueMode::LinkedShared:
 		return "NiagaraEditor.Stack.InputValueIconColor.Linked";
 	case UNiagaraStackFunctionInput::EValueMode::Data:
 		return "NiagaraEditor.Stack.InputValueIconColor.Data";
diff --git a/Engine/Plugins/FX/Niagara/Source/NiagaraEditorWidgets/Private/Stack/SNiagaraStackFunctionInputValue.cpp b/Engine/Plugins/FX/Niagara/Source/NiagaraEditorWidgets/Private/Stack/SNiagaraStackFunctionInputValue.cpp
--- a/Engine/Plugins/FX/Niagara/Source/NiagaraEditorWidgets/Private/Stack/SNiagaraStackFunctionInputValue.cpp
+++ b/Engine/Plugins/FX/Niagara/Source/NiagaraEditorWidgets/Private/Stack/SNiagaraStackFunctionInputValue.cpp
@@ -19,6 +19,7 @@
 #include "NiagaraEditorWidgetsUtilities.h"
 #include "NiagaraNodeAssignment.h"
 #include "NiagaraNodeFunctionCall.h"
+#include "NiagaraNodeOutput.h"
 #include "NiagaraParameterCollection.h"
 #include "NiagaraParameterDefinitions.h"
 #include "NiagaraStackCommandContext.h"
@@ -40,6 +41,7 @@
 #include "ViewModels/NiagaraEmitterViewModel.h"
 #include "ViewModels/NiagaraScratchPadScriptViewModel.h"
 #include "ViewModels/NiagaraScratchPadViewModel.h"
+#include "ViewModels/NiagaraSystemSelectionViewModel.h"
 #include "ViewModels/NiagaraSystemViewModel.h"
 #include "ViewModels/Stack/NiagaraStackFunctionInput.h"
 #include "ViewModels/Stack/NiagaraStackValueCollection.h"
@@ -138,6 +140,87 @@ public:
 	}
 };
 
+class SNiagaraStackFunctionLinkedSharedButton : public SComboButton
+{
+public:
+	SLATE_BEGIN_ARGS(SNiagaraStackFunctionLinkedSharedButton) { }
+	SLATE_END_ARGS()
+
+	void Construct(const FArguments& InArgs, UNiagaraStackFunctionInput* InFunctionInput)
+	{
+		UNiagaraStackFunctionInput* LinkedFunctionInput = InFunctionInput->GetSharedLinkedValueInput();
+		if (LinkedFunctionInput == nullptr)
+		{
+			return;
+		}
+
+		SComboButton::Construct(SComboButton::FArguments()
+			.ComboButtonStyle(&FNiagaraEditorWidgetsStyle::Get().GetWidgetStyle<FComboButtonStyle>("NiagaraEditor.Stack.CompactComboButtonWithBorder"))
+			.OnGetMenuContent(this, &SNiagaraStackFunctionLinkedSharedButton::GetLinkedSharedActionMenu)
+			.ButtonContent()
+			[
+				SNew(SImage)
+				.ColorAndOpacity(FLinearColor(0.7f, .3, 0.7f, 1.0f))
+				.Image(FNiagaraEditorStyle::Get().GetBrush("NiagaraEditor.LinkedSharedIcon"))
+			]);
+
+		SetToolTipText(FText::Format(LOCTEXT("LinkedSharedButtonToolTipFormat",
+			"This input is is linked to a shared parameter {0}.\n"
+			"Edits made here will modify that linked parameter."),
+			FText::FromName(LinkedFunctionInput->GetInputParameterHandle().GetName())));
+
+		FunctionInputWeak = InFunctionInput;
+		LinkedFunctionInputWeak = LinkedFunctionInput;
+	}
+
+private:
+	TSharedRef<SWidget> GetLinkedSharedActionMenu()
+	{
+		FMenuBuilder MenuBuilder(true, nullptr);
+		UNiagaraStackFunctionInput* LinkedFunctionInput = LinkedFunctionInputWeak.Get();
+		if (LinkedFunctionInput != nullptr && LinkedFunctionInput->IsFinalized() == false)
+		{
+			MenuBuilder.BeginSection("LinkedSharedParameterActions", LOCTEXT("LinkedSharedSourceSection", "Linked Shared Parameter Actions"));
+			{
+				MenuBuilder.AddMenuEntry(
+					LOCTEXT("GoToLinkedSharedSource", "Go to source parameter"),
+					LOCTEXT("GoToLinkedSharedSourceToolTip", "Selects the set parameters node that last set the shared source parameter."),
+					FSlateIcon(),
+					FUIAction(FExecuteAction::CreateSP(this, &SNiagaraStackFunctionLinkedSharedButton::SelectLinkedSharedSourceModule)));
+				MenuBuilder.AddMenuEntry(
+					LOCTEXT("UnlinkFromSharedSource", "Unlink from source parameter."),
+					LOCTEXT("UnlinkFromSharedSourceToolTip", "Unlinks this input from the shared source parameter while preserving the value."),
+					FSlateIcon(),
+					FUIAction(FExecuteAction::CreateSP(this, &SNiagaraStackFunctionLinkedSharedButton::UnlinkFromSharedParameter)));
+			}
+		}
+		return MenuBuilder.MakeWidget();
+	}
+
+	void SelectLinkedSharedSourceModule()
+	{
+		UNiagaraStackFunctionInput* LinkedFunctionInput = LinkedFunctionInputWeak.Get();
+		if (LinkedFunctionInput != nullptr && LinkedFunctionInput->IsFinalized() == false)
+		{
+			LinkedFunctionInput->GetSystemViewModel()->GetSelectionViewModel()->EmptySelection();
+			LinkedFunctionInput->GetSystemViewModel()->GetSelectionViewModel()->AddEntryToSelectionByDisplayedObjectDeferred(&LinkedFunctionInput->GetInputFunctionCallNode());
+		}
+	}
+
+	void UnlinkFromSharedParameter()
+	{
+		UNiagaraStackFunctionInput* FunctionInput = FunctionInputWeak.Get();
+		if (FunctionInput != nullptr && FunctionInput->IsFinalized() == false)
+		{
+			FunctionInput->UnlinkFromSharedParameter();
+		}
+	}
+
+private:
+	TWeakObjectPtr<UNiagaraStackFunctionInput> FunctionInputWeak;
+	TWeakObjectPtr<UNiagaraStackFunctionInput> LinkedFunctionInputWeak;
+};
+
 void SNiagaraStackFunctionInputValue::Construct(const FArguments& InArgs, UNiagaraStackFunctionInput* InFunctionInput)
 {
 	FunctionInput = InFunctionInput;
@@ -150,13 +233,9 @@ void SNiagaraStackFunctionInputValue::Construct(const FArguments& InArgs, UNiaga
 	EntryIsEnabled.Bind(this, &SNiagaraStackFunctionInputValue::GetEntryEnabled);
 	SetEnabled(EntryIsEnabled);
 
-	TSharedPtr<SHorizontalBox> OuterChildrenBox;
 	TSharedPtr<SHorizontalBox> ChildrenBox;
 	ChildSlot
 	[
-		SAssignNew(OuterChildrenBox, SHorizontalBox)
-		+ SHorizontalBox::Slot()
-		[
 			SNew(SNiagaraParameterDropTarget)
 			.TypeToTestAgainst(FunctionInput->GetInputType())
 			.ExecutionCategory(FunctionInput->GetExecutionCategoryName())
@@ -173,10 +252,16 @@ void SNiagaraStackFunctionInputValue::Construct(const FArguments& InArgs, UNiaga
 				SAssignNew(ChildrenBox, SHorizontalBox)
 				.IsEnabled(this, &SNiagaraStackFunctionInputValue::GetInputEnabled)
 			])
-		]
 	];
-			// Values
-	if (LayoutMode == ELayoutMode::FullRow)
+
+	bool bAddIndent = LayoutMode == ELayoutMode::FullRow;
+	bool bAddValue = LayoutMode != ELayoutMode::EditDropDownOnly;
+	bool bAddEditDropDown = LayoutMode == ELayoutMode::FullRow || LayoutMode == ELayoutMode::LinkedInline;
+	bool bAddCompactEditDropdown = LayoutMode == ELayoutMode::CompactInline || LayoutMode == ELayoutMode::EditDropDownOnly;
+	bool bAddResetButtons = LayoutMode == ELayoutMode::FullRow;
+
+	// Indent
+	if (bAddIndent)
 	{
 			ChildrenBox->AddSlot()
 				.AutoWidth()
@@ -185,7 +270,8 @@ void SNiagaraStackFunctionInputValue::Construct(const FArguments& InArgs, UNiaga
 					SNew(SNiagaraStackIndent, FunctionInput, ENiagaraStackIndentMode::Value)
 				];
 	}
-	if (LayoutMode == ELayoutMode::FullRow || LayoutMode == ELayoutMode::CompactInline)
+
+	if (bAddValue)
 	{
 			ChildrenBox->AddSlot()
 				.VAlign(VAlign_Center)
@@ -200,10 +286,22 @@ void SNiagaraStackFunctionInputValue::Construct(const FArguments& InArgs, UNiaga
 				];
 
 			ValueModeForGeneratedWidgets = FunctionInput->GetValueMode();
+
+			if (InArgs._ValueDecorator.Widget != SNullWidget::NullWidget)
+			{
+				ChildrenBox->AddSlot()
+					.VAlign(VAlign_Center)
+					.Padding(0)
+					.AutoWidth()
+					[
+						InArgs._ValueDecorator.Widget
+					];
+			}
 	}
 
-	if (LayoutMode == ELayoutMode::FullRow)
-	{ 
+	// Edit drop down
+	if (bAddEditDropDown)
+	{
 				// Handle drop-down button
 				ChildrenBox->AddSlot()
 				.AutoWidth()
@@ -221,7 +319,11 @@ void SNiagaraStackFunctionInputValue::Construct(const FArguments& InArgs, UNiaga
 					.HAlign(HAlign_Center)
 					.VAlign(VAlign_Center)
 				];
+	}
 
+	// Reset buttons
+	if (bAddResetButtons)
+	{
 				// Reset Button
 				ChildrenBox->AddSlot()
 				.AutoWidth()
@@ -264,7 +366,8 @@ void SNiagaraStackFunctionInputValue::Construct(const FArguments& InArgs, UNiaga
 				];
 	}
 
-	if (LayoutMode == ELayoutMode::CompactInline || LayoutMode == ELayoutMode::EditDropDownOnly)
+	// Compact edit drop down
+	if (bAddCompactEditDropdown)
 	{
 		ChildrenBox->AddSlot()
 		.AutoWidth()
@@ -312,6 +415,11 @@ TSharedRef<SWidget> SNiagaraStackFunctionInputValue::ConstructValueWidgets()
 		ConstructLinkedValueWidgets(ValueBox);
 		break;
 
+	case UNiagaraStackFunctionInput::EValueMode::LinkedShared:
+		ConstructLinkedSharedValueWidgets(ValueBox);
+		bShowTypeIcon = false;
+		break;
+
 	case UNiagaraStackFunctionInput::EValueMode::Data:
 		ConstructTextValueWidgets(ValueBox, TAttribute<FText>::CreateSP(this, &SNiagaraStackFunctionInputValue::GetDataValueText));
 		break;
@@ -595,6 +703,28 @@ void SNiagaraStackFunctionInputValue::ConstructLinkedValueWidgets(TSharedRef<SHo
 		];
 }
 
+void SNiagaraStackFunctionInputValue::ConstructLinkedSharedValueWidgets(TSharedRef<SHorizontalBox>& ValueBox)
+{
+	ValueBox->AddSlot()
+		.VAlign(VAlign_Center)
+		.Padding(0)
+		[
+			SNew(SBorder)
+				.BorderImage(FAppStyle::GetBrush("WhiteBrush"))
+				.BorderBackgroundColor(FLinearColor(0.5f, 0, 0.5f, 0.2f))
+				.Padding(0)
+				[
+					SNew(SNiagaraStackFunctionInputValue, FunctionInput->GetSharedLinkedValueInput())
+						.LayoutMode(ELayoutMode::LinkedInline)
+						.ValueDecorator()
+						[
+							SNew(SNiagaraStackFunctionLinkedSharedButton, FunctionInput)
+								.Visibility(this, &SNiagaraStackFunctionInputValue::GetVisibilityForHoverOnlyWidgets)
+						]
+				]
+		];
+}
+
 void SNiagaraStackFunctionInputValue::ConstructDynamicValueWidgets(TSharedRef<SHorizontalBox>& ValueBox)
 {
 	ValueBox->AddSlot()
@@ -979,7 +1109,14 @@ EVisibility SNiagaraStackFunctionInputValue::GetResetButtonVisibility() const
 
 EVisibility SNiagaraStackFunctionInputValue::GetDropdownButtonVisibility() const
 {
-	return FunctionInput->IsStaticParameter() ? EVisibility::Hidden : EVisibility::Visible;
+	if (FunctionInput->GetValueMode() == UNiagaraStackFunctionInput::EValueMode::LinkedShared)
+	{
+		return EVisibility::Collapsed;
+	}
+	else
+	{
+		return FunctionInput->IsStaticParameter() ? EVisibility::Hidden : EVisibility::Visible;
+	}
 }
 
 FReply SNiagaraStackFunctionInputValue::ResetButtonPressed() const
@@ -1211,6 +1348,11 @@ FReply SNiagaraStackFunctionInputValue::ScratchButtonPressed() const
 	return FReply::Unhandled();
 }
 
+EVisibility SNiagaraStackFunctionInputValue::GetVisibilityForHoverOnlyWidgets() const
+{
+	return IsHovered() ? EVisibility::Visible : EVisibility::Collapsed;
+}
+
 void SNiagaraStackFunctionInputValue::OnActionRowHoverEvent(const TSharedPtr<FNiagaraMenuAction_Generic>& ActionNode, bool bIsHovered)
 {
 	ActionNode->bIsHovered = bIsHovered;
@@ -1459,6 +1601,37 @@ TArray<TSharedPtr<FNiagaraMenuAction_Generic>> SNiagaraStackFunctionInputValue::
 		OutAllActions.Add(CreateScratchAction);
 	}
 
+	TArray<ENiagaraScriptUsage> ScriptUsagesForSharedParameters;
+	FunctionInput->GetAvailableScriptUsagesForSharedParameters(ScriptUsagesForSharedParameters);
+	if (ScriptUsagesForSharedParameters.Num() > 0)
+	{
+		UEnum* ScriptUsageEnum = StaticEnum<ENiagaraScriptUsage>();
+		for (ENiagaraScriptUsage TargetUsage : ScriptUsagesForSharedParameters)
+		{
+			FText TargetUsageDisplayName = ScriptUsageEnum->GetDisplayNameTextByValue((int64)TargetUsage);
+			const FText DisplayName = FText::Format(LOCTEXT("MoveToSharedParameterFormat", "Shared {0} parameter."), TargetUsageDisplayName);
+			const FText ToolTip = FText::Format(LOCTEXT("MoveToSharedParameterToolTipFormat", "Move the value of this input to a shared {0} parameter and show that parameter value inline here."), TargetUsageDisplayName);
+
+			TSharedPtr<FNiagaraMenuAction_Generic> ShareInTargetUsageAction(new FNiagaraMenuAction_Generic(
+				FNiagaraMenuAction_Generic::FOnExecuteAction::CreateUObject(FunctionInput, &UNiagaraStackFunctionInput::MoveValueToLinkedSharedParameter, TOptional<ENiagaraScriptUsage>(TargetUsage)),
+				DisplayName, { TEXT("Share Input Value") }, {}, ToolTip, FText()));
+
+			ShareInTargetUsageAction->SourceData = NiagaraSourceData;
+			OutAllActions.Add(ShareInTargetUsageAction);
+		}
+
+		const FText TransientDisplayName = LOCTEXT("MoveToSharedTransientParameterFormat", "Shared transient parameter.");
+		const FText TransientToolTip = LOCTEXT("MoveToSharedTransientParameterToolTipFormat", "Move the value of this input to a shared transient parameter and show that parameter value inline here.");
+
+		TSharedPtr<FNiagaraMenuAction_Generic> ShareInTransientAction(new FNiagaraMenuAction_Generic(
+			FNiagaraMenuAction_Generic::FOnExecuteAction::CreateUObject(FunctionInput, &UNiagaraStackFunctionInput::MoveValueToLinkedSharedParameter, TOptional<ENiagaraScriptUsage>()),
+			TransientDisplayName, { TEXT("Share Input Value") }, {}, TransientToolTip, FText()));
+
+		ShareInTransientAction->SourceData = NiagaraSourceData;
+
+		OutAllActions.Add(ShareInTransientAction);
+	}
+
 	if (FunctionInput->CanDeleteInput())
 	{
 		const FText DisplayName = LOCTEXT("DeleteInput", "Remove this input");
diff --git a/Engine/Plugins/FX/Niagara/Source/NiagaraEditorWidgets/Private/Stack/SNiagaraStackFunctionInputValue.h b/Engine/Plugins/FX/Niagara/Source/NiagaraEditorWidgets/Private/Stack/SNiagaraStackFunctionInputValue.h
--- a/Engine/Plugins/FX/Niagara/Source/NiagaraEditorWidgets/Private/Stack/SNiagaraStackFunctionInputValue.h
+++ b/Engine/Plugins/FX/Niagara/Source/NiagaraEditorWidgets/Private/Stack/SNiagaraStackFunctionInputValue.h
@@ -33,6 +33,7 @@ public:
 	{
 		FullRow,
 		CompactInline,
+		LinkedInline,
 		EditDropDownOnly
 	};
 
@@ -43,6 +44,7 @@ public:
 		{ }
 		SLATE_ARGUMENT(ELayoutMode, LayoutMode)
 		SLATE_ATTRIBUTE(EVisibility, CompactActionMenuButtonVisibility)
+		SLATE_NAMED_SLOT(FArguments, ValueDecorator)
 	SLATE_END_ARGS();
 
 	void Construct(const FArguments& InArgs, UNiagaraStackFunctionInput* InFunctionInput);
@@ -62,6 +64,8 @@ private:
 
 	void ConstructLinkedValueWidgets(TSharedRef<SHorizontalBox>& ValueBox);
 
+	void ConstructLinkedSharedValueWidgets(TSharedRef<SHorizontalBox>& ValueBox);
+
 	void ConstructDynamicValueWidgets(TSharedRef<SHorizontalBox>& ValueBox);
 
 	void OnInputValueChanged();
@@ -171,6 +175,8 @@ private:
 
 	FReply ScratchButtonPressed() const;
 
+	EVisibility GetVisibilityForHoverOnlyWidgets() const;
+
 private:
 	UNiagaraStackFunctionInput* FunctionInput;
 
