Fixed potential issue with Json write.

diff --git a/Engine/Plugins/FX/Niagara/Source/Niagara/Private/NiagaraDataChannelDebugUtilities.cpp b/Engine/Plugins/FX/Niagara/Source/Niagara/Private/NiagaraDataChannelDebugUtilities.cpp
--- a/Engine/Plugins/FX/Niagara/Source/Niagara/Private/NiagaraDataChannelDebugUtilities.cpp
+++ b/Engine/Plugins/FX/Niagara/Source/Niagara/Private/NiagaraDataChannelDebugUtilities.cpp
@@ -187,6 +187,11 @@ void FNiagaraDataChannelDebugUtilities::TearDown()
 
 FString FNiagaraDataChannelDebugUtilities::ToJson(FNiagaraDataChannelGameData* Data)
 {
+	if(Data == nullptr)
+	{
+		return TEXT("{}");
+	}
+
 	FString JsonString;
 	TSharedRef<TJsonWriter<>> JsonWriter = TJsonWriterFactory<>::Create(&JsonString);
 	JsonWriter->WriteArrayStart();
