- Fix crashes from using raw delegates

diff --git a/Engine/Plugins/FX/Niagara/Source/NiagaraEditor/Private/Toolkits/SystemToolkitModes/NiagaraSystemToolkitModeBase.cpp b/Engine/Plugins/FX/Niagara/Source/NiagaraEditor/Private/Toolkits/SystemToolkitModes/NiagaraSystemToolkitModeBase.cpp
--- a/Engine/Plugins/FX/Niagara/Source/NiagaraEditor/Private/Toolkits/SystemToolkitModes/NiagaraSystemToolkitModeBase.cpp
+++ b/Engine/Plugins/FX/Niagara/Source/NiagaraEditor/Private/Toolkits/SystemToolkitModes/NiagaraSystemToolkitModeBase.cpp
@@ -119,14 +119,13 @@ void FNiagaraSystemToolkitModeBase::OnActiveDocumentChanged(TSharedPtr<SDockTab>
 		TSharedPtr<FNiagaraScratchPadScriptViewModel> ScratchScriptVM = Toolkit->GetSystemViewModel()->GetDocumentViewModel()->GetActiveScratchPadViewModelIfSet();
 		if (ScratchScriptVM.IsValid())
 		{
-			LastSelectionUpdateDelegate = ScratchScriptVM->GetGraphViewModel()->GetNodeSelection()->OnSelectedObjectsChanged().AddRaw(this, &FNiagaraSystemToolkitModeBase::UpdateSelectionForActiveDocument);
+			LastSelectionUpdateDelegate = ScratchScriptVM->GetGraphViewModel()->GetNodeSelection()->OnSelectedObjectsChanged().AddSP(this, &FNiagaraSystemToolkitModeBase::UpdateSelectionForActiveDocument);
 			if (INiagaraParameterPanelViewModel* ParameterPanelViewModel = Toolkit->GetSystemViewModel()->GetParameterPanelViewModel())
 			{
-				LastParamPanelSelectionUpdateDelegate = ParameterPanelViewModel->GetVariableObjectSelection()->OnSelectedObjectsChanged().AddRaw(this, &FNiagaraSystemToolkitModeBase::OnParameterPanelViewModelExternalSelectionChanged);
+				LastParamPanelSelectionUpdateDelegate = ParameterPanelViewModel->GetVariableObjectSelection()->OnSelectedObjectsChanged().AddSP(this, &FNiagaraSystemToolkitModeBase::OnParameterPanelViewModelExternalSelectionChanged);
 			}
 			LastSystemSelectionUpdateDelegate = Toolkit->GetSystemViewModel()->GetSelectionViewModel()->OnEntrySelectionChanged().AddSP(this, &FNiagaraSystemToolkitModeBase::OnSystemSelectionChanged);
-			LastGraphEditDelegate = ScratchScriptVM->GetGraphViewModel()->GetGraph()->AddOnGraphNeedsRecompileHandler(
-				FOnGraphChanged::FDelegate::CreateRaw(this, &FNiagaraSystemToolkitModeBase::OnEditedScriptGraphChanged));
+			LastGraphEditDelegate = ScratchScriptVM->GetGraphViewModel()->GetGraph()->AddOnGraphNeedsRecompileHandler(FOnGraphChanged::FDelegate::CreateSP(this, &FNiagaraSystemToolkitModeBase::OnEditedScriptGraphChanged));
 		}
 	}
 }
