- Fix crash when wiring up parameter map for loop from static node

diff --git a/Engine/Plugins/FX/Niagara/Source/NiagaraEditor/Private/NiagaraNode.cpp b/Engine/Plugins/FX/Niagara/Source/NiagaraEditor/Private/NiagaraNode.cpp
--- a/Engine/Plugins/FX/Niagara/Source/NiagaraEditor/Private/NiagaraNode.cpp
+++ b/Engine/Plugins/FX/Niagara/Source/NiagaraEditor/Private/NiagaraNode.cpp
@@ -472,7 +472,10 @@ void UNiagaraNode::AutowireNewNode(UEdGraphPin* FromPin)
 
 		if (FirstPinOfSameType && GetSchema()->TryCreateConnection(FromPin, FirstPinOfSameType))
 		{
-			FromPin->GetOwningNode()->NodeConnectionListChanged();
+			if (UEdGraphNode* FromPinNode = FromPin->GetOwningNodeUnchecked())
+			{
+				FromPinNode->NodeConnectionListChanged();
+			}
 		}
 	}
 }
