- Moved binding for OnFinishedChangingProperties into reachable code

#rnx

diff --git a/Engine/Plugins/FX/Niagara/Source/NiagaraEditor/Private/Toolkits/NiagaraParameterDefinitionsToolkit.cpp b/Engine/Plugins/FX/Niagara/Source/NiagaraEditor/Private/Toolkits/NiagaraParameterDefinitionsToolkit.cpp
--- a/Engine/Plugins/FX/Niagara/Source/NiagaraEditor/Private/Toolkits/NiagaraParameterDefinitionsToolkit.cpp
+++ b/Engine/Plugins/FX/Niagara/Source/NiagaraEditor/Private/Toolkits/NiagaraParameterDefinitionsToolkit.cpp
@@ -252,15 +252,18 @@ TSharedRef<SDockTab> FNiagaraParameterDefinitionsToolkit::SpawnTab_SelectedScrip
 {
 	checkf(Args.GetTabId().TabType == SelectedDetailsTabId, TEXT("Wrong tab ID in NiagaraParameterDefinitionsToolkit!"));
 
-	return SNew(SDockTab)
-	.Label(LOCTEXT("SelectedDetailsTabLabel", "Selected Details"))
-	.TabColorScale(GetTabColorScale())
-	[
-		SAssignNew(SelectedScriptVarDetailsWidget, SNiagaraSelectedObjectsDetails, DetailsScriptSelection.ToSharedRef())
-		.AllowEditingLibraryScriptVariables(true)
-	];
-
-	SelectedScriptVarDetailsWidget->OnFinishedChangingProperties().AddRaw(this, &FNiagaraParameterDefinitionsToolkit::OnEditedParameterDefinitionsPropertyFinishedChanging);
+	TSharedRef<SDockTab> DockTab =
+		SNew(SDockTab)
+		.Label(LOCTEXT("SelectedDetailsTabLabel", "Selected Details"))
+		.TabColorScale(GetTabColorScale())
+		[
+			SAssignNew(SelectedScriptVarDetailsWidget, SNiagaraSelectedObjectsDetails, DetailsScriptSelection.ToSharedRef())
+			.AllowEditingLibraryScriptVariables(true)
+		];
+
+	SelectedScriptVarDetailsWidget->OnFinishedChangingProperties().AddSP(this, &FNiagaraParameterDefinitionsToolkit::OnEditedParameterDefinitionsPropertyFinishedChanging);
+
+	return DockTab;
 }
 
 void FNiagaraParameterDefinitionsToolkit::ExtendToolbar()
