Added two utility functions to query basic user parameter and emitter info of niagara systems.

diff --git a/Engine/Plugins/FX/Niagara/Source/Niagara/Private/NiagaraFunctionLibrary.cpp b/Engine/Plugins/FX/Niagara/Source/Niagara/Private/NiagaraFunctionLibrary.cpp
--- a/Engine/Plugins/FX/Niagara/Source/Niagara/Private/NiagaraFunctionLibrary.cpp
+++ b/Engine/Plugins/FX/Niagara/Source/Niagara/Private/NiagaraFunctionLibrary.cpp
@@ -671,6 +671,50 @@ void UNiagaraFunctionLibrary::SetVolumeTextureObject(UNiagaraComponent* NiagaraS
 	SetTextureObject(NiagaraSystem, OverrideName, Texture);
 }
 
+TArray<FNiagaraUserParameterInfo> UNiagaraFunctionLibrary::GetAllUserParameters(UNiagaraSystem* NiagaraSystem)
+{
+	TArray<FNiagaraUserParameterInfo> Result;
+	if (NiagaraSystem)
+	{
+		TArray<FNiagaraVariable> Parameters;
+		NiagaraSystem->GetExposedParameters().GetUserParameters(Parameters);
+		for (const FNiagaraVariable& UserVar : Parameters)
+		{
+			FNiagaraUserParameterInfo& ParameterInfo = Result.AddDefaulted_GetRef();
+			ParameterInfo.ParameterName = UserVar.GetName();
+			ParameterInfo.ParameterType = UserVar.GetType().GetStruct();
+			ParameterInfo.TypeName = UserVar.GetType().GetFName();
+			ParameterInfo.ParameterSizeBytes = UserVar.GetType().GetSize();
+			ParameterInfo.bIsEnumType = UserVar.GetType().IsEnum();
+		}
+	}
+	return Result;
+}
+
+TArray<FNiagaraMinimalEmitterInfo> UNiagaraFunctionLibrary::GetAllEmitters(UNiagaraSystem* NiagaraSystem)
+{
+	TArray<FNiagaraMinimalEmitterInfo> Result;
+	if (NiagaraSystem)
+	{
+		for (const FNiagaraEmitterHandle& EmitterHandle : NiagaraSystem->GetEmitterHandles())
+		{
+			
+			FNiagaraMinimalEmitterInfo& EmitterInfo = Result.AddDefaulted_GetRef();
+			EmitterInfo.EmitterName = EmitterHandle.GetName();
+			EmitterInfo.bIsEnabled = EmitterHandle.GetIsEnabled();
+			EmitterInfo.bIsLightweight = EmitterHandle.GetEmitterMode() == ENiagaraEmitterMode::Stateless;
+			
+			EmitterHandle.ForEachEnabledRendererWithIndex([&EmitterInfo](const UNiagaraRendererProperties* Properties, int32)
+				{
+					TArray<UMaterialInterface*> Materials;
+					Properties->GetUsedMaterials(nullptr, Materials);
+					EmitterInfo.UsedRendererMaterials.Append(Materials);
+				});
+		}
+	}
+	return Result;
+}
+
 UNiagaraDataInterface* UNiagaraFunctionLibrary::GetDataInterface(UClass* DIClass, UNiagaraComponent* NiagaraSystem, FName OverrideName)
 {
 	if (NiagaraSystem == nullptr)
diff --git a/Engine/Plugins/FX/Niagara/Source/Niagara/Public/NiagaraFunctionLibrary.h b/Engine/Plugins/FX/Niagara/Source/Niagara/Public/NiagaraFunctionLibrary.h
--- a/Engine/Plugins/FX/Niagara/Source/Niagara/Public/NiagaraFunctionLibrary.h
+++ b/Engine/Plugins/FX/Niagara/Source/Niagara/Public/NiagaraFunctionLibrary.h
@@ -16,6 +16,51 @@
 class UNiagaraComponent;
 class USceneComponent;
 class UVolumeTexture;
+class UMaterialInterface;
+
+USTRUCT(MinimalAPI, BlueprintType)
+struct FNiagaraUserParameterInfo
+{
+	GENERATED_BODY()
+
+	/** The name of the user parameter */
+	UPROPERTY(EditAnywhere, BlueprintReadOnly, Category=Niagara)
+	FName ParameterName;
+
+	// the underlying struct, uobject class or enum of the parameter
+	UPROPERTY(EditAnywhere, BlueprintReadOnly, Category=Niagara)
+	TObjectPtr<UStruct> ParameterType;
+
+	// the name of the underlying type (e.g. Vector3f)
+	UPROPERTY(EditAnywhere, BlueprintReadOnly, Category=Niagara)
+	FName TypeName;
+
+	UPROPERTY(EditAnywhere, BlueprintReadOnly, Category=Niagara)
+	int32 ParameterSizeBytes;
+
+	UPROPERTY(EditAnywhere, BlueprintReadOnly, Category=Niagara)
+	bool bIsEnumType;
+};
+
+USTRUCT(MinimalAPI, BlueprintType)
+struct FNiagaraMinimalEmitterInfo
+{
+	GENERATED_BODY()
+
+	UPROPERTY(EditAnywhere, BlueprintReadOnly, Category=Niagara)
+	FName EmitterName;
+
+	UPROPERTY(EditAnywhere, BlueprintReadOnly, Category=Niagara)
+	bool bIsEnabled = false;
+
+	// set to true if this is a lightweight emitter
+	UPROPERTY(EditAnywhere, BlueprintReadOnly, Category=Niagara)
+	bool bIsLightweight = false;
+
+	// A list of all materials referenced by renderers in this emitter. Could actually be less on the cooked target platform, depending on scalability settings.
+	UPROPERTY(EditAnywhere, BlueprintReadOnly, Category=Niagara)
+	TArray<TObjectPtr<UMaterialInterface>> UsedRendererMaterials;
+};
 
 /**
 * A C++ and Blueprint accessible library of utility functions for accessing Niagara simulations
@@ -96,6 +141,14 @@ public:
 	/** Overrides the Volume Texture for a Niagara Volume Texture Data Interface User Parameter.*/
 	UFUNCTION(BlueprintCallable, Category = Niagara)
 	static NIAGARA_API void SetVolumeTextureObject(UNiagaraComponent* NiagaraSystem, UPARAM(DisplayName = "Parameter Name") const FString& OverrideName, UVolumeTexture* Texture);
+
+	/** Returns basic info on all user parameters of a niagara system.*/
+	UFUNCTION(BlueprintCallable, Category = Niagara)
+	static NIAGARA_API TArray<FNiagaraUserParameterInfo> GetAllUserParameters(UNiagaraSystem* NiagaraSystem);
+
+	/** Returns basic info on all emitters of a niagara system.*/
+	UFUNCTION(BlueprintCallable, Category = Niagara)
+	static NIAGARA_API TArray<FNiagaraMinimalEmitterInfo> GetAllEmitters(UNiagaraSystem* NiagaraSystem);
 	
 	/** Finds an array interface of the given class. */
 	static NIAGARA_API UNiagaraDataInterface* GetDataInterface(UClass* DIClass, UNiagaraComponent* NiagaraSystem, FName OverrideName);
