Improved struct field comment

#rnx

diff --git a/Engine/Plugins/FX/Niagara/Source/Niagara/Public/NiagaraFunctionLibrary.h b/Engine/Plugins/FX/Niagara/Source/Niagara/Public/NiagaraFunctionLibrary.h
--- a/Engine/Plugins/FX/Niagara/Source/Niagara/Public/NiagaraFunctionLibrary.h
+++ b/Engine/Plugins/FX/Niagara/Source/Niagara/Public/NiagaraFunctionLibrary.h
@@ -27,7 +27,7 @@ struct FNiagaraUserParameterInfo
 	UPROPERTY(EditAnywhere, BlueprintReadOnly, Category=Niagara)
 	FName ParameterName;
 
-	// the underlying struct, uobject class or enum of the parameter
+	// the underlying struct or uobject class of the parameter (set to integer in case of enums)
 	UPROPERTY(EditAnywhere, BlueprintReadOnly, Category=Niagara)
 	TObjectPtr<UStruct> ParameterType;
 
