Make NiagaraDataInterface binders use BindUObject where possible instead of going through an intermediate lambda

diff --git a/Engine/Plugins/FX/Niagara/Source/Niagara/Classes/NiagaraDataInterface.h b/Engine/Plugins/FX/Niagara/Source/Niagara/Classes/NiagaraDataInterface.h
--- a/Engine/Plugins/FX/Niagara/Source/Niagara/Classes/NiagaraDataInterface.h
+++ b/Engine/Plugins/FX/Niagara/Source/Niagara/Classes/NiagaraDataInterface.h
@@ -119,8 +119,7 @@ struct NDI_FUNC_BINDER(ClassName, FuncName)\
 	template<typename ... ParamTypes>\
 	static void Bind(UNiagaraDataInterface* Interface, const FVMExternalFunctionBindingInfo& BindingInfo, void* InstanceData, FVMExternalFunction &OutFunc)\
 	{\
-		auto Lambda = [Interface](FVectorVMExternalFunctionContext& Context) { static_cast<ClassName*>(Interface)->FuncName<ParamTypes...>(Context); };\
-		OutFunc = FVMExternalFunction::CreateLambda(Lambda);\
+		OutFunc.BindUObject(static_cast<ClassName*>(Interface), &ClassName::FuncName<ParamTypes...>);\
 	}\
 };
 
@@ -129,8 +128,7 @@ struct NDI_FUNC_BINDER(ClassName, FuncName)\
 {\
 	static void Bind(UNiagaraDataInterface* Interface, FVMExternalFunction &OutFunc)\
 	{\
-		auto Lambda = [Interface](FVectorVMExternalFunctionContext& Context) { static_cast<ClassName*>(Interface)->FuncName(Context); };\
-		OutFunc = FVMExternalFunction::CreateLambda(Lambda);\
+		OutFunc.BindUObject(static_cast<ClassName*>(Interface), &ClassName::FuncName);\
 	}\
 };
 
@@ -141,7 +139,7 @@ struct NDI_FUNC_BINDER(ClassName, FuncName)\
 	static void Bind(UNiagaraDataInterface* Interface, FVMExternalFunction &OutFunc, VarTypes... Var)\
 	{\
 		auto Lambda = [Interface, Var...](FVectorVMExternalFunctionContext& Context) { static_cast<ClassName*>(Interface)->FuncName(Context, Var...); };\
-		OutFunc = FVMExternalFunction::CreateLambda(Lambda);\
+		OutFunc.BindLambda(Lambda);\
 	}\
 };
 
diff --git a/Engine/Plugins/FX/Niagara/Source/Niagara/Private/DataInterface/NiagaraDataInterfaceStaticMesh.cpp b/Engine/Plugins/FX/Niagara/Source/Niagara/Private/DataInterface/NiagaraDataInterfaceStaticMesh.cpp
--- a/Engine/Plugins/FX/Niagara/Source/Niagara/Private/DataInterface/NiagaraDataInterfaceStaticMesh.cpp
+++ b/Engine/Plugins/FX/Niagara/Source/Niagara/Private/DataInterface/NiagaraDataInterfaceStaticMesh.cpp
@@ -2643,358 +2643,358 @@ void UNiagaraDataInterfaceStaticMesh::GetVMExternalFunction(const FVMExternalFun
 	// VM Vertex Sampling
 	if (BindingInfo.Name == IsValidVertexName)
 	{
-		OutFunc = FVMExternalFunction::CreateLambda([this](FVectorVMExternalFunctionContext& Context) { VMIsValidVertex(Context); });
+		OutFunc.BindUObject(this, &UNiagaraDataInterfaceStaticMesh::VMIsValidVertex);
 	}
 	else if (BindingInfo.Name == RandomVertexName)
 	{
-		OutFunc = FVMExternalFunction::CreateLambda([this](FVectorVMExternalFunctionContext& Context) { VMRandomVertex(Context); });
+		OutFunc.BindUObject(this, &UNiagaraDataInterfaceStaticMesh::VMRandomVertex);
 	}
 	else if (BindingInfo.Name == GetVertexCountName)
 	{
-		OutFunc = FVMExternalFunction::CreateLambda([this](FVectorVMExternalFunctionContext& Context) { VMGetVertexCount(Context); });
+		OutFunc.BindUObject(this, &UNiagaraDataInterfaceStaticMesh::VMGetVertexCount);
 	}
 	//else if (BindingInfo.Name == IsValidFilteredVertexName)
 	//{
-	//	OutFunc = FVMExternalFunction::CreateLambda([this](FVectorVMExternalFunctionContext& Context) { VMIsValidFilteredVertex(Context); });
+	//	OutFunc.BindUObject(this, &UNiagaraDataInterfaceStaticMesh::VMIsValidFilteredVertex);
 	//}
 	//else if (BindingInfo.Name == RandomFilteredVertexName)
 	//{
-	//	OutFunc = FVMExternalFunction::CreateLambda([this](FVectorVMExternalFunctionContext& Context) { VMRandomFilteredVertexAt(Context); });
+	//	OutFunc.BindUObject(this, &UNiagaraDataInterfaceStaticMesh::VMRandomFilteredVertexAt);
 	//}
 	//else if (BindingInfo.Name == GetFilteredVertexCountName)
 	//{
-	//	OutFunc = FVMExternalFunction::CreateLambda([this](FVectorVMExternalFunctionContext& Context) { VMGetFilteredVertexCount(Context); });
+	//	OutFunc.BindUObject(this, &UNiagaraDataInterfaceStaticMesh::VMGetFilteredVertexCount);
 	//}
 	//else if (BindingInfo.Name == GetFilteredVertexAtName)
 	//{
-	//	OutFunc = FVMExternalFunction::CreateLambda([this](FVectorVMExternalFunctionContext& Context) { VMGetFilteredVertexAt(Context); });
+	//	OutFunc.BindUObject(this, &UNiagaraDataInterfaceStaticMesh::VMGetFilteredVertexAt);
 	//}
 	else if (BindingInfo.Name == GetVertexName)
 	{
-		OutFunc = FVMExternalFunction::CreateLambda([this](FVectorVMExternalFunctionContext& Context) { VMGetVertex<FNDITransformHandlerNoop>(Context); });
+		OutFunc.BindUObject(this, &UNiagaraDataInterfaceStaticMesh::VMGetVertex<FNDITransformHandlerNoop>);
 	}
 	else if (BindingInfo.Name == GetVertexWSName)
 	{
-		OutFunc = FVMExternalFunction::CreateLambda([this](FVectorVMExternalFunctionContext& Context) { VMGetVertex<FNDITransformHandler>(Context); });
+		OutFunc.BindUObject(this, &UNiagaraDataInterfaceStaticMesh::VMGetVertex<FNDITransformHandler>);
 	}
 	else if (BindingInfo.Name == GetVertexWSInterpolatedName)
 	{
-		OutFunc = FVMExternalFunction::CreateUObject(this, &UNiagaraDataInterfaceStaticMesh::VMGetVertexInterpolated<FNDITransformHandler>);
+		OutFunc.BindUObject(this, &UNiagaraDataInterfaceStaticMesh::VMGetVertexInterpolated<FNDITransformHandler>);
 	}
 	else if (BindingInfo.Name == GetVertexColorName)
 	{
-		OutFunc = FVMExternalFunction::CreateLambda([this](FVectorVMExternalFunctionContext& Context) { VMGetVertexColor(Context); });
+		OutFunc.BindUObject(this, &UNiagaraDataInterfaceStaticMesh::VMGetVertexColor);
 	}
 	else if (BindingInfo.Name == GetVertexUVName)
 	{
-		OutFunc = FVMExternalFunction::CreateLambda([this](FVectorVMExternalFunctionContext& Context) { VMGetVertexUV(Context); });
+		OutFunc.BindUObject(this, &UNiagaraDataInterfaceStaticMesh::VMGetVertexUV);
 	}
 
 	//////////////////////////////////////////////////////////////////////////
 	// VM Triangle Sampling
 	else if (BindingInfo.Name == IsValidTriangleName)
 	{
-		OutFunc = FVMExternalFunction::CreateLambda([this](FVectorVMExternalFunctionContext& Context) { VMIsValidTriangle(Context); });
+		OutFunc.BindUObject(this, &UNiagaraDataInterfaceStaticMesh::VMIsValidTriangle);
 	}
 	else if (BindingInfo.Name == RandomTriangleName)
 	{
-		OutFunc = FVMExternalFunction::CreateLambda([this](FVectorVMExternalFunctionContext& Context) { VMRandomTriangle<FNDIRandomHelper>(Context); });
+		OutFunc.BindUObject(this, &UNiagaraDataInterfaceStaticMesh::VMRandomTriangle<FNDIRandomHelper>);
 	}
 	else if (BindingInfo.Name == GetTriangleCountName)
 	{
-		OutFunc = FVMExternalFunction::CreateLambda([this](FVectorVMExternalFunctionContext& Context) { VMGetTriangleCount(Context); });
+		OutFunc.BindUObject(this, &UNiagaraDataInterfaceStaticMesh::VMGetTriangleCount);
 	}
 	else if (BindingInfo.Name == IsValidFilteredTriangleName)
 	{
-		OutFunc = FVMExternalFunction::CreateLambda([this](FVectorVMExternalFunctionContext& Context) { VMIsValidFilteredTriangle(Context); });
+		OutFunc.BindUObject(this, &UNiagaraDataInterfaceStaticMesh::VMIsValidFilteredTriangle);
 	}
 	else if (BindingInfo.Name == RandomFilteredTriangleName)
 	{
-		OutFunc = FVMExternalFunction::CreateLambda([this](FVectorVMExternalFunctionContext& Context) { VMRandomFilteredTriangle(Context); });
+		OutFunc.BindUObject(this, &UNiagaraDataInterfaceStaticMesh::VMRandomFilteredTriangle);
 	}
 	else if (BindingInfo.Name == GetFilteredTriangleCountName)
 	{
-		OutFunc = FVMExternalFunction::CreateLambda([this](FVectorVMExternalFunctionContext& Context) { VMGetFilteredTriangleCount(Context); });
+		OutFunc.BindUObject(this, &UNiagaraDataInterfaceStaticMesh::VMGetFilteredTriangleCount);
 	}
 	else if (BindingInfo.Name == GetFilteredTriangleAtName)
 	{
-		OutFunc = FVMExternalFunction::CreateLambda([this](FVectorVMExternalFunctionContext& Context) { VMGetFilteredTriangleAt(Context); });
+		OutFunc.BindUObject(this, &UNiagaraDataInterfaceStaticMesh::VMGetFilteredTriangleAt);
 	}
 	else if (BindingInfo.Name == IsValidUnfilteredTriangleName)
 	{
-		OutFunc = FVMExternalFunction::CreateLambda([this](FVectorVMExternalFunctionContext& Context) { VMIsValidUnfilteredTriangle(Context); });
+		OutFunc.BindUObject(this, &UNiagaraDataInterfaceStaticMesh::VMIsValidUnfilteredTriangle);
 	}
 	else if (BindingInfo.Name == RandomUnfilteredTriangleName)
 	{
-		OutFunc = FVMExternalFunction::CreateLambda([this](FVectorVMExternalFunctionContext& Context) { VMRandomUnfilteredTriangle(Context); });
+		OutFunc.BindUObject(this, &UNiagaraDataInterfaceStaticMesh::VMRandomUnfilteredTriangle);
 	}
 	else if (BindingInfo.Name == GetUnfilteredTriangleCountName)
 	{
-		OutFunc = FVMExternalFunction::CreateLambda([this](FVectorVMExternalFunctionContext& Context) { VMGetUnfilteredTriangleCount(Context); });
+		OutFunc.BindUObject(this, &UNiagaraDataInterfaceStaticMesh::VMGetUnfilteredTriangleCount);
 	}
 	else if (BindingInfo.Name == GetUnfilteredTriangleAtName)
 	{
-		OutFunc = FVMExternalFunction::CreateLambda([this](FVectorVMExternalFunctionContext& Context) { VMGetUnfilteredTriangleAt(Context); });
+		OutFunc.BindUObject(this, &UNiagaraDataInterfaceStaticMesh::VMGetUnfilteredTriangleAt);
 	}
 	else if (BindingInfo.Name == GetTriangleName)
 	{
-		OutFunc = FVMExternalFunction::CreateLambda([this](FVectorVMExternalFunctionContext& Context) { VMGetTriangle<FNDITransformHandlerNoop>(Context); });
+		OutFunc.BindUObject(this, &UNiagaraDataInterfaceStaticMesh::VMGetTriangle<FNDITransformHandlerNoop>);
 	}
 	else if (BindingInfo.Name == GetTriangleWSName)
 	{
-		OutFunc = FVMExternalFunction::CreateLambda([this](FVectorVMExternalFunctionContext& Context) { VMGetTriangle<FNDITransformHandler>(Context); });
+		OutFunc.BindUObject(this, &UNiagaraDataInterfaceStaticMesh::VMGetTriangle<FNDITransformHandler>);
 	}
 	else if (BindingInfo.Name == GetTriangleWSInterpolatedName)
 	{
-		OutFunc = FVMExternalFunction::CreateUObject(this, &UNiagaraDataInterfaceStaticMesh::VMGetTriangleInterpolated<FNDITransformHandler>);
+		OutFunc.BindUObject(this, &UNiagaraDataInterfaceStaticMesh::VMGetTriangleInterpolated<FNDITransformHandler>);
 	}
 	else if (BindingInfo.Name == GetTriangleColorName)
 	{
-		OutFunc = FVMExternalFunction::CreateLambda([this](FVectorVMExternalFunctionContext& Context) { VMGetTriangleColor(Context); });
+		OutFunc.BindUObject(this, &UNiagaraDataInterfaceStaticMesh::VMGetTriangleColor);
 	}
 	else if (BindingInfo.Name == GetTriangleUVName)
 	{
-		OutFunc = FVMExternalFunction::CreateLambda([this](FVectorVMExternalFunctionContext& Context) { VMGetTriangleUV(Context); });
+		OutFunc.BindUObject(this, &UNiagaraDataInterfaceStaticMesh::VMGetTriangleUV);
 	}
 	else if (BindingInfo.Name == GetTriangleIndicesName)
 	{
-		OutFunc = FVMExternalFunction::CreateLambda([this](FVectorVMExternalFunctionContext& Context) { VMGetTriangleIndices(Context); });
+		OutFunc.BindUObject(this, &UNiagaraDataInterfaceStaticMesh::VMGetTriangleIndices);
 	}
 
 	//////////////////////////////////////////////////////////////////////////
 	// Socket Functions
 	else if (BindingInfo.Name == GetSocketCountName)
 	{
-		OutFunc = FVMExternalFunction::CreateLambda([this](FVectorVMExternalFunctionContext& Context) { VMGetSocketCount(Context); });
+		OutFunc.BindUObject(this, &UNiagaraDataInterfaceStaticMesh::VMGetSocketCount);
 	}
 	else if (BindingInfo.Name == GetFilteredSocketCountName)
 	{
-		OutFunc = FVMExternalFunction::CreateLambda([this](FVectorVMExternalFunctionContext& Context) { VMGetFilteredSocketCount(Context); });
+		OutFunc.BindUObject(this, &UNiagaraDataInterfaceStaticMesh::VMGetFilteredSocketCount);
 	}
 	else if (BindingInfo.Name == GetUnfilteredSocketCountName)
 	{
-		OutFunc = FVMExternalFunction::CreateLambda([this](FVectorVMExternalFunctionContext& Context) { VMGetUnfilteredSocketCount(Context); });
+		OutFunc.BindUObject(this, &UNiagaraDataInterfaceStaticMesh::VMGetUnfilteredSocketCount);
 	}
 	else if (BindingInfo.Name == RandomSocketName)
 	{
-		OutFunc = FVMExternalFunction::CreateUObject(this, &UNiagaraDataInterfaceStaticMesh::VMRandomSocket);
+		OutFunc.BindUObject(this, &UNiagaraDataInterfaceStaticMesh::VMRandomSocket);
 	}
 	else if (BindingInfo.Name == RandomFilteredSocketName)
 	{
-		OutFunc = FVMExternalFunction::CreateUObject(this, &UNiagaraDataInterfaceStaticMesh::VMRandomFilteredSocket);
+		OutFunc.BindUObject(this, &UNiagaraDataInterfaceStaticMesh::VMRandomFilteredSocket);
 	}
 	else if (BindingInfo.Name == RandomUnfilteredSocketName)
 	{
-		OutFunc = FVMExternalFunction::CreateUObject(this, &UNiagaraDataInterfaceStaticMesh::VMRandomUnfilteredSocket);
+		OutFunc.BindUObject(this, &UNiagaraDataInterfaceStaticMesh::VMRandomUnfilteredSocket);
 	}
 	else if (BindingInfo.Name == GetSocketTransformName)
 	{
-		OutFunc = FVMExternalFunction::CreateLambda([this](FVectorVMExternalFunctionContext& Context) { VMGetSocketTransform<FNDITransformHandlerNoop>(Context); });
+		OutFunc.BindUObject(this, &UNiagaraDataInterfaceStaticMesh::VMGetSocketTransform<FNDITransformHandlerNoop>);
 	}
 	else if (BindingInfo.Name == GetSocketTransformWSName)
 	{
-		OutFunc = FVMExternalFunction::CreateLambda([this](FVectorVMExternalFunctionContext& Context) { VMGetSocketTransform<FNDITransformHandler>(Context); });
+		OutFunc.BindUObject(this, &UNiagaraDataInterfaceStaticMesh::VMGetSocketTransform<FNDITransformHandler>);
 	}
 	else if (BindingInfo.Name == GetSocketTransformWSInterpolatedName)
 	{
-		OutFunc = FVMExternalFunction::CreateLambda([this](FVectorVMExternalFunctionContext& Context) { VMGetSocketTransformInterpolated<FNDITransformHandler>(Context); });
+		OutFunc.BindUObject(this, &UNiagaraDataInterfaceStaticMesh::VMGetSocketTransformInterpolated<FNDITransformHandler>);
 	}
 	else if (BindingInfo.Name == GetFilteredSocketTransformName)
 	{
-		OutFunc = FVMExternalFunction::CreateLambda([this](FVectorVMExternalFunctionContext& Context) { VMGetFilteredSocketTransform<FNDITransformHandlerNoop>(Context); });
+		OutFunc.BindUObject(this, &UNiagaraDataInterfaceStaticMesh::VMGetFilteredSocketTransform<FNDITransformHandlerNoop>);
 	}
 	else if (BindingInfo.Name == GetFilteredSocketTransformWSName)
 	{
-		OutFunc = FVMExternalFunction::CreateLambda([this](FVectorVMExternalFunctionContext& Context) { VMGetFilteredSocketTransform<FNDITransformHandler>(Context); });
+		OutFunc.BindUObject(this, &UNiagaraDataInterfaceStaticMesh::VMGetFilteredSocketTransform<FNDITransformHandler>);
 	}
 	else if (BindingInfo.Name == GetFilteredSocketTransformWSInterpolatedName)
 	{
-		OutFunc = FVMExternalFunction::CreateLambda([this](FVectorVMExternalFunctionContext& Context) { VMGetFilteredSocketTransformInterpolated<FNDITransformHandler>(Context); });
+		OutFunc.BindUObject(this, &UNiagaraDataInterfaceStaticMesh::VMGetFilteredSocketTransformInterpolated<FNDITransformHandler>);
 	}
 	else if (BindingInfo.Name == GetUnfilteredSocketTransformName)
 	{
-		OutFunc = FVMExternalFunction::CreateLambda([this](FVectorVMExternalFunctionContext& Context) { VMGetUnfilteredSocketTransform<FNDITransformHandlerNoop>(Context); });
+		OutFunc.BindUObject(this, &UNiagaraDataInterfaceStaticMesh::VMGetUnfilteredSocketTransform<FNDITransformHandlerNoop>);
 	}
 	else if (BindingInfo.Name == GetUnfilteredSocketTransformWSName)
 	{
-		OutFunc = FVMExternalFunction::CreateLambda([this](FVectorVMExternalFunctionContext& Context) { VMGetUnfilteredSocketTransform<FNDITransformHandler>(Context); });
+		OutFunc.BindUObject(this, &UNiagaraDataInterfaceStaticMesh::VMGetUnfilteredSocketTransform<FNDITransformHandler>);
 	}
 	else if (BindingInfo.Name == GetUnfilteredSocketTransformWSInterpolatedName)
 	{
-		OutFunc = FVMExternalFunction::CreateLambda([this](FVectorVMExternalFunctionContext& Context) { VMGetUnfilteredSocketTransformInterpolated<FNDITransformHandler>(Context); });
+		OutFunc.BindUObject(this, &UNiagaraDataInterfaceStaticMesh::VMGetUnfilteredSocketTransformInterpolated<FNDITransformHandler>);
 	}
 	else if (BindingInfo.Name == GetFilteredSocketName)
 	{
-		OutFunc = FVMExternalFunction::CreateUObject(this, &UNiagaraDataInterfaceStaticMesh::VMGetFilteredSocket);
+		OutFunc.BindUObject(this, &UNiagaraDataInterfaceStaticMesh::VMGetFilteredSocket);
 	}
 	else if (BindingInfo.Name == GetUnfilteredSocketName)
 	{
-		OutFunc = FVMExternalFunction::CreateUObject(this, &UNiagaraDataInterfaceStaticMesh::VMGetUnfilteredSocket);
+		OutFunc.BindUObject(this, &UNiagaraDataInterfaceStaticMesh::VMGetUnfilteredSocket);
 	}
 
 	//////////////////////////////////////////////////////////////////////////
 	// Section functions
 	else if (BindingInfo.Name == IsValidSectionName)
 	{
-		OutFunc = FVMExternalFunction::CreateLambda([this](FVectorVMExternalFunctionContext& Context) { VMIsValidSection(Context); });
+		OutFunc.BindUObject(this, &UNiagaraDataInterfaceStaticMesh::VMIsValidSection);
 	}
 	else if (BindingInfo.Name == GetSectionTriangleCountName)
 	{
-		OutFunc = FVMExternalFunction::CreateLambda([this](FVectorVMExternalFunctionContext& Context) { VMGetSectionTriangleCount(Context); });
+		OutFunc.BindUObject(this, &UNiagaraDataInterfaceStaticMesh::VMGetSectionTriangleCount);
 	}
 	else if (BindingInfo.Name == RandomSectionTriangleName)
 	{
-		OutFunc = FVMExternalFunction::CreateLambda([this](FVectorVMExternalFunctionContext& Context) { VMRandomSectionTriangle<FNDIRandomHelper>(Context); });
+		OutFunc.BindUObject(this, &UNiagaraDataInterfaceStaticMesh::VMRandomSectionTriangle<FNDIRandomHelper>);
 	}
 	else if (BindingInfo.Name == GetSectionTriangleAtName)
 	{
-		OutFunc = FVMExternalFunction::CreateLambda([this](FVectorVMExternalFunctionContext& Context) { VMGetSectionTriangleAt(Context); });
+		OutFunc.BindUObject(this, &UNiagaraDataInterfaceStaticMesh::VMGetSectionTriangleAt);
 	}
 	else if (BindingInfo.Name == GetFilteredSectionAtName)
 	{
-		OutFunc = FVMExternalFunction::CreateLambda([this](FVectorVMExternalFunctionContext& Context) { VMGetFilteredSectionAt(Context); });
+		OutFunc.BindUObject(this, &UNiagaraDataInterfaceStaticMesh::VMGetFilteredSectionAt);
 	}
 	else if (BindingInfo.Name == GetUnfilteredSectionAtName)
 	{
-		OutFunc = FVMExternalFunction::CreateLambda([this](FVectorVMExternalFunctionContext& Context) { VMGetUnfilteredSectionAt(Context); });
+		OutFunc.BindUObject(this, &UNiagaraDataInterfaceStaticMesh::VMGetUnfilteredSectionAt);
 	}
 	else if (BindingInfo.Name == GetSectionCountName)
 	{
-		OutFunc = FVMExternalFunction::CreateLambda([this](FVectorVMExternalFunctionContext& Context) { VMGetSectionCount(Context); });
+		OutFunc.BindUObject(this, &UNiagaraDataInterfaceStaticMesh::VMGetSectionCount);
 	}
 	else if (BindingInfo.Name == GetFilteredSectionCountName)
 	{
-		OutFunc = FVMExternalFunction::CreateLambda([this](FVectorVMExternalFunctionContext& Context) { VMGetFilteredSectionCount(Context); });
+		OutFunc.BindUObject(this, &UNiagaraDataInterfaceStaticMesh::VMGetFilteredSectionCount);
 	}
 	else if (BindingInfo.Name == GetUnfilteredSectionCountName)
 	{
-		OutFunc = FVMExternalFunction::CreateLambda([this](FVectorVMExternalFunctionContext& Context) { VMGetUnfilteredSectionCount(Context); });
+		OutFunc.BindUObject(this, &UNiagaraDataInterfaceStaticMesh::VMGetUnfilteredSectionCount);
 	}
 	else if (BindingInfo.Name == RandomSectionName)
 	{
-		OutFunc = FVMExternalFunction::CreateLambda([this](FVectorVMExternalFunctionContext& Context) { VMRandomSection<FNDIRandomHelper>(Context); });
+		OutFunc.BindUObject(this, &UNiagaraDataInterfaceStaticMesh::VMRandomSection<FNDIRandomHelper>);
 	}
 	else if (BindingInfo.Name == RandomFilteredSectionName)
 	{
-		OutFunc = FVMExternalFunction::CreateLambda([this](FVectorVMExternalFunctionContext& Context) { VMRandomFilteredSection<FNDIRandomHelper>(Context); });
+		OutFunc.BindUObject(this, &UNiagaraDataInterfaceStaticMesh::VMRandomFilteredSection<FNDIRandomHelper>);
 	}
 	else if (BindingInfo.Name == RandomUnfilteredSectionName)
 	{
-		OutFunc = FVMExternalFunction::CreateLambda([this](FVectorVMExternalFunctionContext& Context) { VMRandomUnfilteredSection(Context); });
+		OutFunc.BindUObject(this, &UNiagaraDataInterfaceStaticMesh::VMRandomUnfilteredSection);
 	}
 
 	//////////////////////////////////////////////////////////////////////////
 	// VM Misc Functions
 	else if (BindingInfo.Name == IsValidName)
 	{
-		OutFunc = FVMExternalFunction::CreateLambda([this](FVectorVMExternalFunctionContext& Context) { VMIsValid(Context); });
+		OutFunc.BindUObject(this, &UNiagaraDataInterfaceStaticMesh::VMIsValid);
 	}
 	else if (BindingInfo.Name == GetPreSkinnedLocalBoundsName)
 	{
-		OutFunc = FVMExternalFunction::CreateUObject(this, &UNiagaraDataInterfaceStaticMesh::VMGetPreSkinnedLocalBounds);
+		OutFunc.BindUObject(this, &UNiagaraDataInterfaceStaticMesh::VMGetPreSkinnedLocalBounds);
 	}
 	else if (BindingInfo.Name == GetMeshBoundsName)
 	{
-		OutFunc = FVMExternalFunction::CreateUObject(this, &UNiagaraDataInterfaceStaticMesh::VMGetMeshBounds<true>);
+		OutFunc.BindUObject(this, &UNiagaraDataInterfaceStaticMesh::VMGetMeshBounds<true>);
 	}
 	else if (BindingInfo.Name == GetMeshBoundsWSName)
 	{
-		OutFunc = FVMExternalFunction::CreateUObject(this, &UNiagaraDataInterfaceStaticMesh::VMGetMeshBounds<false>);
+		OutFunc.BindUObject(this, &UNiagaraDataInterfaceStaticMesh::VMGetMeshBounds<false>);
 	}
 	else if (BindingInfo.Name == GetLocalToWorldName)
 	{
-		OutFunc = FVMExternalFunction::CreateLambda([this](FVectorVMExternalFunctionContext& Context) { VMGetLocalToWorld(Context); });
+		OutFunc.BindUObject(this, &UNiagaraDataInterfaceStaticMesh::VMGetLocalToWorld);
 	}
 	else if (BindingInfo.Name == GetLocalToWorldInverseTransposedName)
 	{
-		OutFunc = FVMExternalFunction::CreateLambda([this](FVectorVMExternalFunctionContext& Context) { VMGetLocalToWorldInverseTransposed(Context); });
+		OutFunc.BindUObject(this, &UNiagaraDataInterfaceStaticMesh::VMGetLocalToWorldInverseTransposed);
 	}
 	else if (BindingInfo.Name == GetWorldVelocityName)
 	{
-		OutFunc = FVMExternalFunction::CreateLambda([this](FVectorVMExternalFunctionContext& Context) { VMGetWorldVelocity(Context); });
+		OutFunc.BindUObject(this, &UNiagaraDataInterfaceStaticMesh::VMGetWorldVelocity);
 	}
 	else if (BindingInfo.Name == GetInstanceIndexName)
 	{
-		OutFunc = FVMExternalFunction::CreateLambda([this](FVectorVMExternalFunctionContext& Context) { VMGetInstanceIndex(Context); });
+		OutFunc.BindUObject(this, &UNiagaraDataInterfaceStaticMesh::VMGetInstanceIndex);
 	}
 	else if (BindingInfo.Name == SetInstanceIndexName)
 	{
-		OutFunc = FVMExternalFunction::CreateLambda([this](FVectorVMExternalFunctionContext& Context) { VMSetInstanceIndex(Context); });
+		OutFunc.BindUObject(this, &UNiagaraDataInterfaceStaticMesh::VMSetInstanceIndex);
 	}
 
 	//////////////////////////////////////////////////////////////////////////
 	// VM UV mapping functions
 	else if (BindingInfo.Name == GetTriangleCoordAtUVName)
 	{
-		OutFunc = FVMExternalFunction::CreateLambda([this](FVectorVMExternalFunctionContext& Context) { VMGetTriangleCoordAtUV(Context); });
+		OutFunc.BindUObject(this, &UNiagaraDataInterfaceStaticMesh::VMGetTriangleCoordAtUV);
 	}
 	else if (BindingInfo.Name == GetTriangleCoordInAabbName)
 	{
-		OutFunc = FVMExternalFunction::CreateLambda([this](FVectorVMExternalFunctionContext& Context) { VMGetTriangleCoordInAabb(Context); });
+		OutFunc.BindUObject(this, &UNiagaraDataInterfaceStaticMesh::VMGetTriangleCoordInAabb);
 	}
 	else if (BindingInfo.Name == BuildUvMappingName)
 	{
-		OutFunc = FVMExternalFunction::CreateLambda([this](FVectorVMExternalFunctionContext& Context) { VMBuildUvMapping(Context); });
+		OutFunc.BindUObject(this, &UNiagaraDataInterfaceStaticMesh::VMBuildUvMapping);
 	}
 
 	//////////////////////////////////////////////////////////////////////////
 	// Deprecated Functions
 	else if (BindingInfo.Name == Deprecated_GetTriColorName)
 	{
-		OutFunc = FVMExternalFunction::CreateLambda([this](FVectorVMExternalFunctionContext& Context) { VMGetTriangleColor(Context); });
+		OutFunc.BindUObject(this, &UNiagaraDataInterfaceStaticMesh::VMGetTriangleColor);
 	}
 	else if (BindingInfo.Name == Deprecated_GetTriUVName)
 	{
-		OutFunc = FVMExternalFunction::CreateLambda([this](FVectorVMExternalFunctionContext& Context) { VMGetTriangleUV(Context); });
+		OutFunc.BindUObject(this, &UNiagaraDataInterfaceStaticMesh::VMGetTriangleUV);
 	}
 	// Custom VM bindings
 	else if (BindingInfo.Name == Deprecated_RandomSectionName)
 	{
-		OutFunc = FVMExternalFunction::CreateLambda([this](FVectorVMExternalFunctionContext& Context) { VMRandomFilteredSection<FNDIRandomHelperFromStream>(Context); });
+		OutFunc.BindUObject(this, &UNiagaraDataInterfaceStaticMesh::VMRandomFilteredSection<FNDIRandomHelperFromStream>);
 	}
 	else if (BindingInfo.Name == Deprecated_RandomTriCoordName)
 	{
-		OutFunc = FVMExternalFunction::CreateLambda([this](FVectorVMExternalFunctionContext& Context) { VMRandomTriangle<FNDIRandomHelperFromStream>(Context); });
+		OutFunc.BindUObject(this, &UNiagaraDataInterfaceStaticMesh::VMRandomTriangle<FNDIRandomHelperFromStream>);
 	}
 	else if (BindingInfo.Name == Deprecated_RandomTriCoordOnSectionName)
 	{
-		OutFunc = FVMExternalFunction::CreateLambda([this](FVectorVMExternalFunctionContext& Context) { VMRandomSectionTriangle<FNDIRandomHelperFromStream>(Context); });
+		OutFunc.BindUObject(this, &UNiagaraDataInterfaceStaticMesh::VMRandomSectionTriangle<FNDIRandomHelperFromStream>);
 	}
 	else if (BindingInfo.Name == Deprecated_GetVertexPositionName)
 	{
-		OutFunc = FVMExternalFunction::CreateLambda([this](FVectorVMExternalFunctionContext& Context) { VMGetVertexPosition_Deprecated<FNDITransformHandlerNoop>(Context); });
+		OutFunc.BindUObject(this, &UNiagaraDataInterfaceStaticMesh::VMGetVertexPosition_Deprecated<FNDITransformHandlerNoop>);
 	}
 	else if (BindingInfo.Name == Deprecated_GetVertexPositionWSName)
 	{
-		OutFunc = FVMExternalFunction::CreateLambda([this](FVectorVMExternalFunctionContext& Context) { VMGetVertexPosition_Deprecated<FNDITransformHandler>(Context); });
+		OutFunc.BindUObject(this, &UNiagaraDataInterfaceStaticMesh::VMGetVertexPosition_Deprecated<FNDITransformHandler>);
 	}
 	else if (BindingInfo.Name == Deprecated_GetTriPositionName)
 	{
-		OutFunc = FVMExternalFunction::CreateLambda([this](FVectorVMExternalFunctionContext& Context) { VMGetTriPosition_Deprecated<FNDITransformHandlerNoop>(Context); });
+		OutFunc.BindUObject(this, &UNiagaraDataInterfaceStaticMesh::VMGetTriPosition_Deprecated<FNDITransformHandlerNoop>);
 	}
 	else if (BindingInfo.Name == Deprecated_GetTriPositionWSName)
 	{
-		OutFunc = FVMExternalFunction::CreateLambda([this](FVectorVMExternalFunctionContext& Context) { VMGetTriPosition_Deprecated<FNDITransformHandler>(Context); });
+		OutFunc.BindUObject(this, &UNiagaraDataInterfaceStaticMesh::VMGetTriPosition_Deprecated<FNDITransformHandler>);
 	}
 	else if (BindingInfo.Name == Deprecated_GetTriTangentsName)
 	{	
-		OutFunc = FVMExternalFunction::CreateLambda([this](FVectorVMExternalFunctionContext& Context) { VMGetTriangleTangentBasis_Deprecated<FNDITransformHandlerNoop>(Context); });
+		OutFunc.BindUObject(this, &UNiagaraDataInterfaceStaticMesh::VMGetTriangleTangentBasis_Deprecated<FNDITransformHandlerNoop>);
 	}
 	else if (BindingInfo.Name == Deprecated_GetTriTangentsWSName)
 	{
-		OutFunc = FVMExternalFunction::CreateLambda([this](FVectorVMExternalFunctionContext& Context) { VMGetTriangleTangentBasis_Deprecated<FNDITransformHandler>(Context); });
+		OutFunc.BindUObject(this, &UNiagaraDataInterfaceStaticMesh::VMGetTriangleTangentBasis_Deprecated<FNDITransformHandler>);
 	}
 	else if (BindingInfo.Name == Deprecated_GetTriNormalName)
 	{
-		OutFunc = FVMExternalFunction::CreateLambda([this](FVectorVMExternalFunctionContext& Context) { VMGetTriangleNormal_Deprecated<FNDITransformHandlerNoop>(Context); });
+		OutFunc.BindUObject(this, &UNiagaraDataInterfaceStaticMesh::VMGetTriangleNormal_Deprecated<FNDITransformHandlerNoop>);
 	}
 	else if (BindingInfo.Name == Deprecated_GetTriNormalWSName)
 	{
-		OutFunc = FVMExternalFunction::CreateLambda([this](FVectorVMExternalFunctionContext& Context) { VMGetTriangleNormal_Deprecated<FNDITransformHandler>(Context); });
+		OutFunc.BindUObject(this, &UNiagaraDataInterfaceStaticMesh::VMGetTriangleNormal_Deprecated<FNDITransformHandler>);
 	}
 	else if (BindingInfo.Name == Deprecated_GetTriPositionAndVelocityWSName)
 	{
-	OutFunc = FVMExternalFunction::CreateLambda([this](FVectorVMExternalFunctionContext& Context) { VMGetTriPositionAndVelocity_Deprecated<FNDITransformHandler>(Context); });
+		OutFunc.BindUObject(this, &UNiagaraDataInterfaceStaticMesh::VMGetTriPositionAndVelocity_Deprecated<FNDITransformHandler>);
 	}
 }
 
