- Fix static analysis issues

#rnx

diff --git a/Engine/Plugins/FX/Niagara/Source/Niagara/Private/NiagaraRendererRibbons.cpp b/Engine/Plugins/FX/Niagara/Source/Niagara/Private/NiagaraRendererRibbons.cpp
--- a/Engine/Plugins/FX/Niagara/Source/Niagara/Private/NiagaraRendererRibbons.cpp
+++ b/Engine/Plugins/FX/Niagara/Source/Niagara/Private/NiagaraRendererRibbons.cpp
@@ -1430,11 +1430,11 @@ void FNiagaraRendererRibbons::InitializeShape(const UNiagaraRibbonRendererProper
 	}
 	else if (Properties->Shape == ENiagaraRibbonShapeMode::MultiPlane && Properties->MultiPlaneCount > 1 && Properties->MultiPlaneCount <= 16)
 	{
-		GenerateShapeStateMultiPlane(ShapeState, Properties->MultiPlaneCount, Properties->WidthSegmentationCount, Properties->bEnableAccurateGeometry, Properties->bUseMaterialBackfaceCulling);
+		GenerateShapeStateMultiPlane(ShapeState, Properties->MultiPlaneCount, Properties->GetWidthSegmentationCount(), Properties->bEnableAccurateGeometry, Properties->bUseMaterialBackfaceCulling);
 	}
 	else
 	{
-		GenerateShapeStatePlane(ShapeState, Properties->WidthSegmentationCount, Properties->bUseMaterialBackfaceCulling);
+		GenerateShapeStatePlane(ShapeState, Properties->GetWidthSegmentationCount(), Properties->bUseMaterialBackfaceCulling);
 	}	
 }
 
diff --git a/Engine/Plugins/FX/Niagara/Source/Niagara/Public/NiagaraRibbonRendererProperties.h b/Engine/Plugins/FX/Niagara/Source/Niagara/Public/NiagaraRibbonRendererProperties.h
--- a/Engine/Plugins/FX/Niagara/Source/Niagara/Public/NiagaraRibbonRendererProperties.h
+++ b/Engine/Plugins/FX/Niagara/Source/Niagara/Public/NiagaraRibbonRendererProperties.h
@@ -246,6 +246,8 @@ public:
 	TObjectPtr<UMaterialInstanceConstant> MICMaterial;
 #endif
 
+	int32 GetWidthSegmentationCount() const { return FMath::Clamp(WidthSegmentationCount, 1, 16); }
+
 	/** Use the UMaterialInterface bound to this user variable if it is set to a valid value. If this is bound to a valid value and Material is also set, UserParamBinding wins.*/
 	UPROPERTY(EditAnywhere, Category = "Ribbon Rendering")
 	FNiagaraUserParameterBinding MaterialUserParamBinding;
diff --git a/Engine/Plugins/FX/Niagara/Source/NiagaraEditor/Private/NiagaraBakerRendererOutputSparseVolumeTexture.cpp b/Engine/Plugins/FX/Niagara/Source/NiagaraEditor/Private/NiagaraBakerRendererOutputSparseVolumeTexture.cpp
--- a/Engine/Plugins/FX/Niagara/Source/NiagaraEditor/Private/NiagaraBakerRendererOutputSparseVolumeTexture.cpp
+++ b/Engine/Plugins/FX/Niagara/Source/NiagaraEditor/Private/NiagaraBakerRendererOutputSparseVolumeTexture.cpp
@@ -419,17 +419,17 @@ void FNiagaraBakerRendererOutputSparseVolumeTexture::EndBake(FNiagaraBakerFeedba
 
 		if (StartFrame < 0)
 		{
-			UE_LOG(LogNiagaraBaker, Error, TEXT("Start frame must be greater than 0"));
+			UE_LOG(LogNiagaraBaker, Error, TEXT("Start frame must be greater than or equal to 0"));
 			return;
 		}
 
-		if (BlendFrames < 0)
+		if (BlendFrames <= 0)
 		{
 			UE_LOG(LogNiagaraBaker, Error, TEXT("Blend amount must be greater than 0"));
 			return;
 		}
 
-		if (LoopedFrames < 0)
+		if (LoopedFrames <= 0)
 		{
 			UE_LOG(LogNiagaraBaker, Error, TEXT("Cannot loop SVT due to insufficient frames.  Please reduce blend amount and start time or bake more frames."));
 			return;
