Niagara - Fix a crash when upgrading modules (and potentially other copy/paste operations) due to some unsafe iteration of stack function input entries while pasting.

diff --git a/Engine/Plugins/FX/Niagara/Source/NiagaraEditor/Private/ViewModels/Stack/NiagaraStackScriptHierarchyRoot.cpp b/Engine/Plugins/FX/Niagara/Source/NiagaraEditor/Private/ViewModels/Stack/NiagaraStackScriptHierarchyRoot.cpp
--- a/Engine/Plugins/FX/Niagara/Source/NiagaraEditor/Private/ViewModels/Stack/NiagaraStackScriptHierarchyRoot.cpp
+++ b/Engine/Plugins/FX/Niagara/Source/NiagaraEditor/Private/ViewModels/Stack/NiagaraStackScriptHierarchyRoot.cpp
@@ -430,14 +430,32 @@ void UNiagaraStackScriptHierarchyRoot::ToClipboardFunctionInputs(UObject* InOute
 
 void UNiagaraStackScriptHierarchyRoot::SetValuesFromClipboardFunctionInputs(const TArray<const UNiagaraClipboardFunctionInput*>& ClipboardFunctionInputs)
 {
+	// Copy/paste of function inputs must be restricted to inputs owned by a single function but we need to collect inputs recursively since
+	// the hierarchy view allows inputs within a single function to be "child inputs" of each other.  When collecting recursively by type we
+	// can end up with inputs from dynamic input functions which aren't owned by this root, so we need to filter those out to avoid setting 
+	// another functions input values and to prevent crashes which can happen when setting an input on this function removes the dynamic input
+	// which will finalize the dynamic input's child inputs.
+	auto GetChildInputsWithMatchingOwner = [this](TArray<UNiagaraStackFunctionInput*>& OutMatchingInputs)
+	{
+		TArray<UNiagaraStackFunctionInput*> AllInputs;
+		GetUnfilteredChildrenOfType(AllInputs, true);
+		for (UNiagaraStackFunctionInput* Input : AllInputs)
+		{
+			if (&Input->GetInputFunctionCallNode() == OwningFunctionCallNode.Get())
+			{
+				OutMatchingInputs.Add(Input);
+			}
+		}
+	};
+	TArray<UNiagaraStackFunctionInput*> StackFunctionInputs;
+	GetChildInputsWithMatchingOwner(StackFunctionInputs);
+
 	//First try to set each input as a static switch, and if a switch is set refresh the categories
 	// before applying additional inputs.  This is necessary because static switches can change the set
 	// of exposed inputs.
 	// NOTE: It's still possible that inputs could end up missing in cases where the switch dependencies are 
 	// especially complex since we're not doing a full switch dependency check, but this should handle the
 	// vast majority of cases.
-	TArray<UNiagaraStackFunctionInput*> StackFunctionInputs;
-	GetUnfilteredChildrenOfType(StackFunctionInputs, true);
 	for (const UNiagaraClipboardFunctionInput* ClipboardFunctionInput : ClipboardFunctionInputs)
 	{
 		bool bInputSetAsSwitch = false;
@@ -465,7 +483,7 @@ void UNiagaraStackScriptHierarchyRoot::SetValuesFromClipboardFunctionInputs(cons
 		{
 			RefreshChildren();
 			StackFunctionInputs.Empty();
-			GetUnfilteredChildrenOfType(StackFunctionInputs, true);
+			GetChildInputsWithMatchingOwner(StackFunctionInputs);
 		}
 	}
 
