- Remove OnDeferredSyncAllNameMatchParametersHandle from INiagaraParameterDefinitionsSubscriber to avoid leaving potentially bound delegates to invalid memory

#tests ET / PIE / Cook

#FailedValidations Preflight

diff --git a/Engine/Plugins/FX/Niagara/Source/Niagara/Classes/NiagaraParameterDefinitionsSubscriber.h b/Engine/Plugins/FX/Niagara/Source/Niagara/Classes/NiagaraParameterDefinitionsSubscriber.h
--- a/Engine/Plugins/FX/Niagara/Source/Niagara/Classes/NiagaraParameterDefinitionsSubscriber.h
+++ b/Engine/Plugins/FX/Niagara/Source/Niagara/Classes/NiagaraParameterDefinitionsSubscriber.h
@@ -43,10 +43,14 @@ class INiagaraParameterDefinitionsSubscriber
 {
 #if WITH_EDITORONLY_DATA
 public:
-	NIAGARA_API virtual ~INiagaraParameterDefinitionsSubscriber();
+	using FDeferredSyncFunction = TFunction<TArray<INiagaraParameterDefinitionsSubscriber*>(UObject*)>;
 
-	NIAGARA_API void PostLoadDefinitionsSubscriptions();
-	NIAGARA_API void CleanupDefinitionsSubscriptions();
+	virtual ~INiagaraParameterDefinitionsSubscriber() = default;
+
+	// Post load definitions
+	// If the asset registry has not finished loading files the DeferredSync will be called when it has done so with the first argument being resolved Owner
+	// The callback is expected to return all of the subsribers that WeakOwner contains
+	NIAGARA_API void PostLoadDefinitionsSubscriptions(TWeakObjectPtr<UObject> WeakOwner, FDeferredSyncFunction DeferredSync);
 
 	//~ Begin Pure Virtual Methods
 	virtual const TArray<FParameterDefinitionsSubscription>& GetParameterDefinitionsSubscriptions() const = 0;
@@ -80,8 +84,6 @@ public:
 private:
 	FOnSubscribedParameterDefinitionsChanged OnSubscribedParameterDefinitionsChangedDelegate;
 
-	FDelegateHandle OnDeferredSyncAllNameMatchParametersHandle;
-
 private:
 	/** Update the cached change Id of specified Synchronized Subscribed Parameter Definition subscriptions so that they are not marked pending sync. */
 	NIAGARA_API void MarkParameterDefinitionSubscriptionsSynchronized(TArray<FGuid> SynchronizedParameterDefinitionsIds /*= TArray<FGuid>()*/);
diff --git a/Engine/Plugins/FX/Niagara/Source/Niagara/Private/NiagaraEmitter.cpp b/Engine/Plugins/FX/Niagara/Source/Niagara/Private/NiagaraEmitter.cpp
--- a/Engine/Plugins/FX/Niagara/Source/Niagara/Private/NiagaraEmitter.cpp
+++ b/Engine/Plugins/FX/Niagara/Source/Niagara/Private/NiagaraEmitter.cpp
@@ -2025,7 +2025,18 @@ void UNiagaraEmitter::UpdateEmitterAfterLoad()
 	}
 
 	// Synchronize with definitions before merging.
-	PostLoadDefinitionsSubscriptions();
+	PostLoadDefinitionsSubscriptions(
+		MakeWeakObjectPtr(this),
+		[](UObject* ObjectOwner) -> TArray<INiagaraParameterDefinitionsSubscriber*>
+		{
+			TArray<INiagaraParameterDefinitionsSubscriber*> OutArray;
+			if (UNiagaraEmitter* Owner = Cast<UNiagaraEmitter>(ObjectOwner))
+			{
+				OutArray.Add(Owner);
+			}
+			return OutArray;
+		}
+	);
 
 	// If we're owned by another emitter then we're a VersionedParentAtLastMerge
 	if (GetOuter()->IsA<UNiagaraEmitter>())
@@ -3280,7 +3291,6 @@ void UNiagaraEmitter::BeginDestroy()
 			Data.GPUComputeScript->OnGPUScriptCompiled().RemoveAll(this);
 		}
 	}
-	CleanupDefinitionsSubscriptions();
 #endif
 	Super::BeginDestroy();
 }
diff --git a/Engine/Plugins/FX/Niagara/Source/Niagara/Private/NiagaraParameterDefinitionsSubscriber.cpp b/Engine/Plugins/FX/Niagara/Source/Niagara/Private/NiagaraParameterDefinitionsSubscriber.cpp
--- a/Engine/Plugins/FX/Niagara/Source/Niagara/Private/NiagaraParameterDefinitionsSubscriber.cpp
+++ b/Engine/Plugins/FX/Niagara/Source/Niagara/Private/NiagaraParameterDefinitionsSubscriber.cpp
@@ -15,12 +15,56 @@
 
 #if WITH_EDITORONLY_DATA
 
-INiagaraParameterDefinitionsSubscriber::~INiagaraParameterDefinitionsSubscriber()
+namespace NiagaraParameterDefinitionsSubscriberPrivate
 {
-	check(OnDeferredSyncAllNameMatchParametersHandle.IsValid() == false);
+	FDelegateHandle OnAssetRegistryFilesLoaded;
+	TArray<TPair<TWeakObjectPtr<UObject>, INiagaraParameterDefinitionsSubscriber::FDeferredSyncFunction>> DeferredSyncFunctions;
+
+	void OnFilesLoaded()
+	{
+		FSynchronizeWithParameterDefinitionsArgs Args;
+		Args.bForceGatherDefinitions = true;
+		Args.bSubscribeAllNameMatchParameters = true;
+
+		for (const auto& DeferredSync : DeferredSyncFunctions)
+		{
+			if (UObject* Owner = DeferredSync.Key.Get())
+			{
+				for (INiagaraParameterDefinitionsSubscriber* Subscriber : DeferredSync.Value(Owner))
+				{
+					Subscriber->SynchronizeWithParameterDefinitions(Args);
+				}
+			}
+		}
+		DeferredSyncFunctions.Empty();
+
+		FAssetRegistryModule& AssetRegistryModule = FModuleManager::LoadModuleChecked<FAssetRegistryModule>("AssetRegistry");
+		AssetRegistryModule.Get().OnFilesLoaded().Remove(OnAssetRegistryFilesLoaded);
+		OnAssetRegistryFilesLoaded.Reset();
+	}
+
+	void AddDeferredSync(TWeakObjectPtr<UObject> WeakOwner, INiagaraParameterDefinitionsSubscriber::FDeferredSyncFunction Function)
+	{
+		check(IsInGameThread());
+		if (OnAssetRegistryFilesLoaded.IsValid() == false)
+		{
+			FAssetRegistryModule& AssetRegistryModule = FModuleManager::LoadModuleChecked<FAssetRegistryModule>("AssetRegistry");
+			OnAssetRegistryFilesLoaded = AssetRegistryModule.Get().OnFilesLoaded().AddStatic(&OnFilesLoaded);
+		}
+
+		for (const auto& DeferredSync : DeferredSyncFunctions)
+		{
+			if (DeferredSync.Key == WeakOwner)
+			{
+				return;
+			}
+		}
+
+		DeferredSyncFunctions.Emplace(WeakOwner, Function);
+	}
 }
 
-void INiagaraParameterDefinitionsSubscriber::PostLoadDefinitionsSubscriptions()
+void INiagaraParameterDefinitionsSubscriber::PostLoadDefinitionsSubscriptions(TWeakObjectPtr<UObject> WeakOwner, TFunction<TArray<INiagaraParameterDefinitionsSubscriber*>(UObject*)> DeferredSync)
 {
 	// First remove any subscriptions that are pointing to null parameter definitions assets. 
 	// The null entries are either definitions assets that have been deleted, or definitions assets in a content plugin that is not mounted.
@@ -54,8 +98,8 @@ void INiagaraParameterDefinitionsSubscriber::PostLoadDefinitionsSubscriptions()
 	if (bAllParameterDefinitionsDiscovered == false)
 	{
 		SynchronizeWithParameterDefinitions();
-		check(OnDeferredSyncAllNameMatchParametersHandle.IsValid() == false);
-		OnDeferredSyncAllNameMatchParametersHandle = AssetRegistryModule.Get().OnFilesLoaded().AddLambda(SyncAllNameMatchParameters);
+
+		NiagaraParameterDefinitionsSubscriberPrivate::AddDeferredSync(WeakOwner, DeferredSync);
 	}
 	// Else the asset registry has discovered all parameter definitions assets;
 	// Process linking name matched parameters and synchronizing existing definitions now.
@@ -65,22 +109,6 @@ void INiagaraParameterDefinitionsSubscriber::PostLoadDefinitionsSubscriptions()
 	}
 }
 
-void INiagaraParameterDefinitionsSubscriber::CleanupDefinitionsSubscriptions()
-{
-	if(OnDeferredSyncAllNameMatchParametersHandle.IsValid())
-	{
-		if(FAssetRegistryModule* AssetRegistryModule = FModuleManager::GetModulePtr<FAssetRegistryModule>("AssetRegistry"))
-		{
-			IAssetRegistry* AssetRegistry = AssetRegistryModule->TryGet();
-			if (AssetRegistry)
-			{
-				AssetRegistry->OnFilesLoaded().Remove(OnDeferredSyncAllNameMatchParametersHandle);
-				OnDeferredSyncAllNameMatchParametersHandle.Reset();
-			}
-		}
-	}
-}
-
 TArray<UNiagaraParameterDefinitionsBase*> INiagaraParameterDefinitionsSubscriber::GetSubscribedParameterDefinitions() const
 {
 	const TArray<FParameterDefinitionsSubscription>& Subscriptions = GetParameterDefinitionsSubscriptions();
diff --git a/Engine/Plugins/FX/Niagara/Source/Niagara/Private/NiagaraScript.cpp b/Engine/Plugins/FX/Niagara/Source/Niagara/Private/NiagaraScript.cpp
--- a/Engine/Plugins/FX/Niagara/Source/Niagara/Private/NiagaraScript.cpp
+++ b/Engine/Plugins/FX/Niagara/Source/Niagara/Private/NiagaraScript.cpp
@@ -2203,78 +2203,94 @@ void UNiagaraScript::PostLoad()
 
 	const bool bRequestHierarchyRootMigration = bMigrateParameterDataToHierarchyRoot;
 	bMigrateParameterDataToHierarchyRoot = false;
-	
+
 	VersionedScriptAdapters.Reserve(VersionData.Num());
 	for (FVersionedNiagaraScriptData& Data : VersionData)
 	{		
 		UNiagaraScriptSourceBase* Source = Data.Source;
-		if (Source != nullptr)
+		if (Source == nullptr)
 		{
-			Source->ConditionalPostLoad();
+			continue;
+		}
 
-			if (bRequestHierarchyRootMigration)
-			{
-				Source->MigrateParameterDataToHierarchyRoot(Data);
-			}
+		Source->ConditionalPostLoad();
 
-			// Synchronize with Definitions after source scripts have been postloaded.
-			FVersionedNiagaraScript& VersionedScriptAdapter = VersionedScriptAdapters.Emplace_GetRef(this, Data.Version.VersionGuid);
-			VersionedScriptAdapter.PostLoadDefinitionsSubscriptions();
+		if (bRequestHierarchyRootMigration)
+		{
+			Source->MigrateParameterDataToHierarchyRoot(Data);
+		}
 
-			bool bScriptVMNeedsRebuild = false;
-			FString RebuildReason;
-			if (NiagaraVer < FNiagaraCustomVersion::UseHashesToIdentifyCompileStateOfTopLevelScripts && CachedScriptVMId.CompilerVersionID.IsValid())
+		// Synchronize with Definitions after source scripts have been postloaded.
+		FVersionedNiagaraScript& VersionedScriptAdapter = VersionedScriptAdapters.Emplace_GetRef(this, Data.Version.VersionGuid);
+		VersionedScriptAdapter.PostLoadDefinitionsSubscriptions(
+			MakeWeakObjectPtr(this),
+			[](UObject* ObjectOwner) -> TArray<INiagaraParameterDefinitionsSubscriber*>
 			{
-				FGuid BaseId = Source->GetCompileBaseId(Usage, UsageId);
-				if (BaseId.IsValid() == false)
-				{
-					UE_LOG(LogNiagara, Warning,
-						TEXT("Invalidating compile ids for script %s because it doesn't have a valid base id.  The owning asset will continue to compile on load until it is resaved."),
-						*GetPathName());
-					InvalidateCompileResults(TEXT("Script didn't have a valid base id."));
-					Source->ForceGraphToRecompileOnNextCheck();
-				}
-				else
+				TArray<INiagaraParameterDefinitionsSubscriber*> OutArray;
+				if (UNiagaraScript* Owner = Cast<UNiagaraScript>(ObjectOwner))
 				{
-					FNiagaraCompileHash CompileHash = Source->GetCompileHash(Usage, UsageId);
-					if (CompileHash.IsValid())
+					OutArray.Reserve(Owner->VersionedScriptAdapters.Num());
+					for (FVersionedNiagaraScript& Adapter : Owner->VersionedScriptAdapters)
 					{
-						CachedScriptVMId.BaseScriptCompileHash = CompileHash;
-					}
-					else
-					{
-						// If the compile hash isn't valid, the vm id needs to be recalculated and the cached vm needs to be invalidated.
-						bScriptVMNeedsRebuild = true;
-						RebuildReason = TEXT("Script did not have a valid compile hash.");
+						OutArray.Add(&Adapter);
 					}
 				}
+				return OutArray;
 			}
+		);
 
-			if (CachedScriptVMId.CompilerVersionID.IsValid() && CachedScriptVMId.CompilerVersionID != FNiagaraCustomVersion::GetLatestScriptCompileVersion())
+		bool bScriptVMNeedsRebuild = false;
+		FString RebuildReason;
+		if (NiagaraVer < FNiagaraCustomVersion::UseHashesToIdentifyCompileStateOfTopLevelScripts && CachedScriptVMId.CompilerVersionID.IsValid())
+		{
+			FGuid BaseId = Source->GetCompileBaseId(Usage, UsageId);
+			if (BaseId.IsValid() == false)
 			{
-				bScriptVMNeedsRebuild = true;
-				RebuildReason = TEXT("Niagara compiler version changed since the last time the script was compiled.");
+				UE_LOG(LogNiagara, Warning,
+					TEXT("Invalidating compile ids for script %s because it doesn't have a valid base id.  The owning asset will continue to compile on load until it is resaved."),
+					*GetPathName());
+				InvalidateCompileResults(TEXT("Script didn't have a valid base id."));
+				Source->ForceGraphToRecompileOnNextCheck();
 			}
-
-			if (bScriptVMNeedsRebuild)
+			else
 			{
-        			// We modify here in post load so that it will cause the owning asset to resave when running the resave commandlet.
-				Modify();
-				InvalidateCompileResults(RebuildReason);
-
-				// Older versions of graphs require the script to drive the CachedCompileIds being fully refreshed when the script version changes.
-				// In order to bridge the gap without requiring a new package version we trigger the refresh here
-				Source->RefreshGraphCompileId();
+				FNiagaraCompileHash CompileHash = Source->GetCompileHash(Usage, UsageId);
+				if (CompileHash.IsValid())
+				{
+					CachedScriptVMId.BaseScriptCompileHash = CompileHash;
+				}
+				else
+				{
+					// If the compile hash isn't valid, the vm id needs to be recalculated and the cached vm needs to be invalidated.
+					bScriptVMNeedsRebuild = true;
+					RebuildReason = TEXT("Script did not have a valid compile hash.");
+				}
 			}
+		}
 
-			// Convert visibility of old assets
-			if (NiagaraVer < FNiagaraCustomVersion::AddLibraryAssetProperty || (NiagaraVer < FNiagaraCustomVersion::AddLibraryVisibilityProperty && bExposeToLibrary_DEPRECATED))
-			{
-				ScriptData->LibraryVisibility = ENiagaraScriptLibraryVisibility::Library;
-			}
+		if (CachedScriptVMId.CompilerVersionID.IsValid() && CachedScriptVMId.CompilerVersionID != FNiagaraCustomVersion::GetLatestScriptCompileVersion())
+		{
+			bScriptVMNeedsRebuild = true;
+			RebuildReason = TEXT("Niagara compiler version changed since the last time the script was compiled.");
 		}
-	}
 
+		if (bScriptVMNeedsRebuild)
+		{
+        		// We modify here in post load so that it will cause the owning asset to resave when running the resave commandlet.
+			Modify();
+			InvalidateCompileResults(RebuildReason);
+
+			// Older versions of graphs require the script to drive the CachedCompileIds being fully refreshed when the script version changes.
+			// In order to bridge the gap without requiring a new package version we trigger the refresh here
+			Source->RefreshGraphCompileId();
+		}
+
+		// Convert visibility of old assets
+		if (NiagaraVer < FNiagaraCustomVersion::AddLibraryAssetProperty || (NiagaraVer < FNiagaraCustomVersion::AddLibraryVisibilityProperty && bExposeToLibrary_DEPRECATED))
+		{
+			ScriptData->LibraryVisibility = ENiagaraScriptLibraryVisibility::Library;
+		}
+	}
 #endif
 
 	ProcessSerializedShaderMaps();
@@ -3447,13 +3463,6 @@ void UNiagaraScript::BeginDestroy()
 			// if there was nothing to release, then we don't need to wait for anything
 			ReleasedByRT = true;
 		}
-
-#if WITH_EDITORONLY_DATA
-		for (FVersionedNiagaraScript& VersionedScript : VersionedScriptAdapters)
-		{
-			VersionedScript.CleanupDefinitionsSubscriptions();
-		}
-#endif
 	}
 	else
 	{
diff --git a/Engine/Plugins/FX/Niagara/Source/Niagara/Private/NiagaraSystem.cpp b/Engine/Plugins/FX/Niagara/Source/Niagara/Private/NiagaraSystem.cpp
--- a/Engine/Plugins/FX/Niagara/Source/Niagara/Private/NiagaraSystem.cpp
+++ b/Engine/Plugins/FX/Niagara/Source/Niagara/Private/NiagaraSystem.cpp
@@ -348,10 +348,6 @@ void UNiagaraSystem::BeginDestroy()
 	KillAllActiveCompilations();
 #endif
 
-#if WITH_EDITOR
-	CleanupDefinitionsSubscriptions();
-#endif
-
 	FNiagaraWorldManager::DestroyAllSystemSimulations(this);
 
 	// Ensure we wait for any potential render commands to finish executing before we GC the NiagaraSystem
@@ -729,7 +725,18 @@ void UNiagaraSystem::UpdateSystemAfterLoad()
 		AllSystemScripts.Add(SystemUpdateScript);
 
 		// Synchronize with parameter definitions
-		PostLoadDefinitionsSubscriptions();
+		PostLoadDefinitionsSubscriptions(
+			MakeWeakObjectPtr(this),
+			[](UObject* ObjectOwner) -> TArray<INiagaraParameterDefinitionsSubscriber*>
+			{
+				TArray<INiagaraParameterDefinitionsSubscriber*> OutArray;
+				if (UNiagaraSystem* Owner = Cast<UNiagaraSystem>(ObjectOwner))
+				{
+					OutArray.Add(Owner);
+				}
+				return OutArray;
+			}
+		);
 
 		if (UNiagaraEmitter::GetForceCompileOnLoad())
 		{
