- Add icon to show if material bindings have active bindings or not

#tests Add material bindings to a sprite renderer, remove them, observe icon

diff --git a/Engine/Plugins/FX/Niagara/Source/NiagaraEditor/Private/NiagaraEditorStyle.cpp b/Engine/Plugins/FX/Niagara/Source/NiagaraEditor/Private/NiagaraEditorStyle.cpp
--- a/Engine/Plugins/FX/Niagara/Source/NiagaraEditor/Private/NiagaraEditorStyle.cpp
+++ b/Engine/Plugins/FX/Niagara/Source/NiagaraEditor/Private/NiagaraEditorStyle.cpp
@@ -449,6 +449,7 @@ void FNiagaraEditorStyle::InitOverview()
 	// Overview debug icons
 	Set("NiagaraEditor.Overview.DebugActive", new IMAGE_BRUSH("Icons/OverviewDebugActive", Icon16x16));
 	Set("NiagaraEditor.Overview.DebugInactive", new IMAGE_BRUSH("Icons/OverviewDebugInactive", Icon16x16));
+	Set("NiagaraEditor.Overview.HasMaterialParameterBindings", new IMAGE_BRUSH("Icons/HasMaterialParameterBindings", Icon16x16));
 }
 
 void FNiagaraEditorStyle::InitViewportStyle()
diff --git a/Engine/Plugins/FX/Niagara/Source/NiagaraEditor/Private/ViewModels/Stack/NiagaraStackRendererItem.cpp b/Engine/Plugins/FX/Niagara/Source/NiagaraEditor/Private/ViewModels/Stack/NiagaraStackRendererItem.cpp
--- a/Engine/Plugins/FX/Niagara/Source/NiagaraEditor/Private/ViewModels/Stack/NiagaraStackRendererItem.cpp
+++ b/Engine/Plugins/FX/Niagara/Source/NiagaraEditor/Private/ViewModels/Stack/NiagaraStackRendererItem.cpp
@@ -27,6 +27,7 @@
 #include "ViewModels/NiagaraParameterPanelViewModel.h"
 #include "ViewModels/NiagaraEmitterHandleViewModel.h"
 #include "ViewModels/Stack/NiagaraStackViewModel.h"
+#include "NiagaraEditorStyle.h"
 
 #include UE_INLINE_GENERATED_CPP_BY_NAME(NiagaraStackRendererItem)
 
@@ -717,5 +718,55 @@ void UNiagaraStackRendererItem::SetDebugDrawEnabled(bool bInEnabled)
 	}
 }
 
+TArray<FNiagaraStackItemCustomIcon> UNiagaraStackRendererItem::GetCustomIcons() const
+{
+	TArray<FNiagaraStackItemCustomIcon> CustomIcons;
+
+	UNiagaraRendererProperties* Renderer = RendererProperties.Get();
+	if (Renderer == nullptr)
+	{
+		return CustomIcons;
+	}
+
+	// Gather all our material parameter bindings
+	TArray<FNiagaraRendererMaterialParameters*> MaterialParameters;
+	for (TFieldIterator<FStructProperty> PropIt(Renderer->GetClass()); PropIt; ++PropIt)
+	{
+		FStructProperty* StructProp = *PropIt;
+		if (!StructProp || !StructProp->Struct || !StructProp->Struct->IsChildOf(FNiagaraRendererMaterialParameters::StaticStruct()))
+		{
+			continue;
+		}
+
+		MaterialParameters.Add(StructProp->ContainerPtrToValuePtr<FNiagaraRendererMaterialParameters>(Renderer));
+	}
+
+	if (MaterialParameters.Num() > 0)
+	{
+		FNiagaraStackItemCustomIcon& Icon = CustomIcons.AddDefaulted_GetRef();
+
+		Icon.IconBrush			= FNiagaraEditorStyle::Get().GetBrush(TEXT("NiagaraEditor.Overview.HasMaterialParameterBindings"));
+		Icon.ToolTipText		= LOCTEXT("BoundMaterialParametersTooltip", "Renderer has bound material parameters.  See the bindings section on the renderer for more details.");
+		Icon.VisibleFunction	= TAttribute<EVisibility>::CreateLambda(
+			[MaterialParameters, WeakRenderer=RendererProperties]()
+			{
+				if (WeakRenderer.Get())
+				{
+					for (FNiagaraRendererMaterialParameters* Parameters : MaterialParameters)
+					{
+						if (Parameters->HasAnyBindings())
+						{
+							return EVisibility::Visible;
+						}
+					}
+				}
+				return EVisibility::Hidden;
+			}
+		);
+	}
+
+	return CustomIcons;
+}
+
 #undef LOCTEXT_NAMESPACE
 
diff --git a/Engine/Plugins/FX/Niagara/Source/NiagaraEditor/Public/ViewModels/Stack/NiagaraStackItem.h b/Engine/Plugins/FX/Niagara/Source/NiagaraEditor/Public/ViewModels/Stack/NiagaraStackItem.h
--- a/Engine/Plugins/FX/Niagara/Source/NiagaraEditor/Public/ViewModels/Stack/NiagaraStackItem.h
+++ b/Engine/Plugins/FX/Niagara/Source/NiagaraEditor/Public/ViewModels/Stack/NiagaraStackItem.h
@@ -36,6 +36,13 @@ public:
 	virtual void NotifyEnumValueChanged(int32 bInValue) = 0;
 };
 
+struct FNiagaraStackItemCustomIcon
+{
+	const FSlateBrush*		IconBrush = nullptr;
+	TAttribute<FText>		ToolTipText;
+	TAttribute<EVisibility>	VisibleFunction;
+};
+
 UCLASS(MinimalAPI)
 class UNiagaraStackItem : public UNiagaraStackEntry
 {
@@ -81,6 +88,8 @@ public:
 	virtual bool IsDebugDrawEnabled() const { return false; }
 	virtual void SetDebugDrawEnabled(bool bInEnabled) {}
 
+	virtual TArray<FNiagaraStackItemCustomIcon> GetCustomIcons() const { return {}; }
+
 protected:
 	NIAGARAEDITOR_API virtual void RefreshChildrenInternal(const TArray<UNiagaraStackEntry*>& CurrentChildren, TArray<UNiagaraStackEntry*>& NewChildren, TArray<FStackIssue>& NewIssues) override;
 
diff --git a/Engine/Plugins/FX/Niagara/Source/NiagaraEditor/Public/ViewModels/Stack/NiagaraStackRendererItem.h b/Engine/Plugins/FX/Niagara/Source/NiagaraEditor/Public/ViewModels/Stack/NiagaraStackRendererItem.h
--- a/Engine/Plugins/FX/Niagara/Source/NiagaraEditor/Public/ViewModels/Stack/NiagaraStackRendererItem.h
+++ b/Engine/Plugins/FX/Niagara/Source/NiagaraEditor/Public/ViewModels/Stack/NiagaraStackRendererItem.h
@@ -102,6 +102,8 @@ private:
 	virtual bool IsDebugDrawEnabled() const override;
 	virtual void SetDebugDrawEnabled(bool bInEnabled) override;
 
+	virtual TArray<FNiagaraStackItemCustomIcon> GetCustomIcons() const override;
+
 private:
 	TSharedPtr<INiagaraStackRenderersOwner> RenderersOwner;
 
diff --git a/Engine/Plugins/FX/Niagara/Source/NiagaraEditorWidgets/Private/SNiagaraOverviewStack.cpp b/Engine/Plugins/FX/Niagara/Source/NiagaraEditorWidgets/Private/SNiagaraOverviewStack.cpp
--- a/Engine/Plugins/FX/Niagara/Source/NiagaraEditorWidgets/Private/SNiagaraOverviewStack.cpp
+++ b/Engine/Plugins/FX/Niagara/Source/NiagaraEditorWidgets/Private/SNiagaraOverviewStack.cpp
@@ -881,11 +881,11 @@ void SNiagaraOverviewStack::RefreshEntryList()
 
 		UNiagaraStackEntry* RootEntry = StackViewModel->GetRootEntry();
 		checkf(RootEntry != nullptr, TEXT("Root entry was null."));
+
 		TArray<UNiagaraStackEntry*> RootChildren;		
 		RootEntry->GetFilteredChildren(RootChildren);
 		for (UNiagaraStackEntry* RootChild : RootChildren)
 		{
-			checkf(RootEntry != nullptr, TEXT("Root entry child was null."));
 			TArray<UNiagaraStackEntry*> ParentChain;
 			AddEntriesRecursive(*RootChild, FlattenedEntryList, AllowedClasses, ParentChain);
 		}
@@ -967,6 +967,7 @@ TSharedRef<ITableRow> SNiagaraOverviewStack::OnGenerateRowForEntry(UNiagaraStack
 				]
 			];
 		}
+
 		// Icon Text
 		if (StackItem->GetSupportedIconMode() == UNiagaraStackEntry::EIconMode::Text)
 		{
@@ -989,6 +990,7 @@ TSharedRef<ITableRow> SNiagaraOverviewStack::OnGenerateRowForEntry(UNiagaraStack
 				]
 			];
 		}
+
 		// Name
 		ContentBox->AddSlot()
 		.VAlign(VAlign_Center)
@@ -997,6 +999,7 @@ TSharedRef<ITableRow> SNiagaraOverviewStack::OnGenerateRowForEntry(UNiagaraStack
 		[
 			SNew(SNiagaraSystemOverviewItemName, StackItem)
 		];
+
 		// Issue Icon
 		ContentBox->AddSlot()
 		.VAlign(VAlign_Center)
@@ -1008,7 +1011,6 @@ TSharedRef<ITableRow> SNiagaraOverviewStack::OnGenerateRowForEntry(UNiagaraStack
 			.Visibility_Static(&SNiagaraOverviewStack::GetIssueIconVisibility, Item)
 		];
 
-
 		TSharedRef<SHorizontalBox> OptionsBox = SNew(SHorizontalBox);
 		
 		if(StackItem->SupportsStackNotes())
@@ -1035,8 +1037,8 @@ TSharedRef<ITableRow> SNiagaraOverviewStack::OnGenerateRowForEntry(UNiagaraStack
 		}
 		
 
-		UNiagaraStackModuleItem* StackModuleItem = Cast<UNiagaraStackModuleItem>(StackItem);
 		// Scratch icon
+		UNiagaraStackModuleItem* StackModuleItem = Cast<UNiagaraStackModuleItem>(StackItem);
 		if (StackModuleItem && StackModuleItem->IsScratchModule())
 		{			
 			OptionsBox->AddSlot()
@@ -1052,6 +1054,28 @@ TSharedRef<ITableRow> SNiagaraOverviewStack::OnGenerateRowForEntry(UNiagaraStack
 				];
 		}
 
+		TArray<FNiagaraStackItemCustomIcon> CustomIcons = StackItem->GetCustomIcons();
+		for (const FNiagaraStackItemCustomIcon& CustomIcon : CustomIcons)
+		{
+			if (CustomIcon.IconBrush == nullptr)
+			{
+				continue;
+			}
+
+			OptionsBox->AddSlot()
+			.VAlign(VAlign_Center)
+			.Padding(3, 0)
+			[
+				SNew(SBorder)
+				.Visibility(CustomIcon.VisibleFunction)
+				.ToolTipText(CustomIcon.ToolTipText)
+				[
+					SNew(SImage)
+					.Image(CustomIcon.IconBrush)
+				]
+			];
+		}
+
 		// Debug draw 
 		if (StackItem->SupportsDebugDraw())
 		{
