- Adjust saturation of read icon to make it more clear in the stack

diff --git a/Engine/Plugins/FX/Niagara/Source/NiagaraEditorWidgets/Private/NiagaraEditorWidgetsStyle.cpp b/Engine/Plugins/FX/Niagara/Source/NiagaraEditorWidgets/Private/NiagaraEditorWidgetsStyle.cpp
--- a/Engine/Plugins/FX/Niagara/Source/NiagaraEditorWidgets/Private/NiagaraEditorWidgetsStyle.cpp
+++ b/Engine/Plugins/FX/Niagara/Source/NiagaraEditorWidgets/Private/NiagaraEditorWidgetsStyle.cpp
@@ -587,8 +587,8 @@ void FNiagaraEditorWidgetsStyle::InitIcons()
 	Set("Niagara.CompactStackIssue.Info", new IMAGE_BRUSH("Icons/StackIssueInfo", Icon8x8));
 	Set("Niagara.CompactStackIssue.Message", new IMAGE_BRUSH("Icons/StackIssueMessage", Icon8x8));
 
-	Set("NiagaraEditor.Stack.ReadIcon", new IMAGE_BRUSH_SVG("Icons/Read_Icon", Icon16x16));
-	Set("NiagaraEditor.Stack.ReadWriteIcon", new IMAGE_BRUSH_SVG("Icons/Read_Write_Icon", Icon16x16));
+	Set("NiagaraEditor.Stack.ReadIcon", new IMAGE_BRUSH("Icons/Read_Icon", Icon16x16));
+	Set("NiagaraEditor.Stack.ReadWriteIcon", new IMAGE_BRUSH("Icons/Read_Write_Icon", Icon16x16));
 	Set("NiagaraEditor.Stack.WriteIcon", new IMAGE_BRUSH_SVG("Icons/Write_Icon", Icon16x16));
 	Set("NiagaraEditor.Stack.GoToSourceIcon", new CORE_IMAGE_BRUSH("Common/GoToSource", Icon30x30, FLinearColor::White));
 	
