- Make GPU usage rule messaging more clear

diff --git a/Engine/Plugins/FX/Niagara/Source/NiagaraEditor/Private/NiagaraValidationRules.cpp b/Engine/Plugins/FX/Niagara/Source/NiagaraEditor/Private/NiagaraValidationRules.cpp
--- a/Engine/Plugins/FX/Niagara/Source/NiagaraEditor/Private/NiagaraValidationRules.cpp
+++ b/Engine/Plugins/FX/Niagara/Source/NiagaraEditor/Private/NiagaraValidationRules.cpp
@@ -971,8 +971,8 @@ void UNiagaraValidationRule_GpuUsage::CheckValidity(const FNiagaraValidationCont
 
 		FNiagaraValidationResult& ValidationResult = OutResults.Emplace_GetRef(
 			Severity,
-			LOCTEXT("GpuUsageInfo", "GPU usage may not function as expected"),
-			FText::Format(LOCTEXT("GpuUsageInfoDetails", "GPU usage may not function as expected on '{0}'."), FText::FromString(PlatformConflictsString)),
+			LOCTEXT("GpuUsageInfo", "GPU compute may not be available."),
+			FText::Format(LOCTEXT("GpuUsageInfoDetails", "GPU compute may not be enabled on platforms '{0}' consider making a fallback CPU emitter."), FText::FromString(PlatformConflictsString)),
 			NiagaraValidation::GetStackEntry<UNiagaraStackEmitterPropertiesItem>(EmitterHandleModel->GetEmitterStackViewModel())
 		);
 		
