- Only reset Emitter namespace bindings on renderers when copying to LWE

#tests Make a renderer using emitter / system / NPC bindings, copy to LWE emitter and verify emitter binding only is reset

diff --git a/Engine/Plugins/FX/Niagara/Source/Niagara/Private/Stateless/NiagaraStatelessEmitter.cpp b/Engine/Plugins/FX/Niagara/Source/Niagara/Private/Stateless/NiagaraStatelessEmitter.cpp
--- a/Engine/Plugins/FX/Niagara/Source/Niagara/Private/Stateless/NiagaraStatelessEmitter.cpp
+++ b/Engine/Plugins/FX/Niagara/Source/Niagara/Private/Stateless/NiagaraStatelessEmitter.cpp
@@ -745,9 +745,7 @@ void UNiagaraStatelessEmitter::AddRenderer(UNiagaraRendererProperties* Renderer,
 	Renderer->OuterEmitterVersion = EmitterVersion;
 	RendererProperties.Add(Renderer);
 #if WITH_EDITOR
-	// When pasting a renderer from stateful they can come with bindings which are not supported for stateless
-	// temporarily we reset them to default values.  We will need to upgrade all these calls to take in an adapter to be able to
-	// interop between different emitter types, or introduce a base emitter type.
+	// LWE does not support Emitter namespace so reset any bindings that point to emitter parameters
 	if (UNiagaraRendererProperties* RendererCDO = Renderer->GetClass()->GetDefaultObject<UNiagaraRendererProperties>())
 	{
 		for (TFieldIterator<FStructProperty> PropIt(Renderer->GetClass()); PropIt; ++PropIt)
@@ -758,8 +756,11 @@ void UNiagaraStatelessEmitter::AddRenderer(UNiagaraRendererProperties* Renderer,
 				continue;
 			}
 			FNiagaraVariableAttributeBinding* ParameterBinding = StructProp->ContainerPtrToValuePtr<FNiagaraVariableAttributeBinding>(Renderer);
-			FNiagaraVariableAttributeBinding* ParameterBindingDefault = StructProp->ContainerPtrToValuePtr<FNiagaraVariableAttributeBinding>(RendererCDO);
-			*ParameterBinding = *ParameterBindingDefault;
+			if ( ParameterBinding->GetBindingSourceMode() == ENiagaraBindingSource::ExplicitEmitter )
+			{
+				FNiagaraVariableAttributeBinding* ParameterBindingDefault = StructProp->ContainerPtrToValuePtr<FNiagaraVariableAttributeBinding>(RendererCDO);
+				*ParameterBinding = *ParameterBindingDefault;
+			}
 		}
 	}
 
