- Add input type selection for cone velocity

#virtualized

diff --git a/Engine/Plugins/FX/Niagara/Source/Niagara/Internal/Stateless/Modules/NiagaraStatelessModuleCommon.h b/Engine/Plugins/FX/Niagara/Source/Niagara/Internal/Stateless/Modules/NiagaraStatelessModuleCommon.h
--- a/Engine/Plugins/FX/Niagara/Source/Niagara/Internal/Stateless/Modules/NiagaraStatelessModuleCommon.h
+++ b/Engine/Plugins/FX/Niagara/Source/Niagara/Internal/Stateless/Modules/NiagaraStatelessModuleCommon.h
@@ -28,7 +28,9 @@ namespace NiagaraStateless
 
 		bool							bConeVelocity = false;
 		ENiagaraCoordinateSpace			ConeCoordinateSpace = ENiagaraCoordinateSpace::Local;
-		FQuat4f							ConeQuat = FQuat4f::Identity;
+		bool							bUseConeConeRotator = true;
+		FNiagaraStatelessRangeRotator	ConeRotator = FNiagaraStatelessRangeRotator(FRotator3f::ZeroRotator);
+		FNiagaraStatelessRangeVector3	ConeDirection = FNiagaraStatelessRangeVector3(FVector3f::XAxisVector);
 		FNiagaraStatelessRangeFloat		ConeVelocityRange = FNiagaraStatelessRangeFloat(0.0f);
 		float							ConeOuterAngle = 0.0f;
 		float							ConeInnerAngle = 0.0f;
@@ -45,4 +47,22 @@ namespace NiagaraStateless
 		float							NoiseFrequency = 0.0f;
 		FNiagaraStatelessRangeVector3	NoiseFieldOffset = FNiagaraStatelessRangeVector3(FVector3f::ZeroVector);
 	};
+
+	inline FQuat4f DirectionToQuat(FVector3f Direction)
+	{
+		Direction = Direction.GetSafeNormal(UE_KINDA_SMALL_NUMBER, FVector3f::ZAxisVector);
+
+		const FVector3f Forward = FVector3f::ZAxisVector;
+		const float Dot = FVector3f::DotProduct(Forward, Direction);
+		if (Dot < -0.99999f)
+		{
+			return FQuat4f(FVector3f::YAxisVector, PI);
+		}
+
+		const FVector3f Axis = FVector3f::CrossProduct(Forward, Direction);
+		const float Angle = FMath::Acos(Dot);
+		FQuat4f Quat(Axis, Angle);
+		Quat.Normalize();
+		return Quat;
+	}
 }
diff --git a/Engine/Plugins/FX/Niagara/Source/Niagara/Internal/Stateless/Modules/NiagaraStatelessModule_AddVelocity.h b/Engine/Plugins/FX/Niagara/Source/Niagara/Internal/Stateless/Modules/NiagaraStatelessModule_AddVelocity.h
--- a/Engine/Plugins/FX/Niagara/Source/Niagara/Internal/Stateless/Modules/NiagaraStatelessModule_AddVelocity.h
+++ b/Engine/Plugins/FX/Niagara/Source/Niagara/Internal/Stateless/Modules/NiagaraStatelessModule_AddVelocity.h
@@ -10,13 +10,20 @@
 #include "NiagaraStatelessModule_AddVelocity.generated.h"
 
 UENUM()
-enum class ENSM_VelocityType
+enum class ENSM_VelocityType : uint8
 {
 	Linear,			// Velocity Min / Max -> Velocity Speed Scale
 	FromPoint,		// Min / Max / Origin
 	InCone,			// Max, Min, Cone| Cone Axis, Cone Angle, Inner Cone
 };
 
+UENUM()
+enum class ENSM_ConeRotationType : uint8
+{
+	Rotation,
+	Direction,
+};
+
 // Add to the velocity applied to all particles
 UCLASS(MinimalAPI, EditInlineNew, meta = (DisplayName = "Add Velocity"))
 class UNiagaraStatelessModule_AddVelocity : public UNiagaraStatelessModule
@@ -34,8 +41,13 @@ public:
 
 	UPROPERTY(EditAnywhere, Category = "Parameters", meta = (DisplayName = "Cone Velocity", EditConditionHides, EditCondition = "VelocityType == ENSM_VelocityType::InCone"))
 	FNiagaraDistributionRangeFloat ConeVelocityDistribution = FNiagaraDistributionRangeFloat(500.0f);
-	UPROPERTY(EditAnywhere, Category = "Parameters", meta = (EditConditionHides, EditCondition = "VelocityType == ENSM_VelocityType::InCone"))
-	FRotator ConeRotation = FRotator::ZeroRotator;
+	UPROPERTY(EditAnywhere, Category = "Parameters", meta = (SegmentedDisplay, EditConditionHides, EditCondition = "VelocityType == ENSM_VelocityType::InCone"))
+	ENSM_ConeRotationType ConeRotationType = ENSM_ConeRotationType::Rotation;
+	UPROPERTY(EditAnywhere, Category = "Parameters", meta = (EditConditionHides, EditCondition = "VelocityType == ENSM_VelocityType::InCone && ConeRotationType == ENSM_ConeRotationType::Rotation", DisableRangeDistribution, DisableUniformDistribution))
+	FNiagaraDistributionRangeRotator ConeRotation = FNiagaraDistributionRangeRotator(FRotator3f::ZeroRotator);
+	UPROPERTY(EditAnywhere, Category = "Parameters", meta = (EditConditionHides, EditCondition = "VelocityType == ENSM_VelocityType::InCone && ConeRotationType == ENSM_ConeRotationType::Direction", DisableRangeDistribution, DisableUniformDistribution))
+	FNiagaraDistributionRangeVector3 ConeDirection = FNiagaraDistributionRangeVector3(FVector3f::ZAxisVector);
+
 	UPROPERTY(EditAnywhere, Category = "Parameters", meta = (EditConditionHides, EditCondition = "VelocityType == ENSM_VelocityType::InCone", ClampMin = "0.0", ClampMax = "360.0", UIMin = "0.0", UIMax = "360.0"))
 	float ConeAngle = 45.0f;
 	UPROPERTY(EditAnywhere, Category = "Parameters", meta = (EditConditionHides, EditCondition = "VelocityType == ENSM_VelocityType::InCone", ClampMin = "0.0", ClampMax = "360.0", UIMin = "0.0", UIMax = "360.0"))
diff --git a/Engine/Plugins/FX/Niagara/Source/Niagara/Private/Stateless/Modules/NiagaraStatelessModule_AddVelocity.cpp b/Engine/Plugins/FX/Niagara/Source/Niagara/Private/Stateless/Modules/NiagaraStatelessModule_AddVelocity.cpp
--- a/Engine/Plugins/FX/Niagara/Source/Niagara/Private/Stateless/Modules/NiagaraStatelessModule_AddVelocity.cpp
+++ b/Engine/Plugins/FX/Niagara/Source/Niagara/Private/Stateless/Modules/NiagaraStatelessModule_AddVelocity.cpp
@@ -35,7 +35,9 @@ void UNiagaraStatelessModule_AddVelocity::BuildEmitterData(const FNiagaraStatele
 
 		PhysicsBuildData.bConeVelocity = true;
 		PhysicsBuildData.ConeCoordinateSpace = CoordinateSpace;
-		PhysicsBuildData.ConeQuat = FQuat4f(ConeRotation.Quaternion());
+		PhysicsBuildData.bUseConeConeRotator = ConeRotationType == ENSM_ConeRotationType::Rotation;
+		PhysicsBuildData.ConeRotator = ConeRotationType == ENSM_ConeRotationType::Rotation ? BuildContext.ConvertDistributionToRange(ConeRotation, FRotator3f::ZeroRotator) : FNiagaraStatelessRangeRotator(FRotator3f::ZeroRotator);
+		PhysicsBuildData.ConeDirection = ConeRotationType == ENSM_ConeRotationType::Direction ? BuildContext.ConvertDistributionToRange(ConeDirection, FVector3f::ZAxisVector) : FNiagaraStatelessRangeVector3(FVector3f::ZAxisVector);
 		PhysicsBuildData.ConeVelocityRange = BuildContext.ConvertDistributionToRange(ConeVelocityDistribution, 0.0f);
 		PhysicsBuildData.ConeOuterAngle = ConeAngle;
 		PhysicsBuildData.ConeInnerAngle = InnerCone;
@@ -67,7 +69,26 @@ void UNiagaraStatelessModule_AddVelocity::DrawDebug(const FNiagaraStatelessDrawD
 
 		case ENSM_VelocityType::InCone:
 		{
-			const FQuat ConeQuat = DrawDebugContext.TransformRotation(FQuat4f(ConeRotation.Quaternion()));
+			FQuat4f LocalConeQuat;
+			if (ConeRotationType == ENSM_ConeRotationType::Rotation)
+			{
+				if (ConeRotation.IsBinding())
+				{
+					return;
+				}
+
+				LocalConeQuat = ConeRotation.CalculateRange().Min.Quaternion();
+			}
+			else
+			{
+				if (ConeDirection.IsBinding())
+				{
+					return;
+				}
+				LocalConeQuat = NiagaraStateless::DirectionToQuat(ConeDirection.CalculateRange().Min);
+			}
+
+			const FQuat ConeQuat = DrawDebugContext.TransformRotation(LocalConeQuat);
 			const float ConeHAngle = ConeAngle / 2.0f;
 
 			TOptional<float> InnerConeHAngle;
diff --git a/Engine/Plugins/FX/Niagara/Source/Niagara/Private/Stateless/Modules/NiagaraStatelessModule_SolveVelocitiesAndForces.cpp b/Engine/Plugins/FX/Niagara/Source/Niagara/Private/Stateless/Modules/NiagaraStatelessModule_SolveVelocitiesAndForces.cpp
--- a/Engine/Plugins/FX/Niagara/Source/Niagara/Private/Stateless/Modules/NiagaraStatelessModule_SolveVelocitiesAndForces.cpp
+++ b/Engine/Plugins/FX/Niagara/Source/Niagara/Private/Stateless/Modules/NiagaraStatelessModule_SolveVelocitiesAndForces.cpp
@@ -119,7 +119,7 @@ namespace NiagaraStatelessModuleSolveVelocitiesAndForcesPrivate
 					VelocityScale *= FMath::Lerp(1.0f, pf, Parameters->SolveVelocitiesAndForces_ConeVelocityFalloff);
 				}
 
-				InitialVelocity += ModuleBuiltData->PhysicsData.ConeQuat.RotateVector(Direction) * VelocityScale;
+				InitialVelocity += Parameters->SolveVelocitiesAndForces_ConeQuat.RotateVector(Direction) * VelocityScale;
 			}
 
 			if (Parameters->SolveVelocitiesAndForces_PointVelocityEnabled != 0)
@@ -216,7 +216,18 @@ void UNiagaraStatelessModule_SolveVelocitiesAndForces::SetShaderParameters(const
 
 	Parameters->SolveVelocitiesAndForces_ConeVelocityEnabled = PhysicsData.bConeVelocity ? 1 : 0;
 	{
-		Parameters->SolveVelocitiesAndForces_ConeQuat				= SpaceTransforms.TransformRotation(PhysicsData.ConeCoordinateSpace, PhysicsData.ConeQuat);
+		FQuat4f LocalConeQuat;
+		if (PhysicsData.bUseConeConeRotator)
+		{
+			LocalConeQuat = SetShaderParameterContext.ConvertRangeToValue(PhysicsData.ConeRotator).Quaternion();
+		}
+		else
+		{
+			const FVector3f ConeDirection = SetShaderParameterContext.ConvertRangeToValue(PhysicsData.ConeDirection);
+			LocalConeQuat = NiagaraStateless::DirectionToQuat(ConeDirection);
+		}
+
+		Parameters->SolveVelocitiesAndForces_ConeQuat				= SpaceTransforms.TransformRotation(PhysicsData.ConeCoordinateSpace, LocalConeQuat);
 		SetShaderParameterContext.ConvertRangeToScaleBias(PhysicsData.ConeVelocityRange, Parameters->SolveVelocitiesAndForces_ConeVelocityScale, Parameters->SolveVelocitiesAndForces_ConeVelocityBias);
 		Parameters->SolveVelocitiesAndForces_ConeAngleScale			= (PhysicsData.ConeOuterAngle - PhysicsData.ConeInnerAngle) * (UE_PI / 360.0f);
 		Parameters->SolveVelocitiesAndForces_ConeAngleBias			= PhysicsData.ConeInnerAngle * (UE_PI / 360.0f);
diff --git a/Engine/Plugins/FX/Niagara/Source/Niagara/Private/Stateless/NiagaraStatelessDistribution.cpp b/Engine/Plugins/FX/Niagara/Source/Niagara/Private/Stateless/NiagaraStatelessDistribution.cpp
--- a/Engine/Plugins/FX/Niagara/Source/Niagara/Private/Stateless/NiagaraStatelessDistribution.cpp
+++ b/Engine/Plugins/FX/Niagara/Source/Niagara/Private/Stateless/NiagaraStatelessDistribution.cpp
@@ -1329,6 +1329,14 @@ bool FNiagaraDistributionRangeRotator::SerializeFromMismatchedTag(const struct F
 		return true;
 	}
 
+	if (Tag.GetType().IsStruct(NAME_Rotator))
+	{
+		FRotator Value;
+		Slot << Value;
+		*this = FNiagaraDistributionRangeRotator(FRotator3f(Value));
+		return true;
+	}
+
 	return false;
 }
 
