[Niagara Validation] Ignore lightweight (stateless) emitters when requiring parent emitters

diff --git a/Engine/Plugins/FX/Niagara/Source/NiagaraEditor/Private/NiagaraValidationRules.cpp b/Engine/Plugins/FX/Niagara/Source/NiagaraEditor/Private/NiagaraValidationRules.cpp
--- a/Engine/Plugins/FX/Niagara/Source/NiagaraEditor/Private/NiagaraValidationRules.cpp
+++ b/Engine/Plugins/FX/Niagara/Source/NiagaraEditor/Private/NiagaraValidationRules.cpp
@@ -1852,6 +1852,15 @@ void UNiagaraValidationRule_RequireParentEmitter::CheckValidity(const FNiagaraVa
 	{
 		VersionedEmitter.Emitter->ForEachVersionData([&, this](const FVersionedNiagaraEmitterData& VersionedEmitterData)
 		{
+			// Ignore lightweight (stateless) emitters
+			TSharedPtr<FNiagaraEmitterHandleViewModel> EmitterHandleViewModel = Context.ViewModel->GetEmitterHandleViewModelForEmitter(VersionedEmitter);
+			if (EmitterHandleViewModel.IsValid() &&
+			    EmitterHandleViewModel->GetEmitterHandle() != nullptr &&
+			    EmitterHandleViewModel->GetEmitterHandle()->GetEmitterMode() == ENiagaraEmitterMode::Stateless)
+			{
+				return;
+			}
+
 			const UNiagaraEmitter* ParentEmitter = VersionedEmitterData.VersionedParent.Emitter.Get();
 			if (IsValid(ParentEmitter) == false)
 			{
