- Niagara Outliner Changes
- Handle LWE emitters and add information about if LWE or not (fixes a crash if 2 LWE's exists in a system)
- Add button to browse to the asset, and asset information on the system

diff --git a/Engine/Plugins/FX/Niagara/Source/Niagara/Private/NiagaraDebuggerClient.cpp b/Engine/Plugins/FX/Niagara/Source/Niagara/Private/NiagaraDebuggerClient.cpp
--- a/Engine/Plugins/FX/Niagara/Source/Niagara/Private/NiagaraDebuggerClient.cpp
+++ b/Engine/Plugins/FX/Niagara/Source/Niagara/Private/NiagaraDebuggerClient.cpp
@@ -404,6 +404,8 @@ void FNiagaraDebuggerClient::CaptureOutlinerData(FNiagaraOutlinerData& OutlinerD
 		FNiagaraOutlinerSystemData& Instances = WorldData.Systems.FindOrAdd(System ? System->GetPathName() : TEXT("Null System"));
 		if (System)
 		{
+			Instances.AssetPath = FSoftObjectPath(System);
+
 			//Add System specific data.
 #if WITH_PARTICLE_PERF_STATS
 			if (StatsListener)
@@ -462,19 +464,19 @@ void FNiagaraDebuggerClient::CaptureOutlinerData(FNiagaraOutlinerData& OutlinerD
 				for (const FNiagaraEmitterInstanceRef& EmitterInst : Inst->GetEmitters())
 				{
 					FNiagaraOutlinerEmitterInstanceData& EmitterData = InstData.Emitters.AddDefaulted_GetRef();
-					if (const UNiagaraEmitter* Emitter = EmitterInst->GetEmitter())
-					{
-						//TODO: This is a bit wasteful to copy the name into each instance data. Though we can't rely on the debugger side data matchin the actul running data on the device.
-						//We need to build a shared representation of the asset data from the client that we then reference from this per instance data.
-						EmitterData.EmitterName = Emitter->GetUniqueEmitterName();
-						EmitterData.SimTarget = EmitterInst->GetSimTarget();
-						//Move all above to a shared asset representation.
-
-						EmitterData.ExecState = EmitterInst->GetExecutionState();
-						EmitterData.NumParticles = EmitterInst->GetNumParticles();
-
-						EmitterData.bRequiresPersistentIDs = EmitterInst->GetParticleData().RequiresPersistentIDs();
-					}
+					const FNiagaraEmitterHandle& EmitterHandle = EmitterInst->GetEmitterHandle();
+
+					//TODO: This is a bit wasteful to copy the name into each instance data. Though we can't rely on the debugger side data matchin the actul running data on the device.
+					//We need to build a shared representation of the asset data from the client that we then reference from this per instance data.
+					EmitterData.EmitterName = EmitterHandle.GetUniqueInstanceName();
+					EmitterData.SimTarget = EmitterInst->GetSimTarget();
+					//Move all above to a shared asset representation.
+
+					EmitterData.ExecState = EmitterInst->GetExecutionState();
+					EmitterData.NumParticles = EmitterInst->GetNumParticles();
+
+					EmitterData.bRequiresPersistentIDs = EmitterInst->GetParticleData().RequiresPersistentIDs();
+					EmitterData.bIsLightweight = EmitterInst->AsStateless() != nullptr;
 				}
 			}
 		}
diff --git a/Engine/Plugins/FX/Niagara/Source/Niagara/Public/NiagaraDebuggerCommon.h b/Engine/Plugins/FX/Niagara/Source/Niagara/Public/NiagaraDebuggerCommon.h
--- a/Engine/Plugins/FX/Niagara/Source/Niagara/Public/NiagaraDebuggerCommon.h
+++ b/Engine/Plugins/FX/Niagara/Source/Niagara/Public/NiagaraDebuggerCommon.h
@@ -54,11 +54,10 @@ struct FNiagaraOutlinerEmitterInstanceData
 	int32 NumParticles = 0;
 
 	UPROPERTY(VisibleAnywhere, Category = "Emitter")
-	uint32 bRequiresPersistentIDs : 1;
+	uint32 bRequiresPersistentIDs : 1 = false;
 
-	FNiagaraOutlinerEmitterInstanceData()
-		: bRequiresPersistentIDs(false)
-	{}
+	UPROPERTY(VisibleAnywhere, Category = "Emitter")
+	uint32 bIsLightweight : 1 = false;
 
 	//Mem Usage?
 	//Scalability info?
@@ -154,11 +153,12 @@ struct FNiagaraOutlinerSystemData
 
 	//TODO: Cache off any shared representation of the system and emitters here for the instances to reference. 
 
+	UPROPERTY(VisibleAnywhere, Category = "System")
+	FSoftObjectPath AssetPath;
+
 	/** Map of System Instance data indexed by the UNiagaraSystem name. */
 	UPROPERTY(VisibleAnywhere, Category = "System")
 	TArray<FNiagaraOutlinerSystemInstanceData> SystemInstances;
-
-	
 	
 	UPROPERTY(VisibleAnywhere, Category = "Performance")
 	FNiagaraOutlinerTimingData AveragePerFrameTime;
diff --git a/Engine/Plugins/FX/Niagara/Source/NiagaraEditor/Private/Customizations/NiagaraOutlinerCustomization.cpp b/Engine/Plugins/FX/Niagara/Source/NiagaraEditor/Private/Customizations/NiagaraOutlinerCustomization.cpp
--- a/Engine/Plugins/FX/Niagara/Source/NiagaraEditor/Private/Customizations/NiagaraOutlinerCustomization.cpp
+++ b/Engine/Plugins/FX/Niagara/Source/NiagaraEditor/Private/Customizations/NiagaraOutlinerCustomization.cpp
@@ -1026,89 +1026,118 @@ const void* FNiagaraOutlinerTreeSystemItem::GetData()const
 
 TSharedRef<SWidget> FNiagaraOutlinerTreeSystemItem::GetHeaderWidget()
 {
-	if (FNiagaraOutlinerSystemData* Data = (FNiagaraOutlinerSystemData*)GetData())
+	FNiagaraOutlinerSystemData* Data = (FNiagaraOutlinerSystemData*)GetData();
+	if (Data == nullptr)
 	{
-		TSharedPtr<SNiagaraOutlinerTree> Tree = OwnerTree.Pin();
-		check(Tree.IsValid());//If we managed to get valid data then the tree should be valid.
-		TSharedPtr<FNiagaraDebugger>& Debugger = Tree->GetDebugger();
-		check(Debugger.IsValid());
-		UNiagaraOutliner* Outliner = Debugger->GetOutliner();
-		check(Outliner);
+		return SNullWidget::NullWidget;
+	}
 
-		TSharedRef<SHorizontalBox> Box = SNew(SHorizontalBox);
+	TSharedPtr<SNiagaraOutlinerTree> Tree = OwnerTree.Pin();
+	check(Tree.IsValid());//If we managed to get valid data then the tree should be valid.
+	TSharedPtr<FNiagaraDebugger>& Debugger = Tree->GetDebugger();
+	check(Debugger.IsValid());
+	UNiagaraOutliner* Outliner = Debugger->GetOutliner();
+	check(Outliner);
 
-		if (Outliner->ViewSettings.ViewMode == ENiagaraOutlinerViewModes::State)
-		{
-			int32 NumMatchingSystems = 0;
-			for (const TSharedRef<FNiagaraOutlinerTreeItem>& Item : Children)
-			{
-				if (Item->GetType() == ENiagaraOutlinerTreeItemType::Component && Item->bVisible)
+	TSharedRef<SHorizontalBox> Box = SNew(SHorizontalBox);
+
+	if (Data->AssetPath.IsValid())
+	{
+		Box->AddSlot()
+		.AutoWidth()
+		.Padding(FMargin(HeaderPadding, 0.0f, HeaderPadding, 0.0f))
+		[
+			SNew(SButton)
+			.OnClicked_Lambda(
+				[AssetPath = Data->AssetPath]()
 				{
-					++NumMatchingSystems;
+					FAssetRegistryModule& AssetRegistryModule = FModuleManager::LoadModuleChecked<FAssetRegistryModule>("AssetRegistry");
+					FAssetData AssetData;
+					if (AssetRegistryModule.Get().TryGetAssetByObjectPath(AssetPath, AssetData) != UE::AssetRegistry::EExists::Unknown)
+					{
+						GEditor->SyncBrowserToObjects({AssetData});
+					}
+					return FReply::Handled();
 				}
-			}
+			)
+			.ToolTipText(LOCTEXT("NiagaraOutlineCaptureBrowseToAsset", "Browse to this asset in the content browser."))
+			[
+				SNew(SImage)
+				.Image(FAppStyle::GetBrush(TEXT("Icons.BrowseContent")))
+			]
+		];
+	}
 
-			Box->AddSlot()
-				.AutoWidth()
-				.Padding(FMargin(HeaderPadding, 0.0f, HeaderPadding, 0.0f))
-				[
-					SNew(SNiagaraOutlinerTreeItemHeaderDataWidget<int32>, Outliner->ViewSettings)
-					.ToolTipText(LOCTEXT("VisibleInstancesHeaderName", "Num system instances matching current search and filters."))
-					.Data(NumMatchingSystems)
-				];
-		}
-		else if (Outliner->ViewSettings.ViewMode == ENiagaraOutlinerViewModes::Performance)
+	if (Outliner->ViewSettings.ViewMode == ENiagaraOutlinerViewModes::State)
+	{
+		int32 NumMatchingSystems = 0;
+		for (const TSharedRef<FNiagaraOutlinerTreeItem>& Item : Children)
 		{
-			Box->AddSlot()
-				.AutoWidth()
-				.Padding(FMargin(HeaderPadding, 0.0f, HeaderPadding, 0.0f))
-				[
-					SNew(SNiagaraOutlinerTreeItemHeaderDataWidget<FNiagaraOutlinerTimingData>, Outliner->ViewSettings)
-					//.LabelText(LOCTEXT("SystemAvgPerInstanceTime", "Avg"))
-					.ToolTipText(LOCTEXT("SystemAvgPerInstanceTimeToolTip", "Average instance time for this System."))
-				.Data(Data->AveragePerInstanceTime)
-				.MinDesiredWidth(50)
-				];
+			if (Item->GetType() == ENiagaraOutlinerTreeItemType::Component && Item->bVisible)
+			{
+				++NumMatchingSystems;
+			}
+		}
 
-			Box->AddSlot()
-				.AutoWidth()
-				.Padding(FMargin(HeaderPadding, 0.0f, HeaderPadding, 0.0f))
-				[
-					SNew(SNiagaraOutlinerTreeItemHeaderDataWidget<FNiagaraOutlinerTimingData>, Outliner->ViewSettings)
-					//.LabelText(LOCTEXT("SystemMaxPerFrameTime", "Max"))
-					.ToolTipText(LOCTEXT("SystemMaxPerFrameTimeToolTip", "Max instance time for this System."))
-					.Data(Data->MaxPerInstanceTime)
-					.MinDesiredWidth(50)
-				];
+		Box->AddSlot()
+			.AutoWidth()
+			.Padding(FMargin(HeaderPadding, 0.0f, HeaderPadding, 0.0f))
+			[
+				SNew(SNiagaraOutlinerTreeItemHeaderDataWidget<int32>, Outliner->ViewSettings)
+				.ToolTipText(LOCTEXT("VisibleInstancesHeaderName", "Num system instances matching current search and filters."))
+				.Data(NumMatchingSystems)
+			];
+	}
+	else if (Outliner->ViewSettings.ViewMode == ENiagaraOutlinerViewModes::Performance)
+	{
+		Box->AddSlot()
+			.AutoWidth()
+			.Padding(FMargin(HeaderPadding, 0.0f, HeaderPadding, 0.0f))
+			[
+				SNew(SNiagaraOutlinerTreeItemHeaderDataWidget<FNiagaraOutlinerTimingData>, Outliner->ViewSettings)
+				//.LabelText(LOCTEXT("SystemAvgPerInstanceTime", "Avg"))
+				.ToolTipText(LOCTEXT("SystemAvgPerInstanceTimeToolTip", "Average instance time for this System."))
+			.Data(Data->AveragePerInstanceTime)
+			.MinDesiredWidth(50)
+			];
 
-			Box->AddSlot()
-				.AutoWidth()
-				.Padding(FMargin(HeaderPadding, 0.0f, HeaderPadding, 0.0f))
-				[
-					SNew(SNiagaraOutlinerTreeItemHeaderDataWidget<FNiagaraOutlinerTimingData>, Outliner->ViewSettings)
-					//.LabelText(LOCTEXT("SystemAvgPerFrameTime", "Frame Avg"))
-					.ToolTipText(LOCTEXT("SystemAvgPerFrameTimeToolTip", "Average frame time for all instances of this System."))
-					.Data(Data->AveragePerFrameTime)
-					.MinDesiredWidth(50)
-				];
+		Box->AddSlot()
+			.AutoWidth()
+			.Padding(FMargin(HeaderPadding, 0.0f, HeaderPadding, 0.0f))
+			[
+				SNew(SNiagaraOutlinerTreeItemHeaderDataWidget<FNiagaraOutlinerTimingData>, Outliner->ViewSettings)
+				//.LabelText(LOCTEXT("SystemMaxPerFrameTime", "Max"))
+				.ToolTipText(LOCTEXT("SystemMaxPerFrameTimeToolTip", "Max instance time for this System."))
+				.Data(Data->MaxPerInstanceTime)
+				.MinDesiredWidth(50)
+			];
 
-			Box->AddSlot()
-				.AutoWidth()
-				.Padding(FMargin(HeaderPadding, 0.0f, HeaderPadding, 0.0f))
-				[
-					SNew(SNiagaraOutlinerTreeItemHeaderDataWidget<FNiagaraOutlinerTimingData>, Outliner->ViewSettings)
-					//.LabelText(LOCTEXT("SystemMaxPerFrameTime", "Frame Max"))
-					.ToolTipText(LOCTEXT("AllInstancesMaxFrameTimeToolTip", "Max frame time for all instances of this System."))
-					.Data(Data->MaxPerFrameTime)
-					.MinDesiredWidth(50)
-				];
+		Box->AddSlot()
+			.AutoWidth()
+			.Padding(FMargin(HeaderPadding, 0.0f, HeaderPadding, 0.0f))
+			[
+				SNew(SNiagaraOutlinerTreeItemHeaderDataWidget<FNiagaraOutlinerTimingData>, Outliner->ViewSettings)
+				//.LabelText(LOCTEXT("SystemAvgPerFrameTime", "Frame Avg"))
+				.ToolTipText(LOCTEXT("SystemAvgPerFrameTimeToolTip", "Average frame time for all instances of this System."))
+				.Data(Data->AveragePerFrameTime)
+				.MinDesiredWidth(50)
+			];
 
-			//TODO: Add baseline comparison traffic light icon with tooltip containing details.
-		}
+		Box->AddSlot()
+			.AutoWidth()
+			.Padding(FMargin(HeaderPadding, 0.0f, HeaderPadding, 0.0f))
+			[
+				SNew(SNiagaraOutlinerTreeItemHeaderDataWidget<FNiagaraOutlinerTimingData>, Outliner->ViewSettings)
+				//.LabelText(LOCTEXT("SystemMaxPerFrameTime", "Frame Max"))
+				.ToolTipText(LOCTEXT("AllInstancesMaxFrameTimeToolTip", "Max frame time for all instances of this System."))
+				.Data(Data->MaxPerFrameTime)
+				.MinDesiredWidth(50)
+			];
 
-		return Box;
+		//TODO: Add baseline comparison traffic light icon with tooltip containing details.
 	}
-	return SNullWidget::NullWidget;
+
+	return Box;
 }
 
 void FNiagaraOutlinerTreeComponentItem::SortChildren() 
@@ -1281,7 +1310,7 @@ TSharedRef<SWidget> FNiagaraOutlinerTreeComponentItem::GetHeaderWidget()
 			.Padding(FMargin(HeaderPadding, 0.0f, HeaderPadding, 0.0f))
 			[
 				SNew(SButton)
-				.OnClicked_Raw(this, &FNiagaraOutlinerTreeComponentItem::CaputreSimCache)
+				.OnClicked(this, &FNiagaraOutlinerTreeComponentItem::CaputreSimCache)
 				.ToolTipText(LOCTEXT("NiagaraOutlineCaptureSimCacheTooltip", "Capture a new Sim Cache for this component."))
 				.ButtonStyle(FAppStyle::Get(), "FlatButton.Success")
 				//.ForegroundColor(FLinearColor::White)
@@ -1299,7 +1328,7 @@ TSharedRef<SWidget> FNiagaraOutlinerTreeComponentItem::GetHeaderWidget()
 				.Padding(FMargin(HeaderPadding, 0.0f, HeaderPadding, 0.0f))
 				[
 					SNew(SButton)
-					.OnClicked_Raw(this, &FNiagaraOutlinerTreeComponentItem::OpenSimCache)
+					.OnClicked(this, &FNiagaraOutlinerTreeComponentItem::OpenSimCache)
 					.ToolTipText(LOCTEXT("NiagaraOutlineOpenSimCacheTooltip", "Sim Cache is available for this component. Open in new window."))
 					.ButtonStyle(FAppStyle::Get(), "FlatButton.Success")
 					//.ForegroundColor(FLinearColor::White)
diff --git a/Engine/Plugins/FX/Niagara/Source/NiagaraEditor/Private/Customizations/NiagaraOutlinerCustomization.h b/Engine/Plugins/FX/Niagara/Source/NiagaraEditor/Private/Customizations/NiagaraOutlinerCustomization.h
--- a/Engine/Plugins/FX/Niagara/Source/NiagaraEditor/Private/Customizations/NiagaraOutlinerCustomization.h
+++ b/Engine/Plugins/FX/Niagara/Source/NiagaraEditor/Private/Customizations/NiagaraOutlinerCustomization.h
@@ -45,7 +45,7 @@ enum class ENiagaraOutlinerTreeItemType
 	Num,
 };
 
-struct FNiagaraOutlinerTreeItem
+struct FNiagaraOutlinerTreeItem : public TSharedFromThis<FNiagaraOutlinerTreeItem>
 {
 	FNiagaraOutlinerTreeItem()
 	{}
