[Niagara Validation] Clarify wording for fix message with required parent emitters

diff --git a/Engine/Plugins/FX/Niagara/Source/NiagaraEditor/Private/NiagaraValidationRules.cpp b/Engine/Plugins/FX/Niagara/Source/NiagaraEditor/Private/NiagaraValidationRules.cpp
--- a/Engine/Plugins/FX/Niagara/Source/NiagaraEditor/Private/NiagaraValidationRules.cpp
+++ b/Engine/Plugins/FX/Niagara/Source/NiagaraEditor/Private/NiagaraValidationRules.cpp
@@ -1841,6 +1841,24 @@ static bool ObjectIsInSoftObjectArray(const TArray<TSoftObjectPtr<ObjectType>>&
 	return ObjectArray.Contains(SoftObjectToFindPtr);
 }
 
+template<class T>
+static FString JoinSoftObjectArrayAsString(const TArray<TSoftObjectPtr<T>>& ObjectArray, const TCHAR* Separator)
+{
+	FString Result;
+	if (ObjectArray.IsEmpty())
+	{
+		return Result;
+	}
+
+	Result = ObjectArray[0].ToString();
+	for (int32 Idx = 1; Idx < ObjectArray.Num(); ++Idx)
+	{
+		Result += Separator + ObjectArray[Idx].ToString();
+	}
+
+	return Result;
+}
+
 void UNiagaraValidationRule_RequireParentEmitter::CheckValidity(const FNiagaraValidationContext& Context, TArray<FNiagaraValidationResult>& OutResults) const
 {
 	if (AllowedParentEmitters.IsEmpty())
@@ -1848,6 +1866,9 @@ void UNiagaraValidationRule_RequireParentEmitter::CheckValidity(const FNiagaraVa
 		return;
 	}
 
+	const FString AllowedParentEmitterString = JoinSoftObjectArrayAsString(AllowedParentEmitters, TEXT(", "));
+	const FText AllowedParentEmittersText = FText::FromString(AllowedParentEmitterString);
+
 	ForEachEmitterInContext(Context, [&, this](const FVersionedNiagaraEmitter& VersionedEmitter, UNiagaraStackEmitterPropertiesItem* EmitterProperties)
 	{
 		VersionedEmitter.Emitter->ForEachVersionData([&, this](const FVersionedNiagaraEmitterData& VersionedEmitterData)
@@ -1882,8 +1903,8 @@ void UNiagaraValidationRule_RequireParentEmitter::CheckValidity(const FNiagaraVa
 
 			FNiagaraValidationResult& Result = OutResults.AddDefaulted_GetRef();
 			Result.Severity = Severity;
-			Result.SummaryText = FText::Format(LOCTEXT("InvalidParentEmitter_Title", "Parent emitter '{0}' is invalid"), FText::FromName(*GetPathNameSafe(ParentEmitter)));
-			Result.Description = LOCTEXT("InvalidParentEmitter_Description", "The emitter is not explicitly allowed. Please consider using a different parent emitter");
+			Result.SummaryText = FText::Format(LOCTEXT("InvalidParentEmitter_Title", "Parent emitter '{0}' is not allowed"), FText::FromName(*GetPathNameSafe(ParentEmitter)));
+			Result.Description = FText::Format(LOCTEXT("InvalidParentEmitter_Description", "Please recreate the emitter using one of the following as the parent: {0}"), AllowedParentEmittersText);
 			Result.SourceObject = TWeakObjectPtr { EmitterProperties };
 		});
 	});
