- Capture system state inside debug cache
- Add a details view for the cache object

diff --git a/Engine/Plugins/FX/Niagara/Source/Niagara/Internal/NiagaraSimCacheDebugData.h b/Engine/Plugins/FX/Niagara/Source/Niagara/Internal/NiagaraSimCacheDebugData.h
--- a/Engine/Plugins/FX/Niagara/Source/Niagara/Internal/NiagaraSimCacheDebugData.h
+++ b/Engine/Plugins/FX/Niagara/Source/Niagara/Internal/NiagaraSimCacheDebugData.h
@@ -4,6 +4,7 @@
 
 #include "NiagaraCommon.h"
 #include "NiagaraScript.h"
+#include "NiagaraSystemEmitterState.h"
 
 #include "NiagaraSimCacheDebugData.generated.h"
 
@@ -28,6 +29,9 @@ public:
 	void CaptureFrame(FNiagaraSimCacheHelper& Helper, int FrameNumber);
 
 public:
+	UPROPERTY(VisibleAnywhere, Category="DebugData")
+	FNiagaraSystemStateData SystemStateData;
+
 	UPROPERTY()
 	TArray<FNiagaraSimCacheDebugDataFrame> Frames;
 };
diff --git a/Engine/Plugins/FX/Niagara/Source/Niagara/Private/NiagaraSimCache.cpp b/Engine/Plugins/FX/Niagara/Source/Niagara/Private/NiagaraSimCache.cpp
--- a/Engine/Plugins/FX/Niagara/Source/Niagara/Private/NiagaraSimCache.cpp
+++ b/Engine/Plugins/FX/Niagara/Source/Niagara/Private/NiagaraSimCache.cpp
@@ -357,7 +357,15 @@ bool UNiagaraSimCache::BeginWrite(FNiagaraSimCacheCreateParameters InCreateParam
 	DataInterfaceStorage.Empty();
 
 	// Create our debug data (if required)
-	DebugData = CreateParameters.bIncludeDebugData ? NewObject<UNiagaraSimCacheDebugData>(this) : nullptr;
+	if (CreateParameters.bIncludeDebugData)
+	{
+		DebugData = NewObject<UNiagaraSimCacheDebugData>(this);
+		DebugData->SystemStateData = Helper.NiagaraSystem->GetSystemStateData();
+	}
+	else
+	{
+		DebugData = nullptr;
+	}
 
 	// Not explicit capture mode?  Empty the list as internally we reuse this list when in rendering only mode
 	if (CreateParameters.AttributeCaptureMode == ENiagaraSimCacheAttributeCaptureMode::ExplicitAttributes)
diff --git a/Engine/Plugins/FX/Niagara/Source/NiagaraEditor/Private/Widgets/SNiagaraSimCacheDebugDataView.cpp b/Engine/Plugins/FX/Niagara/Source/NiagaraEditor/Private/Widgets/SNiagaraSimCacheDebugDataView.cpp
--- a/Engine/Plugins/FX/Niagara/Source/NiagaraEditor/Private/Widgets/SNiagaraSimCacheDebugDataView.cpp
+++ b/Engine/Plugins/FX/Niagara/Source/NiagaraEditor/Private/Widgets/SNiagaraSimCacheDebugDataView.cpp
@@ -6,9 +6,12 @@
 #include "NiagaraSimCacheDebugData.h"
 
 #include "Framework/MultiBox/MultiBoxBuilder.h"
+#include "Modules/ModuleManager.h"
 #include "Widgets/Input/SComboButton.h"
 #include "Widgets/Layout/SScrollBox.h"
 #include "Widgets/Views/SListView.h"
+#include "IDetailsView.h"
+#include "PropertyEditorModule.h"
 
 #define LOCTEXT_NAMESPACE "NiagaraSimCacheDebugDataView"
 
@@ -145,7 +148,6 @@ void SNiagaraSimCacheDebugDataView::Construct(const FArguments& InArgs)
 	SimCacheViewModel->OnSimCacheChanged().AddSP(this, &SNiagaraSimCacheDebugDataView::RefreshContents);
 	SimCacheViewModel->OnBufferChanged().AddSP(this, &SNiagaraSimCacheDebugDataView::RefreshContents);
 
-
 	if (const FNiagaraSimCacheDebugDataFrame* FrameData = GetCurrentFrameData())
 	{
 		if (FrameData->DebugParameterStores.Num() > 0)
@@ -154,6 +156,20 @@ void SNiagaraSimCacheDebugDataView::Construct(const FArguments& InArgs)
 		}
 	}
 
+	// Create our details view
+	{
+		FPropertyEditorModule& PropertyEditorModule = FModuleManager::LoadModuleChecked<FPropertyEditorModule>("PropertyEditor");
+
+		FDetailsViewArgs DetailsViewArgs;
+		DetailsViewArgs.bAllowSearch = false;
+		DetailsViewArgs.bShowOptions = false;
+		DetailsViewArgs.bShowObjectLabel = false;
+		DetailsViewArgs.NameAreaSettings = FDetailsViewArgs::HideNameArea;
+
+		DetailsView = PropertyEditorModule.CreateDetailView(DetailsViewArgs);
+		DetailsView->SetObject(const_cast<UNiagaraSimCacheDebugData*>(SimCacheViewModel->GetCacheDebugData()));
+	}
+
 	ChildSlot
 	[
 		SNew(SVerticalBox)
@@ -196,6 +212,12 @@ void SNiagaraSimCacheDebugDataView::Construct(const FArguments& InArgs)
 				SAssignNew(OverrideParametersWidget, UE::Niagara::SimCache::DebugDataUI::SParameterStoreListView)
 			]
 		]
+		+ SVerticalBox::Slot()
+		.AutoHeight()
+		.Padding(0.0f, 6.0f, 0.0, 0.0f)
+		[
+			DetailsView.ToSharedRef()
+		]
 	];
 
 	RefreshContents();
@@ -251,6 +273,8 @@ void SNiagaraSimCacheDebugDataView::RefreshContents()
 	}
 
 	OverrideParametersWidget->SetParameterStore(ParameterStore ? *ParameterStore : FNiagaraParameterStore());
+
+	DetailsView->SetObject(const_cast<UNiagaraSimCacheDebugData*>(SimCacheViewModel->GetCacheDebugData()));
 }
 
 void SNiagaraSimCacheDebugDataView::RefreshContents(bool)
diff --git a/Engine/Plugins/FX/Niagara/Source/NiagaraEditor/Private/Widgets/SNiagaraSimCacheDebugDataView.h b/Engine/Plugins/FX/Niagara/Source/NiagaraEditor/Private/Widgets/SNiagaraSimCacheDebugDataView.h
--- a/Engine/Plugins/FX/Niagara/Source/NiagaraEditor/Private/Widgets/SNiagaraSimCacheDebugDataView.h
+++ b/Engine/Plugins/FX/Niagara/Source/NiagaraEditor/Private/Widgets/SNiagaraSimCacheDebugDataView.h
@@ -4,6 +4,8 @@
 
 #include "Widgets/SCompoundWidget.h"
 
+class IDetailsView;
+
 class FNiagaraSimCacheViewModel;
 struct FNiagaraSimCacheDebugDataFrame;
 namespace UE::Niagara::SimCache::DebugDataUI
@@ -28,7 +30,10 @@ private:
 	void RefreshContents(bool);
 
 private:
-	TSharedPtr<FNiagaraSimCacheViewModel> SimCacheViewModel;
+	TSharedPtr<FNiagaraSimCacheViewModel>	SimCacheViewModel;
+	TSharedPtr<IDetailsView>				DetailsView;
+
 	TSharedPtr<UE::Niagara::SimCache::DebugDataUI::SParameterStoreListView> OverrideParametersWidget;
 	FString	SelectedParameterStoreName;
+
 };
