Cooker: SkipOnlyEditorOnly: Niagara: Mark that the AssetTagDefinitionAssets package loads for the deprecated NiagaraAssetTagDefinitions are editor-only loads. This prevents those packages from being added to the cook and prevents log spam.
#rnx

diff --git a/Engine/Plugins/FX/Niagara/Source/NiagaraEditor/Private/NiagaraEditorModule.cpp b/Engine/Plugins/FX/Niagara/Source/NiagaraEditor/Private/NiagaraEditorModule.cpp
--- a/Engine/Plugins/FX/Niagara/Source/NiagaraEditor/Private/NiagaraEditorModule.cpp
+++ b/Engine/Plugins/FX/Niagara/Source/NiagaraEditor/Private/NiagaraEditorModule.cpp
@@ -73,6 +73,7 @@
 #include "ViewModels/NiagaraPlaceholderDataInterfaceManager.h"
 #include "TNiagaraGraphPinEditableName.h"
 #include "UObject/Class.h"
+#include "UObject/ICookInfo.h"
 #include "NiagaraScriptMergeManager.h"
 #include "NiagaraDigestDatabase.h"
 #include "NiagaraEmitter.h"
@@ -418,13 +419,16 @@ PRAGMA_ENABLE_DEPRECATION_WARNINGS
 			TArray<FAssetData> AssetTagDefinitionAssets;
 			AssetRegistryModule.Get().GetAssets(Filter, AssetTagDefinitionAssets);
 
-			for(const FAssetData& AssetTagDefinitionAsset : AssetTagDefinitionAssets)
 			{
-				UDEPRECATED_NiagaraAssetTagDefinitions* TagDefinitionsAsset = Cast<UDEPRECATED_NiagaraAssetTagDefinitions>(AssetTagDefinitionAsset.GetAsset());
-
-				for(const FNiagaraAssetTagDefinition& TagDefinition : TagDefinitionsAsset->GetAssetTagDefinitions())
+				FCookLoadScope EditorOnlyLoad(ECookLoadType::EditorOnly);
+				for (const FAssetData& AssetTagDefinitionAsset : AssetTagDefinitionAssets)
 				{
-					AllTags.Add(TagDefinition.TagGuid, &TagDefinition);
+					UDEPRECATED_NiagaraAssetTagDefinitions* TagDefinitionsAsset = Cast<UDEPRECATED_NiagaraAssetTagDefinitions>(AssetTagDefinitionAsset.GetAsset());
+
+					for (const FNiagaraAssetTagDefinition& TagDefinition : TagDefinitionsAsset->GetAssetTagDefinitions())
+					{
+						AllTags.Add(TagDefinition.TagGuid, &TagDefinition);
+					}
 				}
 			}
 
