- Downgrade ensure to log

#rnx

diff --git a/Engine/Plugins/FX/Niagara/Source/NiagaraEditor/Private/ViewModels/NiagaraSystemViewModel.cpp b/Engine/Plugins/FX/Niagara/Source/NiagaraEditor/Private/ViewModels/NiagaraSystemViewModel.cpp
--- a/Engine/Plugins/FX/Niagara/Source/NiagaraEditor/Private/ViewModels/NiagaraSystemViewModel.cpp
+++ b/Engine/Plugins/FX/Niagara/Source/NiagaraEditor/Private/ViewModels/NiagaraSystemViewModel.cpp
@@ -1689,7 +1689,7 @@ TArray<FNiagaraGraphParameterReference> FNiagaraSystemViewModel::GetGraphParamet
 			if (Builder.Histories.Num() != 1)
 			{
 				// We should only have traversed one emitter (have not visited more than one NiagaraNodeEmitter.)
-				ensureMsgf(false, TEXT("Encountered more than one parameter map history when collecting graph parameter reference collections for system parameter panel view model!"));
+				UE_LOG(LogNiagaraEditor, Warning, TEXT("Encountered more than one parameter map history when collecting graph parameter reference collections for system parameter panel view model!"));
 			}
 			if (Builder.Histories.Num() == 0)
 			{
