- Check Canvas is valid before rendering the HUD / Fix removing from incorrect map

#rnx

diff --git a/Engine/Plugins/FX/Niagara/Source/Niagara/Private/NiagaraDebugHud.cpp b/Engine/Plugins/FX/Niagara/Source/Niagara/Private/NiagaraDebugHud.cpp
--- a/Engine/Plugins/FX/Niagara/Source/Niagara/Private/NiagaraDebugHud.cpp
+++ b/Engine/Plugins/FX/Niagara/Source/Niagara/Private/NiagaraDebugHud.cpp
@@ -1433,7 +1433,7 @@ void FNiagaraDebugHud::Draw(FNiagaraWorldManager* WorldManager, UCanvas* Canvas,
 	DeltaSeconds = float(CurrTime - LastDrawTime);
 
 	// We may want to use the HUD to pause / step but not actually display anything on screen
-	if (Settings.bHudRenderingEnabled)
+	if (Settings.bHudRenderingEnabled && Canvas && Canvas->Canvas)
 	{
 		// Draw in world components
 		DrawComponents(WorldManager, Canvas);
@@ -2429,7 +2429,7 @@ void FNiagaraDebugHud::DrawGpuComputeOverriew(class FNiagaraWorldManager* WorldM
 			UNiagaraEmitter* OwnerEmitter = EmitterIt.Key().Get();
 			if (OwnerEmitter == nullptr)
 			{
-				SystemIt.RemoveCurrent();
+				EmitterIt.RemoveCurrent();
 			}
 		}
 
