Niagara - Prevent a check when initializing the script execution parameter store with a compilled/resolved data mistmatch and log the issue.

diff --git a/Engine/Plugins/FX/Niagara/Source/Niagara/Private/NiagaraSystem.cpp b/Engine/Plugins/FX/Niagara/Source/Niagara/Private/NiagaraSystem.cpp
--- a/Engine/Plugins/FX/Niagara/Source/Niagara/Private/NiagaraSystem.cpp
+++ b/Engine/Plugins/FX/Niagara/Source/Niagara/Private/NiagaraSystem.cpp
@@ -645,8 +645,16 @@ TSharedPtr<const FNiagaraScriptRuntimeData> UNiagaraSystem::GetScriptRuntimeData
 		UNiagaraScript* TargetCompiledScript = GetTargetCompiledScript(ScriptDataKey);
 		if (TargetCompiledScript != nullptr)
 		{
-			TSharedPtr<const FNiagaraScriptRuntimeTransientData> ScriptRuntimeTransientData = MakeShared<FNiagaraScriptRuntimeTransientData>(*TargetCompiledScript, *ScriptRuntimeCompiledDataHandle->Instance);
-			ScriptRuntimeData = MakeShared<FNiagaraScriptRuntimeData>(ScriptRuntimeCompiledDataHandle.ToSharedRef(), ScriptRuntimeTransientData.ToSharedRef());
+			if (TargetCompiledScript->GetVMExecutableData().DataInterfaceInfo.Num() == ScriptRuntimeCompiledDataHandle->Instance->ResolvedDataInterfaces.Num())
+			{
+				TSharedPtr<const FNiagaraScriptRuntimeTransientData> ScriptRuntimeTransientData = MakeShared<FNiagaraScriptRuntimeTransientData>(*TargetCompiledScript, *ScriptRuntimeCompiledDataHandle->Instance);
+				ScriptRuntimeData = MakeShared<FNiagaraScriptRuntimeData>(ScriptRuntimeCompiledDataHandle.ToSharedRef(), ScriptRuntimeTransientData.ToSharedRef());
+			}
+			else
+			{
+				UE_LOG(LogNiagara, Log, TEXT("Could not create script runtime data for script %s due to compiled DI count %i and resolved DI count %i difference."),
+					*TargetCompiledScript->GetPathName(), TargetCompiledScript->GetVMExecutableData().DataInterfaceInfo.Num(), ScriptRuntimeCompiledDataHandle->Instance->ResolvedDataInterfaces.Num());
+			}
 		}
 	}
 
