- Fix LWE performance message showing when you have disabled standard emitters

#rnx

#FailedValidations Preflight

diff --git a/Engine/Plugins/FX/Niagara/Source/NiagaraEditorWidgets/Private/SNiagaraOverviewGraphTitleBar.cpp b/Engine/Plugins/FX/Niagara/Source/NiagaraEditorWidgets/Private/SNiagaraOverviewGraphTitleBar.cpp
--- a/Engine/Plugins/FX/Niagara/Source/NiagaraEditorWidgets/Private/SNiagaraOverviewGraphTitleBar.cpp
+++ b/Engine/Plugins/FX/Niagara/Source/NiagaraEditorWidgets/Private/SNiagaraOverviewGraphTitleBar.cpp
@@ -267,14 +267,21 @@ EVisibility SNiagaraOverviewGraphTitleBar::GetStatelessInfoVisibility() const
 		const UNiagaraSystem& NiagaraSystem = SystemViewModel.Pin()->GetSystem();
 		if (NiagaraSystem.GetEmitterHandles().Num() > 0)
 		{
+			bool bHasEnabledEmitters = false;
 			for (const FNiagaraEmitterHandle& EmitterHandle : NiagaraSystem.GetEmitterHandles())
 			{
-				if (EmitterHandle.GetIsEnabled() && EmitterHandle.GetEmitterMode() != ENiagaraEmitterMode::Stateless)
+				if (EmitterHandle.GetIsEnabled() == false)
+				{
+					continue;
+				}
+
+				bHasEnabledEmitters = true;
+				if (EmitterHandle.GetEmitterMode() != ENiagaraEmitterMode::Stateless)
 				{
 					return EVisibility::Collapsed;
 				}
 			}
-			return EVisibility::Visible;
+			return bHasEnabledEmitters ? EVisibility::Visible : EVisibility::Collapsed;
 		}
 	}
 	return EVisibility::Collapsed;
@@ -288,13 +295,13 @@ FText SNiagaraOverviewGraphTitleBar::GetStatelessInfoText() const
 		const FNiagaraSystemStateData& SystemStateData = NiagaraSystem.GetSystemStateData();
 		if (SystemStateData.bRunUpdateScript)
 		{
-			return LOCTEXT("LWEInfo_WorstPerf", "LWE: Spawn & update scripts will run (Worst Perf)");
+			return LOCTEXT("LWEInfo_WorstPerf", "Lightweight Performance: (Normal) Spawn & update scripts are running");
 		}
 		else if (SystemStateData.bRunSpawnScript)
 		{
-			return LOCTEXT("LWEInfo_GoodPerf", "LWE: Spawn script only will run (Good Perf)");
+			return LOCTEXT("LWEInfo_GoodPerf", "Lightweight Performance: (Good) Spawn script is running");
 		}
-		return LOCTEXT("LWEInfo_BestPerf", "LWE: No scripts will run (Best Perf)");
+		return LOCTEXT("LWEInfo_BestPerf", "Lightweight Performance: (Best) No scripts are running");
 	}
 	return FText::GetEmpty();
 }
