Niagara: VM func binder cleanup

Prefer Bind functions, and don't introduce extra lambdas just to tunnel through this pointers or extra args, Delegates can already do that

diff --git a/Engine/Plugins/FX/Niagara/Source/Niagara/Classes/NiagaraDataInterfaceArrayImpl.h b/Engine/Plugins/FX/Niagara/Source/Niagara/Classes/NiagaraDataInterfaceArrayImpl.h
--- a/Engine/Plugins/FX/Niagara/Source/Niagara/Classes/NiagaraDataInterfaceArrayImpl.h
+++ b/Engine/Plugins/FX/Niagara/Source/Niagara/Classes/NiagaraDataInterfaceArrayImpl.h
@@ -648,17 +648,17 @@ struct FNDIArrayProxyImpl : public INDIArrayProxyBase
 		if (BindingInfo.Name == FNiagaraDataInterfaceArrayImplInternal::Function_LengthName)
 		{
 			check(BindingInfo.GetNumInputs() == 1 && BindingInfo.GetNumOutputs() == 1);
-			OutFunc = FVMExternalFunction::CreateLambda([this](FVectorVMExternalFunctionContext& Context) { this->VMGetLength(Context); });
+			OutFunc.BindRaw(this, &FNDIArrayProxyImpl::VMGetLength);
 		}
 		else if (BindingInfo.Name == FNiagaraDataInterfaceArrayImplInternal::Function_IsValidIndexName)
 		{
 			check(BindingInfo.GetNumInputs() == 2 && BindingInfo.GetNumOutputs() == 1);
-			OutFunc = FVMExternalFunction::CreateLambda([this](FVectorVMExternalFunctionContext& Context) { this->VMIsValidIndex(Context); });
+			OutFunc.BindRaw(this, &FNDIArrayProxyImpl::VMIsValidIndex);
 		}
 		else if (BindingInfo.Name == FNiagaraDataInterfaceArrayImplInternal::Function_LastIndexName)
 		{
 			check(BindingInfo.GetNumInputs() == 1 && BindingInfo.GetNumOutputs() == 1);
-			OutFunc = FVMExternalFunction::CreateLambda([this](FVectorVMExternalFunctionContext& Context) { this->VMGetLastIndex(Context); });
+			OutFunc.BindRaw(this, &FNDIArrayProxyImpl::VMGetLastIndex);
 		}
 	}
 
@@ -670,32 +670,32 @@ struct FNDIArrayProxyImpl : public INDIArrayProxyBase
 		{
 			// Note: Outputs is variable based upon type
 			//check(BindingInfo.GetNumInputs() == 2 && BindingInfo.GetNumOutputs() == 1);
-			OutFunc = FVMExternalFunction::CreateLambda([this](FVectorVMExternalFunctionContext& Context) { this->VMGetValue(Context); });
+			OutFunc.BindRaw(this, &FNDIArrayProxyImpl::VMGetValue);
 		}
 		// Mutable functions
 		else if (BindingInfo.Name == FNiagaraDataInterfaceArrayImplInternal::Function_ClearName)
 		{
 			check(BindingInfo.GetNumInputs() == 1 && BindingInfo.GetNumOutputs() == 0);
-			OutFunc = FVMExternalFunction::CreateLambda([this](FVectorVMExternalFunctionContext& Context) { this->VMClear(Context); });
+			OutFunc.BindRaw(this, &FNDIArrayProxyImpl::VMClear);
 		}
 		else if (BindingInfo.Name == FNiagaraDataInterfaceArrayImplInternal::Function_ResizeName)
 		{
 			check(BindingInfo.GetNumInputs() == 2 && BindingInfo.GetNumOutputs() == 0);
-			OutFunc = FVMExternalFunction::CreateLambda([this](FVectorVMExternalFunctionContext& Context) { this->VMResize(Context); });
+			OutFunc.BindRaw(this, &FNDIArrayProxyImpl::VMResize);
 		}
 		else if (BindingInfo.Name == FNiagaraDataInterfaceArrayImplInternal::Function_SetArrayElemName)
 		{
-			OutFunc = FVMExternalFunction::CreateLambda([this](FVectorVMExternalFunctionContext& Context) { this->VMSetValue(Context); });
+			OutFunc.BindRaw(this, &FNDIArrayProxyImpl::VMSetValue);
 		}
 		else if (BindingInfo.Name == FNiagaraDataInterfaceArrayImplInternal::Function_AddName)
 		{
 			// Note: Inputs is variable based upon type
-			OutFunc = FVMExternalFunction::CreateLambda([this](FVectorVMExternalFunctionContext& Context) { this->VMPushValue(Context); });
+			OutFunc.BindRaw(this, &FNDIArrayProxyImpl::VMPushValue);
 		}
 		else if (BindingInfo.Name == FNiagaraDataInterfaceArrayImplInternal::Function_RemoveLastElemName)
 		{
 			// Note: Outputs is variable based upon type
-			OutFunc = FVMExternalFunction::CreateLambda([this](FVectorVMExternalFunctionContext& Context) { this->VMPopValue(Context); });
+			OutFunc.BindRaw(this, &FNDIArrayProxyImpl::VMPopValue);
 		}
 	}
 
@@ -709,15 +709,15 @@ struct FNDIArrayProxyImpl : public INDIArrayProxyBase
 	{
 		if (BindingInfo.Name == FNiagaraDataInterfaceArrayImplInternal::Function_AtomicAddName)
 		{
-			OutFunc = FVMExternalFunction::CreateLambda([this](FVectorVMExternalFunctionContext& Context) { this->VMAtomicAdd(Context); });
+			OutFunc.BindRaw(this, &FNDIArrayProxyImpl::VMAtomicAdd);
 		}
 		else if (BindingInfo.Name == FNiagaraDataInterfaceArrayImplInternal::Function_AtomicMinName)
 		{
-			OutFunc = FVMExternalFunction::CreateLambda([this](FVectorVMExternalFunctionContext& Context) { this->VMAtomicMin(Context); });
+			OutFunc.BindRaw(this, &FNDIArrayProxyImpl::VMAtomicMin);
 		}
 		else if (BindingInfo.Name == FNiagaraDataInterfaceArrayImplInternal::Function_AtomicMaxName)
 		{
-			OutFunc = FVMExternalFunction::CreateLambda([this](FVectorVMExternalFunctionContext& Context) { this->VMAtomicMax(Context); });
+			OutFunc.BindRaw(this, &FNDIArrayProxyImpl::VMAtomicMax);
 		}
 	}
 
diff --git a/Engine/Plugins/FX/Niagara/Source/Niagara/Private/DataInterface/NiagaraDataInterfaceActorComponent.cpp b/Engine/Plugins/FX/Niagara/Source/Niagara/Private/DataInterface/NiagaraDataInterfaceActorComponent.cpp
--- a/Engine/Plugins/FX/Niagara/Source/Niagara/Private/DataInterface/NiagaraDataInterfaceActorComponent.cpp
+++ b/Engine/Plugins/FX/Niagara/Source/Niagara/Private/DataInterface/NiagaraDataInterfaceActorComponent.cpp
@@ -227,15 +227,15 @@ void UNiagaraDataInterfaceActorComponent::GetVMExternalFunction(const FVMExterna
 	using namespace NDIActorComponentLocal;
 	if (BindingInfo.Name == GetMatrixName)
 	{
-		OutFunc = FVMExternalFunction::CreateLambda([this](FVectorVMExternalFunctionContext& Context) { this->VMGetMatrix(Context); });
+		OutFunc.BindUObject(this, &UNiagaraDataInterfaceActorComponent::VMGetMatrix);
 	}
 	else if (BindingInfo.Name == GetTransformName)
 	{
-		OutFunc = FVMExternalFunction::CreateLambda([this](FVectorVMExternalFunctionContext& Context) { this->VMGetTransform(Context); });
+		OutFunc.BindUObject(this, &UNiagaraDataInterfaceActorComponent::VMGetTransform);
 	}
 	else if (BindingInfo.Name == GetVelocityName)
 	{
-		OutFunc = FVMExternalFunction::CreateLambda([this](FVectorVMExternalFunctionContext& Context) { this->VMGetPhysicsVelocity(Context); });
+		OutFunc.BindUObject(this, &UNiagaraDataInterfaceActorComponent::VMGetPhysicsVelocity);
 	}
 }
 
diff --git a/Engine/Plugins/FX/Niagara/Source/Niagara/Private/DataInterface/NiagaraDataInterfaceConsoleVariable.cpp b/Engine/Plugins/FX/Niagara/Source/Niagara/Private/DataInterface/NiagaraDataInterfaceConsoleVariable.cpp
--- a/Engine/Plugins/FX/Niagara/Source/Niagara/Private/DataInterface/NiagaraDataInterfaceConsoleVariable.cpp
+++ b/Engine/Plugins/FX/Niagara/Source/Niagara/Private/DataInterface/NiagaraDataInterfaceConsoleVariable.cpp
@@ -129,7 +129,7 @@ namespace NDIConsoleVariableLocal
 	};
 
 	template<typename TVariableType>
-	void VMGetConsoleVariable(FVectorVMExternalFunctionContext& Context, uint32 ValueOffset, bool bValid)
+	void VMGetConsoleVariable(FVectorVMExternalFunctionContext& Context, int32 ValueOffset, bool bValid)
 	{
 		VectorVM::FUserPtrHandler<FInstanceData_GameThread>	InstanceData_GT(Context);
 		FNDIOutputParam<bool>								SuccessValue(Context);
@@ -313,7 +313,7 @@ void UNiagaraDataInterfaceConsoleVariable::GetVMExternalFunction(const FVMExtern
 		{
 			InstanceData_GT->ConsoleVariableFloats.Emplace(ConsoleVariable, ValueOffset);
 		}
-		OutFunc = FVMExternalFunction::CreateLambda([ValueOffset, bValid](FVectorVMExternalFunctionContext& Context) { VMGetConsoleVariable<float>(Context, ValueOffset, bValid); });
+		OutFunc.BindStatic(&VMGetConsoleVariable<float>, ValueOffset, bValid);
 	}
 	else if (BindingInfo.Name == GetConsoleVariableInt)
 	{
@@ -323,7 +323,7 @@ void UNiagaraDataInterfaceConsoleVariable::GetVMExternalFunction(const FVMExtern
 		{
 			InstanceData_GT->ConsoleVariableInts.Emplace(ConsoleVariable, ValueOffset);
 		}
-		OutFunc = FVMExternalFunction::CreateLambda([ValueOffset, bValid](FVectorVMExternalFunctionContext& Context) { VMGetConsoleVariable<int32>(Context, ValueOffset, bValid); });
+		OutFunc.BindStatic(&VMGetConsoleVariable<int32>, ValueOffset, bValid);
 	}
 	else if (BindingInfo.Name == GetConsoleVariableBool)
 	{
@@ -333,7 +333,7 @@ void UNiagaraDataInterfaceConsoleVariable::GetVMExternalFunction(const FVMExtern
 		{
 			InstanceData_GT->ConsoleVariableBools.Emplace(ConsoleVariable, ValueOffset);
 		}
-		OutFunc = FVMExternalFunction::CreateLambda([ValueOffset, bValid](FVectorVMExternalFunctionContext& Context) { VMGetConsoleVariable<bool>(Context, ValueOffset, bValid); });
+		OutFunc.BindStatic(&VMGetConsoleVariable<bool>, ValueOffset, bValid);
 	}
 }
 
diff --git a/Engine/Plugins/FX/Niagara/Source/Niagara/Private/DataInterface/NiagaraDataInterfaceDataChannelRead.cpp b/Engine/Plugins/FX/Niagara/Source/Niagara/Private/DataInterface/NiagaraDataInterfaceDataChannelRead.cpp
--- a/Engine/Plugins/FX/Niagara/Source/Niagara/Private/DataInterface/NiagaraDataInterfaceDataChannelRead.cpp
+++ b/Engine/Plugins/FX/Niagara/Source/Niagara/Private/DataInterface/NiagaraDataInterfaceDataChannelRead.cpp
@@ -1283,11 +1283,11 @@ void UNiagaraDataInterfaceDataChannelRead::GetVMExternalFunction(const FVMExtern
 {
 	if (BindingInfo.Name == NDIDataChannelReadLocal::NAME_Num)
 	{
-		OutFunc = FVMExternalFunction::CreateLambda([this](FVectorVMExternalFunctionContext& Context) { this->Num(Context); });
+		OutFunc.BindUObject(this, &UNiagaraDataInterfaceDataChannelRead::Num);
 	}
 	else if (BindingInfo.Name == NDIDataChannelUtilities::GetNDCSpawnDataName)
 	{
-		OutFunc = FVMExternalFunction::CreateLambda([this](FVectorVMExternalFunctionContext& Context) { this->GetNDCSpawnData(Context); });
+		OutFunc.BindUObject(this, &UNiagaraDataInterfaceDataChannelRead::GetNDCSpawnData);
 	}
 	else if (BindingInfo.Name == NDIDataChannelReadLocal::NAME_SpawnDirect)
 	{
@@ -1308,7 +1308,7 @@ void UNiagaraDataInterfaceDataChannelRead::GetVMExternalFunction(const FVMExtern
 			FNiagaraTypeDefinition NDCVarType = NDCVar->GetType();
 			if (NDCVarType == FNiagaraTypeDefinition::GetIntDef()) 
 			{
-				OutFunc = FVMExternalFunction::CreateLambda([this, NDCVarName](FVectorVMExternalFunctionContext& Context) { this->SpawnDirect<int32>(Context, NDCVarName); });
+				OutFunc.BindUObject(this, &UNiagaraDataInterfaceDataChannelRead::SpawnDirect<int32>, NDCVarName);
 			}
 			else
 			{
@@ -1338,12 +1338,12 @@ void UNiagaraDataInterfaceDataChannelRead::GetVMExternalFunction(const FVMExtern
 		{
 			FName NDCVarName = NDCVar->GetName();
 			FNiagaraTypeDefinition NDCVarType = NDCVar->GetType();
-			if (NDCVarType == FNiagaraTypeDefinition::GetIntDef()) OutFunc = FVMExternalFunction::CreateLambda([this, NDCVarName](FVectorVMExternalFunctionContext& Context) { this->ScaleSpawnCount<int32>(Context, NDCVarName); });
-			else if (NDCVarType == FNiagaraTypeHelper::GetDoubleDef()) OutFunc = FVMExternalFunction::CreateLambda([this, NDCVarName](FVectorVMExternalFunctionContext& Context) { this->ScaleSpawnCount<double>(Context, NDCVarName); });
-			else if (NDCVarType == FNiagaraTypeHelper::GetVector2DDef()) OutFunc = FVMExternalFunction::CreateLambda([this, NDCVarName](FVectorVMExternalFunctionContext& Context) { this->ScaleSpawnCount<FVector2D>(Context, NDCVarName); });
-			else if (NDCVarType == FNiagaraTypeHelper::GetVectorDef()) OutFunc = FVMExternalFunction::CreateLambda([this, NDCVarName](FVectorVMExternalFunctionContext& Context) { this->ScaleSpawnCount<FVector>(Context, NDCVarName); });
-			else if (NDCVarType == FNiagaraTypeHelper::GetVector4Def()) OutFunc = FVMExternalFunction::CreateLambda([this, NDCVarName](FVectorVMExternalFunctionContext& Context) { this->ScaleSpawnCount<FVector4>(Context, NDCVarName); });
-			else if (NDCVarType == FNiagaraTypeDefinition::GetPositionDef()) OutFunc = FVMExternalFunction::CreateLambda([this, NDCVarName](FVectorVMExternalFunctionContext& Context) { this->ScaleSpawnCount<FNiagaraPosition>(Context, NDCVarName); });
+			if (NDCVarType == FNiagaraTypeDefinition::GetIntDef()) OutFunc.BindUObject(this, &UNiagaraDataInterfaceDataChannelRead::ScaleSpawnCount<int32>, NDCVarName);
+			else if (NDCVarType == FNiagaraTypeHelper::GetDoubleDef()) OutFunc.BindUObject(this, &UNiagaraDataInterfaceDataChannelRead::ScaleSpawnCount<double>, NDCVarName);
+			else if (NDCVarType == FNiagaraTypeHelper::GetVector2DDef()) OutFunc.BindUObject(this, &UNiagaraDataInterfaceDataChannelRead::ScaleSpawnCount<FVector2D>, NDCVarName);
+			else if (NDCVarType == FNiagaraTypeHelper::GetVectorDef()) OutFunc.BindUObject(this, &UNiagaraDataInterfaceDataChannelRead::ScaleSpawnCount<FVector>, NDCVarName);
+			else if (NDCVarType == FNiagaraTypeHelper::GetVector4Def()) OutFunc.BindUObject(this, &UNiagaraDataInterfaceDataChannelRead::ScaleSpawnCount<FVector4>, NDCVarName);
+			else if (NDCVarType == FNiagaraTypeDefinition::GetPositionDef()) OutFunc.BindUObject(this, &UNiagaraDataInterfaceDataChannelRead::ScaleSpawnCount<FNiagaraPosition>, NDCVarName);
 			else
 			{
 				UE_LOG(LogTemp, Display, TEXT("Failed to bind Data Interface function as this is not a valid variable type for ScaleSpawnCount.\nDI: %s.\nReceived Name: %s\nNDC: %s"), *GetPathNameSafe(this), *BindingInfo.Name.ToString(), *GetPathNameSafe(Channel));
@@ -1359,15 +1359,15 @@ void UNiagaraDataInterfaceDataChannelRead::GetVMExternalFunction(const FVMExtern
 		int32 FuncIndex = CompiledData.FindFunctionInfoIndex(BindingInfo.Name, BindingInfo.VariadicInputs, BindingInfo.VariadicOutputs);
 		if (BindingInfo.Name == NDIDataChannelReadLocal::NAME_Read)
 		{
-			OutFunc = FVMExternalFunction::CreateLambda([this, FuncIndex](FVectorVMExternalFunctionContext& Context) { this->Read(Context, FuncIndex); });
+			OutFunc.BindUObject(this, &UNiagaraDataInterfaceDataChannelRead::Read, FuncIndex);
 		}
 		else if (BindingInfo.Name == NDIDataChannelReadLocal::NAME_Consume)
 		{
-			OutFunc = FVMExternalFunction::CreateLambda([this, FuncIndex](FVectorVMExternalFunctionContext& Context) { this->Consume(Context, FuncIndex); });
+			OutFunc.BindUObject(this, &UNiagaraDataInterfaceDataChannelRead::Consume, FuncIndex);
 		}
 		else if (BindingInfo.Name == NDIDataChannelReadLocal::NAME_SpawnConditional)
 		{
-			OutFunc = FVMExternalFunction::CreateLambda([this, FuncIndex](FVectorVMExternalFunctionContext& Context) { this->SpawnConditional(Context, FuncIndex); });
+			OutFunc.BindUObject(this, &UNiagaraDataInterfaceDataChannelRead::SpawnConditional, FuncIndex);
 		}
 		else
 		{
diff --git a/Engine/Plugins/FX/Niagara/Source/Niagara/Private/DataInterface/NiagaraDataInterfaceDataChannelWrite.cpp b/Engine/Plugins/FX/Niagara/Source/Niagara/Private/DataInterface/NiagaraDataInterfaceDataChannelWrite.cpp
--- a/Engine/Plugins/FX/Niagara/Source/Niagara/Private/DataInterface/NiagaraDataInterfaceDataChannelWrite.cpp
+++ b/Engine/Plugins/FX/Niagara/Source/Niagara/Private/DataInterface/NiagaraDataInterfaceDataChannelWrite.cpp
@@ -1210,23 +1210,23 @@ void UNiagaraDataInterfaceDataChannelWrite::GetVMExternalFunction(const FVMExter
 {
 	if (BindingInfo.Name == NDIDataChannelWriteLocal::GetFunctionSig_Num().Name)
 	{
-		OutFunc = FVMExternalFunction::CreateLambda([this](FVectorVMExternalFunctionContext& Context) { this->Num(Context); });
+		OutFunc.BindUObject(this, &UNiagaraDataInterfaceDataChannelWrite::Num);
 	}
 // 	else if (BindingInfo.Name == NDIDataChannelWriteLocal::GetFunctionSig_Allocate().Name)
 // 	{
-// 		OutFunc = FVMExternalFunction::CreateLambda([this](FVectorVMExternalFunctionContext& Context) { this->Allocate(Context); });
+//		OutFunc.BindUObject(this, &UNiagaraDataInterfaceDataChannelWrite::Allocate);
 // 	}
 	else
 	{
 		int32 FuncIndex = CompiledData.FindFunctionInfoIndex(BindingInfo.Name, BindingInfo.VariadicInputs, BindingInfo.VariadicOutputs);
 		if (BindingInfo.Name == NDIDataChannelWriteLocal::GetFunctionSig_Write().Name)
 		{
-			OutFunc = FVMExternalFunction::CreateLambda([this, FuncIndex](FVectorVMExternalFunctionContext& Context) { this->Write(Context, FuncIndex); });
+			OutFunc.BindUObject(this, &UNiagaraDataInterfaceDataChannelWrite::Write, FuncIndex);
 		}
 		else 
 		if (BindingInfo.Name == NDIDataChannelWriteLocal::GetFunctionSig_Append().Name)
 		{
-			OutFunc = FVMExternalFunction::CreateLambda([this, FuncIndex](FVectorVMExternalFunctionContext& Context) { this->Append(Context, FuncIndex); });
+			OutFunc.BindUObject(this, &UNiagaraDataInterfaceDataChannelWrite::Append, FuncIndex);
 		}
 		else
 		{
diff --git a/Engine/Plugins/FX/Niagara/Source/Niagara/Private/DataInterface/NiagaraDataInterfaceDataTable.cpp b/Engine/Plugins/FX/Niagara/Source/Niagara/Private/DataInterface/NiagaraDataInterfaceDataTable.cpp
--- a/Engine/Plugins/FX/Niagara/Source/Niagara/Private/DataInterface/NiagaraDataInterfaceDataTable.cpp
+++ b/Engine/Plugins/FX/Niagara/Source/Niagara/Private/DataInterface/NiagaraDataInterfaceDataTable.cpp
@@ -674,15 +674,15 @@ void UNiagaraDataInterfaceDataTable::GetVMExternalFunction(const FVMExternalFunc
 
 	if ( BindingInfo.Name == NAME_IsValid )
 	{
-		OutFunc = FVMExternalFunction::CreateStatic(&VMIsValid);
+		OutFunc.BindStatic(&VMIsValid);
 	}
 	else if ( BindingInfo.Name == NAME_GetNumRows )
 	{
-		OutFunc = FVMExternalFunction::CreateStatic(&VMGetNumRows<false>);
+		OutFunc.BindStatic(&VMGetNumRows<false>);
 	}
 	else if (BindingInfo.Name == NAME_GetNumFilteredRows)
 	{
-		OutFunc = FVMExternalFunction::CreateStatic(&VMGetNumRows<true>);
+		OutFunc.BindStatic(&VMGetNumRows<true>);
 	}
 	else if ( BindingInfo.Name == NAME_GetRow )
 	{
diff --git a/Engine/Plugins/FX/Niagara/Source/Niagara/Private/DataInterface/NiagaraDataInterfaceEmitterProperties.cpp b/Engine/Plugins/FX/Niagara/Source/Niagara/Private/DataInterface/NiagaraDataInterfaceEmitterProperties.cpp
--- a/Engine/Plugins/FX/Niagara/Source/Niagara/Private/DataInterface/NiagaraDataInterfaceEmitterProperties.cpp
+++ b/Engine/Plugins/FX/Niagara/Source/Niagara/Private/DataInterface/NiagaraDataInterfaceEmitterProperties.cpp
@@ -262,19 +262,19 @@ void UNiagaraDataInterfaceEmitterProperties::GetVMExternalFunction(const FVMExte
 
 	if ( BindingInfo.Name == GetLocalSpaceName )
 	{
-		OutFunc = FVMExternalFunction::CreateLambda([this](FVectorVMExternalFunctionContext& Context) { VMGetLocalSpace(Context); });
+		OutFunc.BindUObject(this, &UNiagaraDataInterfaceEmitterProperties::VMGetLocalSpace);
 	}
 	else if (BindingInfo.Name == GetBoundsName)
 	{
-		OutFunc = FVMExternalFunction::CreateLambda([this](FVectorVMExternalFunctionContext& Context) { VMGetBounds(Context); });
+		OutFunc.BindUObject(this, &UNiagaraDataInterfaceEmitterProperties::VMGetBounds);
 	}
 	else if (BindingInfo.Name == GetFixedBoundsName)
 	{
-		OutFunc = FVMExternalFunction::CreateLambda([this](FVectorVMExternalFunctionContext& Context) { VMGetFixedBounds(Context); });
+		OutFunc.BindUObject(this, &UNiagaraDataInterfaceEmitterProperties::VMGetFixedBounds);
 	}
 	else if (BindingInfo.Name == SetFixedBoundsName)
 	{
-		OutFunc = FVMExternalFunction::CreateLambda([this](FVectorVMExternalFunctionContext& Context) { VMSetFixedBounds(Context); });
+		OutFunc.BindUObject(this, &UNiagaraDataInterfaceEmitterProperties::VMSetFixedBounds);
 	}
 }
 
diff --git a/Engine/Plugins/FX/Niagara/Source/Niagara/Private/DataInterface/NiagaraDataInterfaceUObjectPropertyReader.cpp b/Engine/Plugins/FX/Niagara/Source/Niagara/Private/DataInterface/NiagaraDataInterfaceUObjectPropertyReader.cpp
--- a/Engine/Plugins/FX/Niagara/Source/Niagara/Private/DataInterface/NiagaraDataInterfaceUObjectPropertyReader.cpp
+++ b/Engine/Plugins/FX/Niagara/Source/Niagara/Private/DataInterface/NiagaraDataInterfaceUObjectPropertyReader.cpp
@@ -953,18 +953,18 @@ void UNiagaraDataInterfaceUObjectPropertyReader::GetVMExternalFunction(const FVM
 
 	if ( BindingInfo.Name == GetComponentTransformName )
 	{
-		OutFunc = FVMExternalFunction::CreateLambda([this](FVectorVMExternalFunctionContext& Context) { VMGetComponentTransform(Context); });
+		OutFunc.BindUObject(this, &UNiagaraDataInterfaceUObjectPropertyReader::VMGetComponentTransform);
 	}
 	else if (BindingInfo.Name == GetComponentInvTransformName)
 	{
-		OutFunc = FVMExternalFunction::CreateLambda([this](FVectorVMExternalFunctionContext& Context) { VMGetComponentInvTransform(Context); });
+		OutFunc.BindUObject(this, &UNiagaraDataInterfaceUObjectPropertyReader::VMGetComponentInvTransform);
 	}
 	// Bind property functions
 	#define NDI_PROPERTY_TYPE(TYPE) \
 		else if (BindingInfo.Name == FTypeHelper<TYPE>::GetFunctionName()) \
 		{ \
 			const uint32 GetterIndex = InstanceData_GT->AddProperty(FNiagaraVariableBase(FTypeHelper<TYPE>::GetTypeDef(), GetRemappedPropertyName(BindingInfo.FunctionSpecifiers[0].Value))); \
-			OutFunc = FVMExternalFunction::CreateLambda([GetterIndex](FVectorVMExternalFunctionContext& Context) { FTypeHelper<TYPE>::VMFunction(Context, GetterIndex); }); \
+			OutFunc.BindStatic(&FTypeHelper<TYPE>::VMFunction, GetterIndex); \
 		} \
 
 		NDI_PROPERTY_TYPES
diff --git a/Engine/Plugins/FX/Niagara/Source/Niagara/Private/NDISkeletalMesh_BoneSampling.cpp b/Engine/Plugins/FX/Niagara/Source/Niagara/Private/NDISkeletalMesh_BoneSampling.cpp
--- a/Engine/Plugins/FX/Niagara/Source/Niagara/Private/NDISkeletalMesh_BoneSampling.cpp
+++ b/Engine/Plugins/FX/Niagara/Source/Niagara/Private/NDISkeletalMesh_BoneSampling.cpp
@@ -269,58 +269,58 @@ void UNiagaraDataInterfaceSkeletalMesh::BindSkeletonSamplingFunction(const FVMEx
 	else if (BindingInfo.Name == FSkeletalMeshInterfaceHelper::RandomBoneName)
 	{
 		check(BindingInfo.GetNumInputs() == 4 && BindingInfo.GetNumOutputs() == 1);
-		OutFunc = FVMExternalFunction::CreateLambda([this](FVectorVMExternalFunctionContext& Context) { this->RandomBone(Context); });
+		OutFunc.BindUObject(this, &UNiagaraDataInterfaceSkeletalMesh::RandomBone);
 	}
 	else if (BindingInfo.Name == FSkeletalMeshInterfaceHelper::GetBoneCountName)
 	{
 		check(BindingInfo.GetNumInputs() == 1 && BindingInfo.GetNumOutputs() == 1);
-		OutFunc = FVMExternalFunction::CreateLambda([this](FVectorVMExternalFunctionContext& Context) { this->GetBoneCount(Context); });
+		OutFunc.BindUObject(this, &UNiagaraDataInterfaceSkeletalMesh::GetBoneCount);
 	}
 	else if (BindingInfo.Name == FSkeletalMeshInterfaceHelper::GetParentBoneName)
 	{
 		check(BindingInfo.GetNumInputs() == 2 && BindingInfo.GetNumOutputs() == 1);
-		OutFunc = FVMExternalFunction::CreateLambda([this](FVectorVMExternalFunctionContext& Context) { this->GetParentBone(Context); });
+		OutFunc.BindUObject(this, &UNiagaraDataInterfaceSkeletalMesh::GetParentBone);
 	}
 	else if (BindingInfo.Name == FSkeletalMeshInterfaceHelper::RandomFilteredBoneName)
 	{
 		check(BindingInfo.GetNumInputs() == 4 && BindingInfo.GetNumOutputs() == 1);
-		OutFunc = FVMExternalFunction::CreateLambda([this](FVectorVMExternalFunctionContext& Context) { this->RandomFilteredBone(Context); });
+		OutFunc.BindUObject(this, &UNiagaraDataInterfaceSkeletalMesh::RandomFilteredBone);
 	}
 	else if (BindingInfo.Name == FSkeletalMeshInterfaceHelper::GetFilteredBoneCountName)
 	{
 		check(BindingInfo.GetNumInputs() == 1 && BindingInfo.GetNumOutputs() == 1);
-		OutFunc = FVMExternalFunction::CreateLambda([this](FVectorVMExternalFunctionContext& Context) { this->GetFilteredBoneCount(Context); });
+		OutFunc.BindUObject(this, &UNiagaraDataInterfaceSkeletalMesh::GetFilteredBoneCount);
 	}
 	else if (BindingInfo.Name == FSkeletalMeshInterfaceHelper::GetFilteredBoneAtName)
 	{
 		check(BindingInfo.GetNumInputs() == 2 && BindingInfo.GetNumOutputs() == 1);
-		OutFunc = FVMExternalFunction::CreateLambda([this](FVectorVMExternalFunctionContext& Context) { this->GetFilteredBoneAt(Context); });
+		OutFunc.BindUObject(this, &UNiagaraDataInterfaceSkeletalMesh::GetFilteredBoneAt);
 	}
 	else if (BindingInfo.Name == FSkeletalMeshInterfaceHelper::RandomUnfilteredBoneName)
 	{
 		check(BindingInfo.GetNumInputs() == 4 && BindingInfo.GetNumOutputs() == 1);
-		OutFunc = FVMExternalFunction::CreateLambda([this](FVectorVMExternalFunctionContext& Context) { this->RandomUnfilteredBone(Context); });
+		OutFunc.BindUObject(this, &UNiagaraDataInterfaceSkeletalMesh::RandomUnfilteredBone);
 	}
 	else if (BindingInfo.Name == FSkeletalMeshInterfaceHelper::GetUnfilteredBoneCountName)
 	{
 		check(BindingInfo.GetNumInputs() == 1 && BindingInfo.GetNumOutputs() == 1);
-		OutFunc = FVMExternalFunction::CreateLambda([this](FVectorVMExternalFunctionContext& Context) { this->GetUnfilteredBoneCount(Context); });
+		OutFunc.BindUObject(this, &UNiagaraDataInterfaceSkeletalMesh::GetUnfilteredBoneCount);
 	}
 	else if (BindingInfo.Name == FSkeletalMeshInterfaceHelper::GetUnfilteredBoneAtName)
 	{
 		check(BindingInfo.GetNumInputs() == 2 && BindingInfo.GetNumOutputs() == 1);
-		OutFunc = FVMExternalFunction::CreateLambda([this](FVectorVMExternalFunctionContext& Context) { this->GetUnfilteredBoneAt(Context); });
+		OutFunc.BindUObject(this, &UNiagaraDataInterfaceSkeletalMesh::GetUnfilteredBoneAt);
 	}
 	//Socket Functions
 	else if (BindingInfo.Name == FSkeletalMeshInterfaceHelper::RandomFilteredSocketName)
 	{
 		check(BindingInfo.GetNumInputs() == 4 && BindingInfo.GetNumOutputs() == 1);
-		OutFunc = FVMExternalFunction::CreateLambda([this](FVectorVMExternalFunctionContext& Context) { this->RandomFilteredSocket(Context); });
+		OutFunc.BindUObject(this, &UNiagaraDataInterfaceSkeletalMesh::RandomFilteredSocket);
 	}
 	else if (BindingInfo.Name == FSkeletalMeshInterfaceHelper::GetFilteredSocketCountName)
 	{
 		check(BindingInfo.GetNumInputs() == 1 && BindingInfo.GetNumOutputs() == 1);
-		OutFunc = FVMExternalFunction::CreateLambda([this](FVectorVMExternalFunctionContext& Context) { this->GetFilteredSocketCount(Context); });
+		OutFunc.BindUObject(this, &UNiagaraDataInterfaceSkeletalMesh::GetFilteredSocketCount);
 	}
 	else if (BindingInfo.Name == FSkeletalMeshInterfaceHelper::GetFilteredSocketBoneAtName)
 	{
@@ -336,17 +336,17 @@ void UNiagaraDataInterfaceSkeletalMesh::BindSkeletonSamplingFunction(const FVMEx
 	else if (BindingInfo.Name == FSkeletalMeshInterfaceHelper::RandomFilteredSocketOrBoneName)
 	{
 		check(BindingInfo.GetNumInputs() == 4 && BindingInfo.GetNumOutputs() == 1);
-		OutFunc = FVMExternalFunction::CreateLambda([this](FVectorVMExternalFunctionContext& Context) { this->RandomFilteredSocketOrBone(Context); });
+		OutFunc.BindUObject(this, &UNiagaraDataInterfaceSkeletalMesh::RandomFilteredSocketOrBone);
 	}
 	else if (BindingInfo.Name == FSkeletalMeshInterfaceHelper::GetFilteredSocketOrBoneCountName)
 	{
 		check(BindingInfo.GetNumInputs() == 1 && BindingInfo.GetNumOutputs() == 1);
-		OutFunc = FVMExternalFunction::CreateLambda([this](FVectorVMExternalFunctionContext& Context) { this->GetFilteredSocketOrBoneCount(Context); });
+		OutFunc.BindUObject(this, &UNiagaraDataInterfaceSkeletalMesh::GetFilteredSocketOrBoneCount);
 	}
 	else if (BindingInfo.Name == FSkeletalMeshInterfaceHelper::GetFilteredSocketOrBoneAtName)
 	{
 		check(BindingInfo.GetNumInputs() == 2 && BindingInfo.GetNumOutputs() == 1);
-		OutFunc = FVMExternalFunction::CreateLambda([this](FVectorVMExternalFunctionContext& Context) { this->GetFilteredSocketOrBoneBoneAt(Context); });
+		OutFunc.BindUObject(this, &UNiagaraDataInterfaceSkeletalMesh::GetFilteredSocketOrBoneBoneAt);
 	}
 }
 
diff --git a/Engine/Plugins/FX/Niagara/Source/Niagara/Private/NDISkeletalMesh_TriangleSampling.cpp b/Engine/Plugins/FX/Niagara/Source/Niagara/Private/NDISkeletalMesh_TriangleSampling.cpp
--- a/Engine/Plugins/FX/Niagara/Source/Niagara/Private/NDISkeletalMesh_TriangleSampling.cpp
+++ b/Engine/Plugins/FX/Niagara/Source/Niagara/Private/NDISkeletalMesh_TriangleSampling.cpp
@@ -368,7 +368,7 @@ void UNiagaraDataInterfaceSkeletalMesh::BindTriangleSamplingFunction(const FVMEx
 	}
 	else if (BindingInfo.Name == FSkeletalMeshInterfaceHelper::GetTriangleIndicesName)
 	{
-		OutFunc = FVMExternalFunction::CreateLambda([this](FVectorVMExternalFunctionContext& Context) {this->GetTriangleIndices(Context); });
+		OutFunc.BindUObject(this, &UNiagaraDataInterfaceSkeletalMesh::GetTriangleIndices);
 	}
 	else if (BindingInfo.Name == FSkeletalMeshInterfaceHelper::GetSkinnedTriangleDataName)
 	{
@@ -453,17 +453,17 @@ void UNiagaraDataInterfaceSkeletalMesh::BindTriangleSamplingFunction(const FVMEx
 	}
 	else if (BindingInfo.Name == FSkeletalMeshInterfaceHelper::GetTriCoordVerticesName)
 	{
-		OutFunc = FVMExternalFunction::CreateLambda([this](FVectorVMExternalFunctionContext& Context) {this->GetTriangleIndices(Context); });
+		OutFunc.BindUObject(this, &UNiagaraDataInterfaceSkeletalMesh::GetTriangleIndices);
 	}
 	else if (BindingInfo.Name == FSkeletalMeshInterfaceHelper::RandomTriangleName)
 	{
 		check(BindingInfo.GetNumInputs() == 4 && BindingInfo.GetNumOutputs() == 4);
-		OutFunc = FVMExternalFunction::CreateLambda([this](FVectorVMExternalFunctionContext& Context) { this->RandomTriangle(Context); });
+		OutFunc.BindUObject(this, &UNiagaraDataInterfaceSkeletalMesh::RandomTriangle);
 	}
 	else if (BindingInfo.Name == FSkeletalMeshInterfaceHelper::GetTriangleCountName)
 	{
 		check(BindingInfo.GetNumInputs() == 1 && BindingInfo.GetNumOutputs() == 1);
-		OutFunc = FVMExternalFunction::CreateLambda([this](FVectorVMExternalFunctionContext& Context) { this->GetTriangleCount(Context); });
+		OutFunc.BindUObject(this, &UNiagaraDataInterfaceSkeletalMesh::GetTriangleCount);
 	}
 	else if (BindingInfo.Name == FSkeletalMeshInterfaceHelper::RandomFilteredTriangleName)
 	{
diff --git a/Engine/Plugins/FX/Niagara/Source/Niagara/Private/NDISkeletalMesh_VertexSampling.cpp b/Engine/Plugins/FX/Niagara/Source/Niagara/Private/NDISkeletalMesh_VertexSampling.cpp
--- a/Engine/Plugins/FX/Niagara/Source/Niagara/Private/NDISkeletalMesh_VertexSampling.cpp
+++ b/Engine/Plugins/FX/Niagara/Source/Niagara/Private/NDISkeletalMesh_VertexSampling.cpp
@@ -271,17 +271,17 @@ void UNiagaraDataInterfaceSkeletalMesh::BindVertexSamplingFunction(const FVMExte
 	else if (BindingInfo.Name == FSkeletalMeshInterfaceHelper::IsValidVertexName)
 	{
 		check(BindingInfo.GetNumInputs() == 2 && BindingInfo.GetNumOutputs() == 1);
-		OutFunc = FVMExternalFunction::CreateLambda([this](FVectorVMExternalFunctionContext& Context) { this->IsValidVertex(Context); });
+		OutFunc.BindUObject(this, &UNiagaraDataInterfaceSkeletalMesh::IsValidVertex);
 	}
 	else if (BindingInfo.Name == FSkeletalMeshInterfaceHelper::RandomVertexName)
 	{
 		check(BindingInfo.GetNumInputs() == 4 && BindingInfo.GetNumOutputs() == 1);
-		OutFunc = FVMExternalFunction::CreateLambda([this](FVectorVMExternalFunctionContext& Context) { this->RandomVertex(Context); });
+		OutFunc.BindUObject(this, &UNiagaraDataInterfaceSkeletalMesh::RandomVertex);
 	}
 	else if (BindingInfo.Name == FSkeletalMeshInterfaceHelper::GetVertexCountName)
 	{
 		check(BindingInfo.GetNumInputs() == 1 && BindingInfo.GetNumOutputs() == 1);
-		OutFunc = FVMExternalFunction::CreateLambda([this](FVectorVMExternalFunctionContext& Context) { this->GetVertexCount(Context); });
+		OutFunc.BindUObject(this, &UNiagaraDataInterfaceSkeletalMesh::GetVertexCount);
 	}
 	//////////////////////////////////////////////////////////////////////////
 	else if (BindingInfo.Name == FSkeletalMeshInterfaceHelper::IsValidFilteredVertexName)
diff --git a/Engine/Plugins/FX/Niagara/Source/Niagara/Private/NiagaraDataInterfaceNeighborGrid3D.cpp b/Engine/Plugins/FX/Niagara/Source/Niagara/Private/NiagaraDataInterfaceNeighborGrid3D.cpp
--- a/Engine/Plugins/FX/Niagara/Source/Niagara/Private/NiagaraDataInterfaceNeighborGrid3D.cpp
+++ b/Engine/Plugins/FX/Niagara/Source/Niagara/Private/NiagaraDataInterfaceNeighborGrid3D.cpp
@@ -198,17 +198,17 @@ void UNiagaraDataInterfaceNeighborGrid3D::GetVMExternalFunction(const FVMExterna
 	if (BindingInfo.Name == UNiagaraDataInterfaceRWBase::WorldBBoxSizeFunctionName)
 	{
 		check(BindingInfo.GetNumInputs() == 1 && BindingInfo.GetNumOutputs() == 3);
-		OutFunc = FVMExternalFunction::CreateLambda([&](FVectorVMExternalFunctionContext& Context) { GetWorldBBoxSize(Context); });
+		OutFunc.BindUObject(this, &UNiagaraDataInterfaceNeighborGrid3D::GetWorldBBoxSize);
 	}
 	else if (BindingInfo.Name == UNiagaraDataInterfaceRWBase::NumCellsFunctionName)
 	{
 		check(BindingInfo.GetNumInputs() == 1 && BindingInfo.GetNumOutputs() == 3);
-		OutFunc = FVMExternalFunction::CreateLambda([&](FVectorVMExternalFunctionContext& Context) { GetNumCells(Context); });
+		OutFunc.BindUObject(this, &UNiagaraDataInterfaceNeighborGrid3D::GetNumCells);
 	}
 	else if (BindingInfo.Name == MaxNeighborsPerCellFunctionName)
 	{
 		check(BindingInfo.GetNumInputs() == 1 && BindingInfo.GetNumOutputs() == 1);
-		OutFunc = FVMExternalFunction::CreateLambda([&](FVectorVMExternalFunctionContext& Context) { GetMaxNeighborsPerCell(Context); });
+		OutFunc.BindUObject(this, &UNiagaraDataInterfaceNeighborGrid3D::GetMaxNeighborsPerCell);
 	}
 	else if (BindingInfo.Name == SetNumCellsFunctionName)
 	{
diff --git a/Engine/Plugins/FX/Niagara/Source/Niagara/Private/NiagaraDataInterfaceParticleRead.cpp b/Engine/Plugins/FX/Niagara/Source/Niagara/Private/NiagaraDataInterfaceParticleRead.cpp
--- a/Engine/Plugins/FX/Niagara/Source/Niagara/Private/NiagaraDataInterfaceParticleRead.cpp
+++ b/Engine/Plugins/FX/Niagara/Source/Niagara/Private/NiagaraDataInterfaceParticleRead.cpp
@@ -965,7 +965,7 @@ void UNiagaraDataInterfaceParticleRead::GetVMExternalFunction(const FVMExternalF
 	// Misc functionality
 	if ( BindingInfo.Name == GetLocalSpaceName )
 	{	
-		OutFunc = FVMExternalFunction::CreateLambda([this](FVectorVMExternalFunctionContext& Context) { VMGetLocalSpace(Context); });
+		OutFunc.BindUObject(this, &UNiagaraDataInterfaceParticleRead::VMGetLocalSpace);
 		return;
 	}
 
diff --git a/Engine/Plugins/FX/Niagara/Source/Niagara/Private/NiagaraDataInterfaceSkeletalMesh.cpp b/Engine/Plugins/FX/Niagara/Source/Niagara/Private/NiagaraDataInterfaceSkeletalMesh.cpp
--- a/Engine/Plugins/FX/Niagara/Source/Niagara/Private/NiagaraDataInterfaceSkeletalMesh.cpp
+++ b/Engine/Plugins/FX/Niagara/Source/Niagara/Private/NiagaraDataInterfaceSkeletalMesh.cpp
@@ -2438,7 +2438,7 @@ void UNiagaraDataInterfaceSkeletalMesh::GetVMExternalFunction(const FVMExternalF
 	// Bind misc functions
 	if (BindingInfo.Name == NDISkelMeshLocal::NAME_GetPreSkinnedLocalBounds)
 	{
-		OutFunc = FVMExternalFunction::CreateUObject(this, &UNiagaraDataInterfaceSkeletalMesh::VMGetPreSkinnedLocalBounds);
+		OutFunc.BindUObject(this, &UNiagaraDataInterfaceSkeletalMesh::VMGetPreSkinnedLocalBounds);
 		return;
 	}
 
