- Fix incorrect test when getting an array element
- Fix incorrect current value in atomic min / max

#rnx

diff --git a/Engine/Plugins/FX/Niagara/Source/Niagara/Classes/NiagaraDataInterfaceArrayImpl.h b/Engine/Plugins/FX/Niagara/Source/Niagara/Classes/NiagaraDataInterfaceArrayImpl.h
--- a/Engine/Plugins/FX/Niagara/Source/Niagara/Classes/NiagaraDataInterfaceArrayImpl.h
+++ b/Engine/Plugins/FX/Niagara/Source/Niagara/Classes/NiagaraDataInterfaceArrayImpl.h
@@ -607,7 +607,7 @@ struct FNDIArrayProxyImpl : public INDIArrayProxyBase
 		FNDIArrayInstanceData_GameThread<TArrayType>* InstanceData = PerInstanceData_GameThread.IsEmpty() ? nullptr : PerInstanceData_GameThread.CreateConstIterator().Value();
 		FReadArrayRef ArrayRef(Owner, InstanceData);
 
-		if (!ArrayRef.GetArray().IsValidIndex(Index))
+		if (ArrayRef.GetArray().IsValidIndex(Index))
 		{
 			ValueOut = ArrayRef.GetArray()[Index];
 		}
@@ -1307,7 +1307,7 @@ struct FNDIArrayProxyImpl : public INDIArrayProxyBase
 			{
 				TVMArrayType PreviousValue = FNDIArrayImplHelper<TArrayType>::AtomicMin(&ArrayData.GetArray().GetData()[Index], Value);
 				OutPrevValue.SetAndAdvance(PreviousValue);
-				OutCurrValue.SetAndAdvance(PreviousValue + Value);
+				OutCurrValue.SetAndAdvance(FMath::Min(PreviousValue, Value));
 			}
 			else
 			{
@@ -1342,7 +1342,7 @@ struct FNDIArrayProxyImpl : public INDIArrayProxyBase
 			{
 				TVMArrayType PreviousValue = FNDIArrayImplHelper<TArrayType>::AtomicMax(&ArrayData.GetArray().GetData()[Index], Value);
 				OutPrevValue.SetAndAdvance(PreviousValue);
-				OutCurrValue.SetAndAdvance(PreviousValue + Value);
+				OutCurrValue.SetAndAdvance(FMath::Max(PreviousValue, Value));
 			}
 			else
 			{
