- Fix potential aliasing issue

#rnx

diff --git a/Engine/Plugins/FX/Niagara/Source/Niagara/Private/DataInterface/NiagaraDataInterfaceConsoleVariable.cpp b/Engine/Plugins/FX/Niagara/Source/Niagara/Private/DataInterface/NiagaraDataInterfaceConsoleVariable.cpp
--- a/Engine/Plugins/FX/Niagara/Source/Niagara/Private/DataInterface/NiagaraDataInterfaceConsoleVariable.cpp
+++ b/Engine/Plugins/FX/Niagara/Source/Niagara/Private/DataInterface/NiagaraDataInterfaceConsoleVariable.cpp
@@ -135,7 +135,8 @@ namespace NDIConsoleVariableLocal
 		FNDIOutputParam<bool>								SuccessValue(Context);
 		FNDIOutputParam<TVariableType>						OutValue(Context);
 
-		const TVariableType Value = reinterpret_cast<const TVariableType&>(InstanceData_GT->ConsoleVariableData[ValueOffset]);
+		TVariableType Value;
+		FMemory::Memcpy(&Value, &InstanceData_GT->ConsoleVariableData[ValueOffset], sizeof(TVariableType));
 		for (int32 i=0; i < Context.GetNumInstances(); ++i)
 		{
 			SuccessValue.SetAndAdvance(bValid);
