- Test mode for fx.Niagara.SystemSimulation.UpdateOnSpawn to allow us to include the full tick rather than doing a 1ms tick on update

#rnx

diff --git a/Engine/Plugins/FX/Niagara/Source/Niagara/Private/NiagaraSystemSimulation.cpp b/Engine/Plugins/FX/Niagara/Source/Niagara/Private/NiagaraSystemSimulation.cpp
--- a/Engine/Plugins/FX/Niagara/Source/Niagara/Private/NiagaraSystemSimulation.cpp
+++ b/Engine/Plugins/FX/Niagara/Source/Niagara/Private/NiagaraSystemSimulation.cpp
@@ -2186,8 +2186,11 @@ void FNiagaraSystemSimulation::UpdateSystemInstances(FNiagaraSystemSimulationTic
 			UpdateExecContext->BindData(1, UpdateInstanceParameterDataSet, OrigNum, false);
 
 			FNiagaraGlobalParameters UpdateOnSpawnParameters(Context.Instances[0]->GetGlobalParameters());
-			UpdateOnSpawnParameters.EngineDeltaTime = 0.0001f;
-			UpdateOnSpawnParameters.EngineInvDeltaTime = 10000.0f;
+			if (GNiagaraSystemSimulationUpdateOnSpawn == 1)
+			{
+				UpdateOnSpawnParameters.EngineDeltaTime = 0.0001f;
+				UpdateOnSpawnParameters.EngineInvDeltaTime = 10000.0f;
+			}
 
 			FScriptExecutionConstantBufferTable UpdateConstantBufferTable;
 			BuildConstantBufferTable(UpdateOnSpawnParameters, UpdateExecContext, UpdateConstantBufferTable);
