Pass a command list down to the GPU Readback Manager.

diff --git a/Engine/Plugins/FX/Niagara/Source/Niagara/Classes/NiagaraGPUInstanceCountManager.h b/Engine/Plugins/FX/Niagara/Source/Niagara/Classes/NiagaraGPUInstanceCountManager.h
--- a/Engine/Plugins/FX/Niagara/Source/Niagara/Classes/NiagaraGPUInstanceCountManager.h
+++ b/Engine/Plugins/FX/Niagara/Source/Niagara/Classes/NiagaraGPUInstanceCountManager.h
@@ -11,7 +11,7 @@ NiagaraGPUInstanceCountManager.h: GPU particle count handling
 #include "NiagaraDrawIndirect.h"
 #include "RHIGPUReadback.h"
 
-class FRHIGPUMemoryReadback;
+class FRHIGPUBufferReadback;
 class FNiagaraGpuComputeDispatchInterface;
 
 // The number of GPU renderers registered in the instance count manager.
@@ -154,7 +154,7 @@ protected:
 	FRWBuffer CountBuffer;
 	FRWBuffer CulledCountBuffer;
 	TArray<uint32> FreeEntries;
-	FRHIGPUMemoryReadback* CountReadback = nullptr;
+	FRHIGPUBufferReadback* CountReadback = nullptr;
 	int32 CountReadbackSize = 0;
 
 	/** Mirrored copy of CountBuffer for multi-view rendering. */
diff --git a/Engine/Plugins/FX/Niagara/Source/Niagara/Private/NiagaraDataSet.cpp b/Engine/Plugins/FX/Niagara/Source/Niagara/Private/NiagaraDataSet.cpp
--- a/Engine/Plugins/FX/Niagara/Source/Niagara/Private/NiagaraDataSet.cpp
+++ b/Engine/Plugins/FX/Niagara/Source/Niagara/Private/NiagaraDataSet.cpp
@@ -1558,9 +1558,9 @@ void FNiagaraDataBuffer::TransferGPUToCPUImmediate(FRHICommandListImmediate& RHI
 	// Read the count and allocato space for results
 	uint32 GPUNumInstances = 0;
 	{
-		const uint32* DataPtr = static_cast<const uint32*>(RHILockStagingBuffer(CountReadbackBuffer, 0, sizeof(int32)));
+		const uint32* DataPtr = static_cast<const uint32*>(RHICmdList.LockStagingBuffer(CountReadbackBuffer, 0, sizeof(int32)));
 		GPUNumInstances = *DataPtr;
-		RHIUnlockStagingBuffer(CountReadbackBuffer);
+		RHICmdList.UnlockStagingBuffer(CountReadbackBuffer);
 	}
 
 	CPUBuffer->Allocate(GPUNumInstances);
@@ -1574,7 +1574,7 @@ void FNiagaraDataBuffer::TransferGPUToCPUImmediate(FRHICommandListImmediate& RHI
 	if (FloatReadbackBuffer)
 	{
 		const uint32 GPUStride	= FloatStride;
-		const uint8* GPUData	= reinterpret_cast<uint8*>(RHILockStagingBuffer(FloatReadbackBuffer, 0, GPUBufferFloat.NumBytes));
+		const uint8* GPUData	= reinterpret_cast<uint8*>(RHICmdList.LockStagingBuffer(FloatReadbackBuffer, 0, GPUBufferFloat.NumBytes));
 
 		const uint32 CPUStride	= CPUBuffer->GetFloatStride();
 		uint8* CPUData			= CPUBuffer->GetComponentPtrFloat(0);
@@ -1585,13 +1585,13 @@ void FNiagaraDataBuffer::TransferGPUToCPUImmediate(FRHICommandListImmediate& RHI
 			FMemory::Memcpy(CPUData + (CPUStride * i), GPUData + (GPUStride * i), GPUNumInstances * sizeof(float));
 		}
 
-		RHIUnlockStagingBuffer(FloatReadbackBuffer);
+		RHICmdList.UnlockStagingBuffer(FloatReadbackBuffer);
 	}
 
 	if (HalfReadbackBuffer)
 	{
 		const uint32 GPUStride = HalfStride;
-		const uint8* GPUData = reinterpret_cast<uint8*>(RHILockStagingBuffer(HalfReadbackBuffer, 0, GPUBufferHalf.NumBytes));
+		const uint8* GPUData = reinterpret_cast<uint8*>(RHICmdList.LockStagingBuffer(HalfReadbackBuffer, 0, GPUBufferHalf.NumBytes));
 
 		const uint32 CPUStride = CPUBuffer->GetHalfStride();
 		uint8* CPUData = CPUBuffer->GetComponentPtrHalf(0);
@@ -1602,13 +1602,13 @@ void FNiagaraDataBuffer::TransferGPUToCPUImmediate(FRHICommandListImmediate& RHI
 			FMemory::Memcpy(CPUData + (CPUStride * i), GPUData + (GPUStride * i), GPUNumInstances * sizeof(FFloat16));
 		}
 
-		RHIUnlockStagingBuffer(HalfReadbackBuffer);
+		RHICmdList.UnlockStagingBuffer(HalfReadbackBuffer);
 	}
 
 	if (Int32ReadbackBuffer)
 	{
 		const uint32 GPUStride = Int32Stride;
-		const uint8* GPUData = reinterpret_cast<uint8*>(RHILockStagingBuffer(Int32ReadbackBuffer, 0, GPUBufferInt.NumBytes));
+		const uint8* GPUData = reinterpret_cast<uint8*>(RHICmdList.LockStagingBuffer(Int32ReadbackBuffer, 0, GPUBufferInt.NumBytes));
 
 		const uint32 CPUStride = CPUBuffer->GetInt32Stride();
 		uint8* CPUData = CPUBuffer->GetComponentPtrInt32(0);
@@ -1619,7 +1619,7 @@ void FNiagaraDataBuffer::TransferGPUToCPUImmediate(FRHICommandListImmediate& RHI
 			FMemory::Memcpy(CPUData + (CPUStride * i), GPUData + (GPUStride * i), GPUNumInstances * sizeof(int32));
 		}
 
-		RHIUnlockStagingBuffer(Int32ReadbackBuffer);
+		RHICmdList.UnlockStagingBuffer(Int32ReadbackBuffer);
 	}
 }
 
diff --git a/Engine/Plugins/FX/Niagara/Source/Niagara/Private/NiagaraGpuComputeDispatch.cpp b/Engine/Plugins/FX/Niagara/Source/Niagara/Private/NiagaraGpuComputeDispatch.cpp
--- a/Engine/Plugins/FX/Niagara/Source/Niagara/Private/NiagaraGpuComputeDispatch.cpp
+++ b/Engine/Plugins/FX/Niagara/Source/Niagara/Private/NiagaraGpuComputeDispatch.cpp
@@ -442,7 +442,7 @@ void FNiagaraGpuComputeDispatch::ProcessPendingTicksFlush(FRHICommandListImmedia
 			// Early out if we have no pending ticks to process
 			if (MaxPendingTicks == 0)
 			{
-				GpuReadbackManagerPtr->Tick();
+				GpuReadbackManagerPtr->Tick(RHICmdList);
 				return;
 			}
 
@@ -1955,7 +1955,7 @@ void FNiagaraGpuComputeDispatch::PreInitViews(FRDGBuilder& GraphBuilder, bool bA
 	}
 
 	EmptyUAVPoolPtr->Tick();
-	GpuReadbackManagerPtr->Tick();
+	GpuReadbackManagerPtr->Tick(GraphBuilder.RHICmdList);
 #if NIAGARA_COMPUTEDEBUG_ENABLED
 	if ( FNiagaraGpuComputeDebug* GpuComputeDebug = GpuComputeDebugPtr.Get() )
 	{
diff --git a/Engine/Plugins/FX/Niagara/Source/Niagara/Private/NiagaraGpuReadbackManager.cpp b/Engine/Plugins/FX/Niagara/Source/Niagara/Private/NiagaraGpuReadbackManager.cpp
--- a/Engine/Plugins/FX/Niagara/Source/Niagara/Private/NiagaraGpuReadbackManager.cpp
+++ b/Engine/Plugins/FX/Niagara/Source/Niagara/Private/NiagaraGpuReadbackManager.cpp
@@ -16,12 +16,12 @@ FNiagaraGpuReadbackManager::~FNiagaraGpuReadbackManager()
 {
 }
 
-void FNiagaraGpuReadbackManager::Tick()
+void FNiagaraGpuReadbackManager::Tick(FRHICommandListImmediate& RHICmdList)
 {
-	TickInternal(false);
+	TickInternal(RHICmdList, false);
 }
 
-void FNiagaraGpuReadbackManager::TickInternal(bool bAssumeGpuIdle)
+void FNiagaraGpuReadbackManager::TickInternal(FRHICommandListImmediate& RHICmdList, bool bAssumeGpuIdle)
 {
 	check(IsInRenderingThread());
 
@@ -37,7 +37,7 @@ void FNiagaraGpuReadbackManager::TickInternal(bool bAssumeGpuIdle)
 		// Gather data an execute the callback
 		for (const TPair<FStagingBufferRHIRef, uint32>& StagingBuffer : Readback->StagingBuffers)
 		{
-			void* DataPtr = RHILockStagingBuffer(StagingBuffer.Key, Readback->Fence.GetReference(), 0, StagingBuffer.Value);
+			void* DataPtr = RHICmdList.LockStagingBuffer(StagingBuffer.Key, Readback->Fence.GetReference(), 0, StagingBuffer.Value);
 			ensureMsgf(DataPtr || (StagingBuffer.Value == 0), TEXT("NiagaraGpuReadbackManager StagingBuffer returned nullptr and size is %d bytes, bAssumeGpuIdle(%d)"), StagingBuffer.Value, bAssumeGpuIdle);
 			ReadbackData.Emplace(DataPtr, StagingBuffer.Value);
 		}
@@ -46,7 +46,7 @@ void FNiagaraGpuReadbackManager::TickInternal(bool bAssumeGpuIdle)
 
 		for (const TPair<FStagingBufferRHIRef, uint32>& StagingBuffer : Readback->StagingBuffers)
 		{
-			RHIUnlockStagingBuffer(StagingBuffer.Key);
+			RHICmdList.UnlockStagingBuffer(StagingBuffer.Key);
 		}
 
 		ReadbackData.Reset();
@@ -65,7 +65,7 @@ void FNiagaraGpuReadbackManager::WaitCompletion(FRHICommandListImmediate& RHICmd
 	RHICmdList.BlockUntilGPUIdle();
 
 	// Perform a tick which will flush everything
-	TickInternal(true);
+	TickInternal(RHICmdList, true);
 }
 
 void FNiagaraGpuReadbackManager::EnqueueReadback(FRDGBuilder& GraphBuilder, FRDGBufferRef Buffer, FCompletionCallback Callback)
diff --git a/Engine/Plugins/FX/Niagara/Source/Niagara/Public/NiagaraGpuReadbackManager.h b/Engine/Plugins/FX/Niagara/Source/Niagara/Public/NiagaraGpuReadbackManager.h
--- a/Engine/Plugins/FX/Niagara/Source/Niagara/Public/NiagaraGpuReadbackManager.h
+++ b/Engine/Plugins/FX/Niagara/Source/Niagara/Public/NiagaraGpuReadbackManager.h
@@ -45,11 +45,11 @@ public:
 	NIAGARA_API ~FNiagaraGpuReadbackManager();
 
 	// Tick call which polls for completed readbacks
-	NIAGARA_API void Tick();
+	NIAGARA_API void Tick(FRHICommandListImmediate& RHICmdList);
 	
 private:
 	// Internal tick impl
-	void TickInternal(bool bAssumeGpuIdle);
+	void TickInternal(FRHICommandListImmediate& RHICmdList, bool bAssumeGpuIdle);
 
 public:
 	// Wait for all pending readbacks to complete
