- Make LW scaling consistent with standard emitters

#FailedValidations Preflight

diff --git a/Engine/Plugins/FX/Niagara/Source/Niagara/Internal/Stateless/NiagaraStatelessSetShaderParameterContext.h b/Engine/Plugins/FX/Niagara/Source/Niagara/Internal/Stateless/NiagaraStatelessSetShaderParameterContext.h
--- a/Engine/Plugins/FX/Niagara/Source/Niagara/Internal/Stateless/NiagaraStatelessSetShaderParameterContext.h
+++ b/Engine/Plugins/FX/Niagara/Source/Niagara/Internal/Stateless/NiagaraStatelessSetShaderParameterContext.h
@@ -63,6 +63,7 @@ public:
 	}
 
 	void TransformVectorRangeToScaleBias(const FNiagaraStatelessRangeVector3& Range, ENiagaraCoordinateSpace SourceSpace, FVector3f& OutScale, FVector3f& OutBias) const;
+	void TransformVectorNoScaleRangeToScaleBias(const FNiagaraStatelessRangeVector3& Range, ENiagaraCoordinateSpace SourceSpace, FVector3f& OutScale, FVector3f& OutBias) const;
 	void TransformPositionRangeToScaleBias(const FNiagaraStatelessRangeVector3& Range, ENiagaraCoordinateSpace SourceSpace, FVector3f& OutScale, FVector3f& OutBias) const;
 	void TransformRotationRangeToScaleBias(const FNiagaraStatelessRangeVector3& Range, ENiagaraCoordinateSpace SourceSpace, FVector3f& OutScale, FVector3f& OutBias) const;
 	FVector3f TransformPositionRangeToValue(const FNiagaraStatelessRangeVector3& Range, ENiagaraCoordinateSpace SourceSpace) const;
diff --git a/Engine/Plugins/FX/Niagara/Source/Niagara/Private/Stateless/Modules/NiagaraStatelessModule_SolveVelocitiesAndForces.cpp b/Engine/Plugins/FX/Niagara/Source/Niagara/Private/Stateless/Modules/NiagaraStatelessModule_SolveVelocitiesAndForces.cpp
--- a/Engine/Plugins/FX/Niagara/Source/Niagara/Private/Stateless/Modules/NiagaraStatelessModule_SolveVelocitiesAndForces.cpp
+++ b/Engine/Plugins/FX/Niagara/Source/Niagara/Private/Stateless/Modules/NiagaraStatelessModule_SolveVelocitiesAndForces.cpp
@@ -197,21 +197,21 @@ void UNiagaraStatelessModule_SolveVelocitiesAndForces::SetShaderParameters(const
 	SetShaderParameterContext.ConvertRangeToScaleBias(PhysicsData.MassRange, Parameters->SolveVelocitiesAndForces_MassScale, Parameters->SolveVelocitiesAndForces_MassBias);
 	SetShaderParameterContext.ConvertRangeToScaleBias(PhysicsData.DragRange, Parameters->SolveVelocitiesAndForces_DragScale, Parameters->SolveVelocitiesAndForces_DragBias);
 
-	SetShaderParameterContext.TransformVectorRangeToScaleBias(PhysicsData.VelocityRange, PhysicsData.VelocityCoordinateSpace, Parameters->SolveVelocitiesAndForces_VelocityScale, Parameters->SolveVelocitiesAndForces_VelocityBias);
+	SetShaderParameterContext.TransformVectorNoScaleRangeToScaleBias(PhysicsData.VelocityRange, PhysicsData.VelocityCoordinateSpace, Parameters->SolveVelocitiesAndForces_VelocityScale, Parameters->SolveVelocitiesAndForces_VelocityBias);
 	const float VelocityScale = SetShaderParameterContext.ConvertRangeToValue(PhysicsData.LinearVelocityScale);
 	Parameters->SolveVelocitiesAndForces_VelocityScale *= VelocityScale;
 	Parameters->SolveVelocitiesAndForces_VelocityBias  *= VelocityScale;
 
-	SetShaderParameterContext.TransformVectorRangeToScaleBias(PhysicsData.WindRange, PhysicsData.WindCoordinateSpace, Parameters->SolveVelocitiesAndForces_WindScale, Parameters->SolveVelocitiesAndForces_WindBias);
+	SetShaderParameterContext.TransformVectorNoScaleRangeToScaleBias(PhysicsData.WindRange, PhysicsData.WindCoordinateSpace, Parameters->SolveVelocitiesAndForces_WindScale, Parameters->SolveVelocitiesAndForces_WindBias);
 	{
 		FVector3f AccelerationScale, AccelerationBias;
 		FVector3f GravityScale, GravityBias;
 		SetShaderParameterContext.ConvertRangeToScaleBias(PhysicsData.AccelerationRange, AccelerationScale, AccelerationBias);
 		SetShaderParameterContext.ConvertRangeToScaleBias(PhysicsData.GravityRange, GravityScale, GravityBias);
-		Parameters->SolveVelocitiesAndForces_AccelerationScale	= SpaceTransforms.TransformVector(PhysicsData.AccelerationCoordinateSpace, AccelerationScale);
-		Parameters->SolveVelocitiesAndForces_AccelerationBias	= SpaceTransforms.TransformVector(PhysicsData.AccelerationCoordinateSpace, AccelerationBias);
-		Parameters->SolveVelocitiesAndForces_GravityScale	= SpaceTransforms.TransformVector(ENiagaraCoordinateSpace::World, GravityScale);
-		Parameters->SolveVelocitiesAndForces_GravityBias	= SpaceTransforms.TransformVector(ENiagaraCoordinateSpace::World, GravityBias);
+		Parameters->SolveVelocitiesAndForces_AccelerationScale	= SpaceTransforms.TransformVectorNoScale(PhysicsData.AccelerationCoordinateSpace, AccelerationScale);
+		Parameters->SolveVelocitiesAndForces_AccelerationBias	= SpaceTransforms.TransformVectorNoScale(PhysicsData.AccelerationCoordinateSpace, AccelerationBias);
+		Parameters->SolveVelocitiesAndForces_GravityScale	= SpaceTransforms.TransformVectorNoScale(ENiagaraCoordinateSpace::World, GravityScale);
+		Parameters->SolveVelocitiesAndForces_GravityBias	= SpaceTransforms.TransformVectorNoScale(ENiagaraCoordinateSpace::World, GravityBias);
 	}
 
 	Parameters->SolveVelocitiesAndForces_ConeVelocityEnabled = PhysicsData.bConeVelocity ? 1 : 0;
diff --git a/Engine/Plugins/FX/Niagara/Source/Niagara/Private/Stateless/NiagaraStatelessSetShaderParameterContext.cpp b/Engine/Plugins/FX/Niagara/Source/Niagara/Private/Stateless/NiagaraStatelessSetShaderParameterContext.cpp
--- a/Engine/Plugins/FX/Niagara/Source/Niagara/Private/Stateless/NiagaraStatelessSetShaderParameterContext.cpp
+++ b/Engine/Plugins/FX/Niagara/Source/Niagara/Private/Stateless/NiagaraStatelessSetShaderParameterContext.cpp
@@ -26,6 +26,20 @@ void FNiagaraStatelessSetShaderParameterContext::TransformVectorRangeToScaleBias
 	}
 }
 
+void FNiagaraStatelessSetShaderParameterContext::TransformVectorNoScaleRangeToScaleBias(const FNiagaraStatelessRangeVector3& Range, ENiagaraCoordinateSpace SourceSpace, FVector3f& OutScale, FVector3f& OutBias) const
+{
+	if (Range.ParameterOffset == INDEX_NONE)
+	{
+		OutScale = SpaceTransforms.TransformVectorNoScale(SourceSpace, Range.GetScale());
+		OutBias = SpaceTransforms.TransformVectorNoScale(SourceSpace, Range.Min);
+	}
+	else
+	{
+		OutScale = FNiagaraStatelessRangeDefaultValue<FVector3f>::Zero();
+		OutBias = SpaceTransforms.TransformVectorNoScale(SourceSpace, GetRendererParameterValue(Range.ParameterOffset, Range.Min));
+	}
+}
+
 void FNiagaraStatelessSetShaderParameterContext::TransformPositionRangeToScaleBias(const FNiagaraStatelessRangeVector3& Range, ENiagaraCoordinateSpace SourceSpace, FVector3f& OutScale, FVector3f& OutBias) const
 {
 	if (Range.ParameterOffset == INDEX_NONE)
