- When reading particle various remove any 'Particle.' from the start

diff --git a/Engine/Plugins/FX/Niagara/Source/Niagara/Private/NiagaraDebugHud.cpp b/Engine/Plugins/FX/Niagara/Source/Niagara/Private/NiagaraDebugHud.cpp
--- a/Engine/Plugins/FX/Niagara/Source/Niagara/Private/NiagaraDebugHud.cpp
+++ b/Engine/Plugins/FX/Niagara/Source/Niagara/Private/NiagaraDebugHud.cpp
@@ -6,6 +6,7 @@
 #include "DrawDebugHelpers.h"
 #include "GameFramework/Pawn.h"
 #include "NiagaraComponent.h"
+#include "NiagaraConstants.h"
 #include "NiagaraCullProxyComponent.h"
 #include "NiagaraDataSetDebugAccessor.h"
 #include "NiagaraDataSetReadback.h"
@@ -520,6 +521,20 @@ namespace NiagaraDebugLocal
 				const TConstArrayView<FNiagaraEmitterHandle> EmitterHandles = NiagaraSystem->GetEmitterHandles();
 				const TArray<TSharedRef<const FNiagaraEmitterCompiledData>>& AllStatefulEmittersCompiledData = NiagaraSystem->GetEmitterCompiledData();
 
+				TArray<FNiagaraDebugHUDVariable> LocalParticlesVariables;
+				LocalParticlesVariables.Reserve(Settings.ParticlesVariables.Num());
+				Algo::TransformIf(
+					Settings.ParticlesVariables,
+					LocalParticlesVariables,
+					[](const FNiagaraDebugHUDVariable& Variable) { return Variable.bEnabled; },
+					[](const FNiagaraDebugHUDVariable& Variable)
+					{
+						FNiagaraDebugHUDVariable NewVariable = Variable;
+						NewVariable.Name.RemoveFromStart(FNiagaraConstants::ParticleAttributeNamespaceString);
+						return NewVariable;
+					}
+				);				
+
 				CachedVariables->ParticleVariables.AddDefaulted(EmitterHandles.Num());
 				CachedVariables->ParticlePositionAccessors.AddDefaulted(EmitterHandles.Num());
 				for (int32 iEmitter=0; iEmitter < EmitterHandles.Num(); ++iEmitter)
@@ -554,7 +569,7 @@ namespace NiagaraDebugLocal
 
 					FindVariablesByWildcard(
 						EmitterCompiledData->Variables,
-						Settings.ParticlesVariables,
+						LocalParticlesVariables,
 						[&](const FNiagaraVariable& Variable) { CachedVariables->ParticleVariables[iEmitter].AddDefaulted_GetRef().Init(*EmitterCompiledData, Variable.GetName()); }
 					);
 
