- Conditionally remove Niagara debug shaders from cooked builds

diff --git a/Engine/Plugins/FX/Niagara/Source/Niagara/Public/NiagaraCommon.h b/Engine/Plugins/FX/Niagara/Source/Niagara/Public/NiagaraCommon.h
--- a/Engine/Plugins/FX/Niagara/Source/Niagara/Public/NiagaraCommon.h
+++ b/Engine/Plugins/FX/Niagara/Source/Niagara/Public/NiagaraCommon.h
@@ -52,10 +52,9 @@ using FNiagaraShaderMapRef = TRefCountPtr<FNiagaraShaderMap>;
 
 #define NIAGARA_VERBOSE_LOGGING_ENABLED (!NO_LOGGING && !UE_BUILD_SHIPPING)
 
-
-#ifndef NIAGARA_COMPUTEDEBUG_ENABLED
-	#define NIAGARA_COMPUTEDEBUG_ENABLED WITH_EDITOR
-#endif
+#define NIAGARA_COMPUTEDEBUG_ENABLED_IN_EDITOR	1
+#define NIAGARA_COMPUTEDEBUG_ENABLED_IN_COOKED	0
+#define NIAGARA_COMPUTEDEBUG_ENABLED ((WITH_EDITOR && NIAGARA_COMPUTEDEBUG_ENABLED_IN_EDITOR) || (!WITH_EDITOR && NIAGARA_COMPUTEDEBUG_ENABLED_IN_COOKED))
 
 #define WITH_NIAGARA_GPU_PROFILER_EDITOR (WITH_EDITOR && STATS)
 #define WITH_NIAGARA_GPU_PROFILER ((WITH_PARTICLE_PERF_STATS && !UE_BUILD_SHIPPING) || WITH_NIAGARA_GPU_PROFILER_EDITOR)
diff --git a/Engine/Plugins/FX/Niagara/Source/NiagaraShader/Private/NiagaraDebugShaders.cpp b/Engine/Plugins/FX/Niagara/Source/NiagaraShader/Private/NiagaraDebugShaders.cpp
--- a/Engine/Plugins/FX/Niagara/Source/NiagaraShader/Private/NiagaraDebugShaders.cpp
+++ b/Engine/Plugins/FX/Niagara/Source/NiagaraShader/Private/NiagaraDebugShaders.cpp
@@ -1,6 +1,7 @@
 // Copyright Epic Games, Inc. All Rights Reserved.
 
 #include "NiagaraDebugShaders.h"
+#include "NiagaraCommon.h"
 
 #include "FXRenderingUtils.h"
 #include "RHI.h"
@@ -41,6 +42,17 @@ static FAutoConsoleVariableRef CVarNiagaraGpuComputeDebug_OccludedLineColorScale
 
 //////////////////////////////////////////////////////////////////////////
 
+namespace NiagaraDebugShaders
+{
+	static bool ShouldCompilePermutation(const FGlobalShaderPermutationParameters& Parameters)
+	{
+		const bool bTargetIsEditor = EnumHasAllFlags(Parameters.Flags, EShaderPermutationFlags::HasEditorOnlyData);
+		return (bTargetIsEditor && NIAGARA_COMPUTEDEBUG_ENABLED_IN_EDITOR) || (!bTargetIsEditor && NIAGARA_COMPUTEDEBUG_ENABLED_IN_COOKED);
+	}
+}
+
+//////////////////////////////////////////////////////////////////////////
+
 class FNiagaraVisualizeTexturePS : public FGlobalShader
 {
 	DECLARE_GLOBAL_SHADER(FNiagaraVisualizeTexturePS);
@@ -53,7 +65,7 @@ class FNiagaraVisualizeTexturePS : public FGlobalShader
 
 	static bool ShouldCompilePermutation(const FGlobalShaderPermutationParameters& Parameters)
 	{
-		return true;
+		return NiagaraDebugShaders::ShouldCompilePermutation(Parameters);
 	}
 
 	BEGIN_SHADER_PARAMETER_STRUCT(FParameters, )
@@ -113,7 +125,7 @@ class FNiagaraDebugDrawLineVS : public FGlobalShader
 
 	static bool ShouldCompilePermutation(const FGlobalShaderPermutationParameters& Parameters)
 	{
-		return true;
+		return NiagaraDebugShaders::ShouldCompilePermutation(Parameters);
 	}
 
 	static void ModifyCompilationEnvironment(const FGlobalShaderPermutationParameters& Parameters, FShaderCompilerEnvironment& OutEnvironment)
@@ -135,7 +147,7 @@ class FNiagaraDebugDrawLinePS : public FGlobalShader
 
 	static bool ShouldCompilePermutation(const FGlobalShaderPermutationParameters& Parameters)
 	{
-		return true;
+		return NiagaraDebugShaders::ShouldCompilePermutation(Parameters);
 	}
 
 	static void ModifyCompilationEnvironment(const FGlobalShaderPermutationParameters& Parameters, FShaderCompilerEnvironment& OutEnvironment)
