- Reduce ribbon compute shader permutations from 127 to 33
- Remove FRibbonLinkIsFloat, FRibbonHasTwist, FRibbonWantsConstantTessellation, FRibbonWantsAutomaticTessellation & FRibbonHasHighSliceComplexity
- This results in a small performance loss in some cases and neutral in others

#rnx

diff --git a/Engine/Plugins/FX/Niagara/Shaders/Private/Ribbons/NiagaraRibbonAggregationApply.usf b/Engine/Plugins/FX/Niagara/Shaders/Private/Ribbons/NiagaraRibbonAggregationApply.usf
--- a/Engine/Plugins/FX/Niagara/Shaders/Private/Ribbons/NiagaraRibbonAggregationApply.usf
+++ b/Engine/Plugins/FX/Niagara/Shaders/Private/Ribbons/NiagaraRibbonAggregationApply.usf
@@ -13,9 +13,7 @@ RWBuffer<float> OutputTangentsAndDistances;
 RWBuffer<uint> OutputMultiRibbonIndices;
 #endif
 
-#if RIBBONS_WANTS_AUTOMATIC_TESSELLATION
 RWBuffer<float> OutputTessellationStats;
-#endif
 
 RWBuffer<uint> OutputSegments;
 
@@ -39,20 +37,18 @@ void RibbonAggregationApply(uint3 DispatchThreadId : SV_DispatchThreadID)
 		OutputMultiRibbonIndices[RawParticleID] = InputAccumulation[RawParticleID].MultiRibbonCount - 1;
 #endif
 
-#if RIBBONS_WANTS_AUTOMATIC_TESSELLATION
-#if RIBBON_HAS_TWIST
-		const uint TessellationBaseOffset = RawParticleID * 5;
-#else
-		const uint TessellationBaseOffset = RawParticleID * 3;
-#endif // RIBBON_HAS_TWIST
-
-		OutputTessellationStats[TessellationBaseOffset + 0] = InputAccumulation[RawParticleID].TessTotalLength;
-		OutputTessellationStats[TessellationBaseOffset + 1] = InputAccumulation[RawParticleID].TessAvgSegmentLength;
-		OutputTessellationStats[TessellationBaseOffset + 2] = InputAccumulation[RawParticleID].TessAvgSegmentAngle;
-#if RIBBON_HAS_TWIST
-		OutputTessellationStats[TessellationBaseOffset + 3] = InputAccumulation[RawParticleID].TessTwistAvgAngle;
-		OutputTessellationStats[TessellationBaseOffset + 4] = InputAccumulation[RawParticleID].TessTwistAvgWidth;
-#endif // RIBBON_HAS_TWIST
-#endif // RIBBONS_WANTS_AUTOMATIC_TESSELLATION
+		if (RibbonWantsAutomaticTessellation())
+		{
+			const uint TessellationBaseOffset = RawParticleID * GetTessellationDataStride();
+
+			OutputTessellationStats[TessellationBaseOffset + 0] = InputAccumulation[RawParticleID].TessTotalLength;
+			OutputTessellationStats[TessellationBaseOffset + 1] = InputAccumulation[RawParticleID].TessAvgSegmentLength;
+			OutputTessellationStats[TessellationBaseOffset + 2] = InputAccumulation[RawParticleID].TessAvgSegmentAngle;
+			if (RibbonHasTwist())
+			{
+				OutputTessellationStats[TessellationBaseOffset + 3] = InputAccumulation[RawParticleID].TessTwistAvgAngle;
+				OutputTessellationStats[TessellationBaseOffset + 4] = InputAccumulation[RawParticleID].TessTwistAvgWidth;
+			}
+		}
 	}
 }
diff --git a/Engine/Plugins/FX/Niagara/Shaders/Private/Ribbons/NiagaraRibbonAggregationStep.usf b/Engine/Plugins/FX/Niagara/Shaders/Private/Ribbons/NiagaraRibbonAggregationStep.usf
--- a/Engine/Plugins/FX/Niagara/Shaders/Private/Ribbons/NiagaraRibbonAggregationStep.usf
+++ b/Engine/Plugins/FX/Niagara/Shaders/Private/Ribbons/NiagaraRibbonAggregationStep.usf
@@ -46,15 +46,18 @@ void RibbonAggregationStep(uint3 DispatchThreadId : SV_DispatchThreadID)
 #if HAS_RIBBON_ID
 			OutputAccumulation[RawParticleID].MultiRibbonCount = InputAccumulation[RawParticleID].MultiRibbonCount + InputAccumulation[RawSourceParticleID].MultiRibbonCount;
 #endif // HAS_RIBBON_ID
-#if RIBBONS_WANTS_AUTOMATIC_TESSELLATION
-			OutputAccumulation[RawParticleID].TessTotalLength = InputAccumulation[RawParticleID].TessTotalLength + InputAccumulation[RawSourceParticleID].TessTotalLength;
-			OutputAccumulation[RawParticleID].TessAvgSegmentLength = InputAccumulation[RawParticleID].TessAvgSegmentLength + InputAccumulation[RawSourceParticleID].TessAvgSegmentLength;
-			OutputAccumulation[RawParticleID].TessAvgSegmentAngle = InputAccumulation[RawParticleID].TessAvgSegmentAngle + InputAccumulation[RawSourceParticleID].TessAvgSegmentAngle;
-#if RIBBON_HAS_TWIST
-			OutputAccumulation[RawParticleID].TessTwistAvgAngle = InputAccumulation[RawParticleID].TessTwistAvgAngle + InputAccumulation[RawSourceParticleID].TessTwistAvgAngle;
-			OutputAccumulation[RawParticleID].TessTwistAvgWidth = InputAccumulation[RawParticleID].TessTwistAvgWidth + InputAccumulation[RawSourceParticleID].TessTwistAvgWidth;
-#endif // RIBBON_HAS_TWIST
-#endif // RIBBONS_WANTS_AUTOMATIC_TESSELLATION
+
+			if (RibbonWantsAutomaticTessellation())
+			{
+				OutputAccumulation[RawParticleID].TessTotalLength = InputAccumulation[RawParticleID].TessTotalLength + InputAccumulation[RawSourceParticleID].TessTotalLength;
+				OutputAccumulation[RawParticleID].TessAvgSegmentLength = InputAccumulation[RawParticleID].TessAvgSegmentLength + InputAccumulation[RawSourceParticleID].TessAvgSegmentLength;
+				OutputAccumulation[RawParticleID].TessAvgSegmentAngle = InputAccumulation[RawParticleID].TessAvgSegmentAngle + InputAccumulation[RawSourceParticleID].TessAvgSegmentAngle;
+				if (RibbonHasTwist())
+				{
+					OutputAccumulation[RawParticleID].TessTwistAvgAngle = InputAccumulation[RawParticleID].TessTwistAvgAngle + InputAccumulation[RawSourceParticleID].TessTwistAvgAngle;
+					OutputAccumulation[RawParticleID].TessTwistAvgWidth = InputAccumulation[RawParticleID].TessTwistAvgWidth + InputAccumulation[RawSourceParticleID].TessTwistAvgWidth;
+				}
+			}
 		}
 		else
 		{
diff --git a/Engine/Plugins/FX/Niagara/Shaders/Private/Ribbons/NiagaraRibbonCommon.ush b/Engine/Plugins/FX/Niagara/Shaders/Private/Ribbons/NiagaraRibbonCommon.ush
--- a/Engine/Plugins/FX/Niagara/Shaders/Private/Ribbons/NiagaraRibbonCommon.ush
+++ b/Engine/Plugins/FX/Niagara/Shaders/Private/Ribbons/NiagaraRibbonCommon.ush
@@ -13,17 +13,11 @@
 #define PACKED_PER_RIBBON_STARTPARTICLEOFFSET 6
 #define PACKED_PER_RIBBON_ENDPARTICLEOFFSET 7
 
-#if RIBBONS_WANTS_AUTOMATIC_TESSELLATION
-#define VERTEX_GEN_OUTPUT_DATA_STRIDE 15
-#else
-#define VERTEX_GEN_OUTPUT_DATA_STRIDE 5
-#endif
-
 #define VERTEX_GEN_OUTPUT_DATA_FINALIZATION_INDIRECT_ARGS_OFFSET 0
 #define VERTEX_GEN_OUTPUT_DATA_TOTAL_NUM_SEGMENTS_OFFSET 3
 #define VERTEX_GEN_OUTPUT_DATA_TOTAL_NUM_RIBBONS_OFFSET 4
 
-#if RIBBONS_WANTS_AUTOMATIC_TESSELLATION
+// VERTEX_GEN_OUTPUT_DATA_TESSELLATION is only valid when RibbonTessellationMode is none zero
 #define VERTEX_GEN_OUTPUT_DATA_TESSELLATION_ANGLE 5
 #define VERTEX_GEN_OUTPUT_DATA_TESSELLATION_CURVATURE 6
 #define VERTEX_GEN_OUTPUT_DATA_TESSELLATION_TWIST_ANGLE 7
@@ -35,7 +29,6 @@
 #define VERTEX_GEN_OUTPUT_DATA_TESSELLATION_CURRENT_FRAME_AVERAGE_SEGMENT_ANGLE 12
 #define VERTEX_GEN_OUTPUT_DATA_TESSELLATION_CURRENT_FRAME_AVERAGE_TWIST_ANGLE 13
 #define VERTEX_GEN_OUTPUT_DATA_TESSELLATION_CURRENT_FRAME_AVERAGE_WIDTH 14
-#endif
 
 // This layout must match FNiagaraRibbonIndirectDrawBufferLayout
 #define INDEX_GEN_EXECUTE_INDIRECT_OFFSET			0
@@ -90,10 +83,11 @@ int U1OverrideDataOffset;
 int V1RangeOverrideDataOffset;
 int U0DistributionMode;
 int U1DistributionMode;
-int RibbonLinkOrderDataOffset;
-
 
+int RibbonLinkOrderIsFloat;
+int RibbonLinkOrderDataOffset;
 
+int RibbonTessellationMode;
 
 // Gets the num particles from the indirect buffer if available, or if not from the fallback direct input
 int GetTotalNumParticles()
@@ -132,8 +126,6 @@ int2 GetInt2(int RegisterIdx, uint InstanceID) { return NiagaraGetInt2(RegisterI
 int SafeGetInt(int RegisterIdx, uint InstanceID, int Default) { return NiagaraSafeGetInt(RegisterIdx, InstanceID, Default); }
 int2 SafeGetInt2(int RegisterIdx, uint InstanceID, int2 Default) { return NiagaraSafeGetInt2(RegisterIdx, InstanceID, Default); }
 
-
-
 #define HAS_RIBBON_ID (RIBBONID_IS_NIAGARAID || RIBBONID_IS_INT)
 
 #define FNiagaraID int2
@@ -229,17 +221,43 @@ static const FNiagaraID InvalidID = { INT_MIN, INT_MIN };
 
 #endif
 
-#if RIBBONLINK_IS_FLOAT
-	float GetRibbonLinkID(int InstanceID)
+float GetRibbonLinkID(int InstanceID)
+{
+	if ( RibbonLinkOrderIsFloat != 0 )
 	{
 		return GetFloat(RibbonLinkOrderDataOffset, InstanceID);
 	}
-#else
-	int GetRibbonLinkID(int InstanceID)
+	else
 	{
 		return GetInt(RibbonLinkOrderDataOffset, InstanceID);
 	}
-#endif
+}
+
+bool RibbonHasTwist()
+{
+	return TwistDataOffset != -1;
+}
+
+bool RibbonWantsAutomaticTessellation()
+{
+	return RibbonTessellationMode == ENiagaraGpuTessellationMode::Automatic;	
+}
+
+bool RibbonWantsConstantTessellation()
+{
+	return RibbonTessellationMode == ENiagaraGpuTessellationMode::Constant;
+}
+
+uint GetTessellationDataStride()
+{
+	return RibbonHasTwist() ? NIAGARA_RIBBON_TESSELLATION_FULL_STATS_NUM_ELEMENTS : NIAGARA_RIBBON_TESSELLATION_NO_TWIST_STATS_NUM_ELEMENTS;
+}
+
+uint GetVertexGenOutputDataStride()
+{
+	// See FNiagaraRibbonCommandBufferLayout
+	return RibbonTessellationMode == ENiagaraGpuTessellationMode::Automatic ? VERTEX_GEN_OUTPUT_DATA_AUTO_TESSELLATION_NUM_ELEMENTS : VERTEX_GEN_OUTPUT_DATA_NUM_ELEMENTS;
+}
 
 // max absolute error 9.0x10^-3
 // Eberly's polynomial degree 1 - respect bounds
diff --git a/Engine/Plugins/FX/Niagara/Shaders/Private/Ribbons/NiagaraRibbonGenerateIndices.usf b/Engine/Plugins/FX/Niagara/Shaders/Private/Ribbons/NiagaraRibbonGenerateIndices.usf
--- a/Engine/Plugins/FX/Niagara/Shaders/Private/Ribbons/NiagaraRibbonGenerateIndices.usf
+++ b/Engine/Plugins/FX/Niagara/Shaders/Private/Ribbons/NiagaraRibbonGenerateIndices.usf
@@ -90,11 +90,6 @@ void GenerateIndices(uint3 DispatchThreadId : SV_DispatchThreadID)
 	    
 			uint TriangleId = 0;
 
-#if RIBBON_HAS_HIGH_SLICE_COMPLEXITY
-			LOOP
-#else
-			UNROLL_N(LOOP_UNROLL_SIZE)
-#endif
 			for (; TriangleId < FlipGeometryIndex; TriangleId += 2)
 			{
 				const uint FirstIndex = TriangleToVertexIds[TriangleId];
@@ -110,11 +105,6 @@ void GenerateIndices(uint3 DispatchThreadId : SV_DispatchThreadID)
 				IndicesOffset += 6;
 			}
 			
-#if RIBBON_HAS_HIGH_SLICE_COMPLEXITY
-			LOOP
-#else
-			UNROLL_N(LOOP_UNROLL_SIZE)
-#endif
 			for (; TriangleId < TriangleToVertexIdsCount; TriangleId += 2)
 			{
 				const uint FirstIndex = TriangleToVertexIds[TriangleId];
diff --git a/Engine/Plugins/FX/Niagara/Shaders/Private/Ribbons/NiagaraRibbonInitializeIndices.usf b/Engine/Plugins/FX/Niagara/Shaders/Private/Ribbons/NiagaraRibbonInitializeIndices.usf
--- a/Engine/Plugins/FX/Niagara/Shaders/Private/Ribbons/NiagaraRibbonInitializeIndices.usf
+++ b/Engine/Plugins/FX/Niagara/Shaders/Private/Ribbons/NiagaraRibbonInitializeIndices.usf
@@ -63,77 +63,72 @@ void InitializeIndices(uint3 DispatchThreadId : SV_DispatchThreadID)
 		return;
 	}
 
-#if RIBBONS_WANTS_AUTOMATIC_TESSELLATION
-	
-	bool bUseConstantFactor = false;
-	uint TessellationFactor = GNiagaraRibbonMaxTessellation;
-	float TessellationMinAngle = GNiagaraRibbonTessellationAngle;
-	float ScreenPercentage = GNiagaraRibbonTessellationScreenPercentage;
-
-	switch (TessellationMode)
-	{
-	case TESSELLATIONMODE_AUTOMATIC:
-		break;
-	case TESSELLATIONMODE_CUSTOM:
-		TessellationFactor = min(TessellationFactor, CustomTessellationFactor); // Don't allow factors bigger than the platform limit.
-		bUseConstantFactor = bCustomUseConstantFactor == 1;
-		TessellationMinAngle = CustomTessellationMinAngle;
-		ScreenPercentage = bCustomUseScreenSpace == 1 && !bUseConstantFactor ? GNiagaraRibbonTessellationScreenPercentage : 0.f;
-		break;
-	case TESSELLATIONMODE_DISABLED:
-		TessellationFactor = 1;
-		break;
-	default:
-		break;
-	}
-	
+	const uint VertexGenerationResultsOffset = VertexGenerationResultsIndex * GetVertexGenOutputDataStride();
 	uint SegmentTessellation = 1;
-
-	float TessellationAngle = asfloat(VertexGenerationResults[VertexGenerationResultsIndex * VERTEX_GEN_OUTPUT_DATA_STRIDE + VERTEX_GEN_OUTPUT_DATA_TESSELLATION_ANGLE]);
-	float TessellationCurvature = asfloat(VertexGenerationResults[VertexGenerationResultsIndex * VERTEX_GEN_OUTPUT_DATA_STRIDE + VERTEX_GEN_OUTPUT_DATA_TESSELLATION_CURVATURE]);
-	float TessellationTwistAngle = asfloat(VertexGenerationResults[VertexGenerationResultsIndex * VERTEX_GEN_OUTPUT_DATA_STRIDE + VERTEX_GEN_OUTPUT_DATA_TESSELLATION_TWIST_ANGLE]);
-	float TessellationTwistCurvature = asfloat(VertexGenerationResults[VertexGenerationResultsIndex * VERTEX_GEN_OUTPUT_DATA_STRIDE + VERTEX_GEN_OUTPUT_DATA_TESSELLATION_TWIST_CURVATURE]);
-	float TessellationTotalSegmentLength = asfloat(VertexGenerationResults[VertexGenerationResultsIndex * VERTEX_GEN_OUTPUT_DATA_STRIDE + VERTEX_GEN_OUTPUT_DATA_TESSELLATION_TOTAL_SEGMENT_LENGTH]);
-
-	if (bUseConstantFactor)
+	if ( RibbonWantsAutomaticTessellation() )
 	{
-		SegmentTessellation = TessellationFactor;
-	}
-	else if (GNiagaraRibbonTessellationEnabled == 1 && TessellationFactor > 1 && TessellationCurvature > UE_SMALL_NUMBER)
-	{	
-		const float MinTesselation = (TessellationMinAngle == 0.f || bUseConstantFactor)?
-			                             float(TessellationFactor)	:
-			                             max(1.f, max(TessellationTwistAngle, TessellationAngle) / max(UE_SMALL_NUMBER, TessellationMinAngle));
+		bool bUseConstantFactor = false;
+		uint TessellationFactor = GNiagaraRibbonMaxTessellation;
+		float TessellationMinAngle = GNiagaraRibbonTessellationAngle;
+		float ScreenPercentage = GNiagaraRibbonTessellationScreenPercentage;
 
-		const float MAX_CURVATURE_FACTOR = 0.002f; // This will clamp the curvature to around 2.5 km and avoid numerical issues.
-		
-		const float MaxDisplacementError = max(GNiagaraRibbonTessellationMinDisplacementError, ScreenPercentage * sqrt(ViewDistance) / LODDistanceFactor);
-		float Tess = TessellationAngle / max(MAX_CURVATURE_FACTOR, acos(TessellationCurvature / (TessellationCurvature + MaxDisplacementError)));
-		// FMath::RoundUpToPowerOfTwo ? This could avoid vertices moving around as tesselation increases
-
-		if (TessellationTwistAngle > 0 && TessellationTwistCurvature > 0)
+		switch (TessellationMode)
 		{
-			const float TwistTess = TessellationTwistAngle / max(MAX_CURVATURE_FACTOR, acos(TessellationTwistCurvature / (TessellationTwistCurvature + MaxDisplacementError)));
-			Tess = max(TwistTess, Tess);
+			case TESSELLATIONMODE_AUTOMATIC:
+				break;
+			case TESSELLATIONMODE_CUSTOM:
+				TessellationFactor = min(TessellationFactor, CustomTessellationFactor); // Don't allow factors bigger than the platform limit.
+				bUseConstantFactor = bCustomUseConstantFactor == 1;
+				TessellationMinAngle = CustomTessellationMinAngle;
+				ScreenPercentage = bCustomUseScreenSpace == 1 && !bUseConstantFactor ? GNiagaraRibbonTessellationScreenPercentage : 0.f;
+				break;
+			case TESSELLATIONMODE_DISABLED:
+				TessellationFactor = 1;
+				break;
+			default:
+				break;
 		}
-		SegmentTessellation = clamp(uint(round(Tess)), uint(round(MinTesselation)), TessellationFactor);
 
-		// Ensure we don't blow the max tessellation allowed, this shouldn't happen but let's be safe
-		SegmentTessellation = min(SegmentTessellation, GNiagaraRibbonMaxTessellation);
+		float TessellationAngle = asfloat(VertexGenerationResults[VertexGenerationResultsOffset + VERTEX_GEN_OUTPUT_DATA_TESSELLATION_ANGLE]);
+		float TessellationCurvature = asfloat(VertexGenerationResults[VertexGenerationResultsOffset + VERTEX_GEN_OUTPUT_DATA_TESSELLATION_CURVATURE]);
+		float TessellationTwistAngle = asfloat(VertexGenerationResults[VertexGenerationResultsOffset + VERTEX_GEN_OUTPUT_DATA_TESSELLATION_TWIST_ANGLE]);
+		float TessellationTwistCurvature = asfloat(VertexGenerationResults[VertexGenerationResultsOffset + VERTEX_GEN_OUTPUT_DATA_TESSELLATION_TWIST_CURVATURE]);
+		float TessellationTotalSegmentLength = asfloat(VertexGenerationResults[VertexGenerationResultsOffset + VERTEX_GEN_OUTPUT_DATA_TESSELLATION_TOTAL_SEGMENT_LENGTH]);
+
+		if (GNiagaraRibbonTessellationEnabled == 1 && TessellationFactor > 1 && TessellationCurvature > UE_SMALL_NUMBER)
+		{	
+			const float MinTesselation = (TessellationMinAngle == 0.f || bUseConstantFactor)?
+											 float(TessellationFactor)	:
+											 max(1.f, max(TessellationTwistAngle, TessellationAngle) / max(UE_SMALL_NUMBER, TessellationMinAngle));
+
+			const float MAX_CURVATURE_FACTOR = 0.002f; // This will clamp the curvature to around 2.5 km and avoid numerical issues.
+		
+			const float MaxDisplacementError = max(GNiagaraRibbonTessellationMinDisplacementError, ScreenPercentage * sqrt(ViewDistance) / LODDistanceFactor);
+			float Tess = TessellationAngle / max(MAX_CURVATURE_FACTOR, acos(TessellationCurvature / (TessellationCurvature + MaxDisplacementError)));
+			// FMath::RoundUpToPowerOfTwo ? This could avoid vertices moving around as tesselation increases
+
+			if (TessellationTwistAngle > 0 && TessellationTwistCurvature > 0)
+			{
+				const float TwistTess = TessellationTwistAngle / max(MAX_CURVATURE_FACTOR, acos(TessellationTwistCurvature / (TessellationTwistCurvature + MaxDisplacementError)));
+				Tess = max(TwistTess, Tess);
+			}
+			SegmentTessellation = clamp(uint(round(Tess)), uint(round(MinTesselation)), TessellationFactor);
+
+			// Ensure we don't blow the max tessellation allowed, this shouldn't happen but let's be safe
+			SegmentTessellation = min(SegmentTessellation, GNiagaraRibbonMaxTessellation);
+		}
 	}
-#elif RIBBONS_WANTS_CONSTANT_TESSELLATION
 	// Has a constant factor so just use it
-	const uint SegmentTessellation = min(GNiagaraRibbonMaxTessellation, CustomTessellationFactor);
-#else
-	// Tessellation is disabled, so just set it to a factor of 1
-	const uint SegmentTessellation = 1;
-#endif
+	else if ( RibbonWantsConstantTessellation() )
+	{
+		SegmentTessellation = min(GNiagaraRibbonMaxTessellation, CustomTessellationFactor);
+	}
 
 	// We need to run for every possible segment * tessellation so one less than total particles * tessellation
 	const uint TotalNumInvocations = (TotalNumParticles - 1) * SegmentTessellation;	
 	const uint NumThreadGroups = DivideAndRoundUp(TotalNumInvocations, IndexGenThreadSize);
 
-	const uint SegmentCount = VertexGenerationResults[VertexGenerationResultsIndex * VERTEX_GEN_OUTPUT_DATA_STRIDE + VERTEX_GEN_OUTPUT_DATA_TOTAL_NUM_SEGMENTS_OFFSET];
+	const uint SegmentCount = VertexGenerationResults[VertexGenerationResultsOffset + VERTEX_GEN_OUTPUT_DATA_TOTAL_NUM_SEGMENTS_OFFSET];
 
 	const float OneOverSubSegmentTessellation = 1.0f / float(SegmentTessellation);
 	
diff --git a/Engine/Plugins/FX/Niagara/Shaders/Private/Ribbons/NiagaraRibbonRibbonUVParamCalculation.usf b/Engine/Plugins/FX/Niagara/Shaders/Private/Ribbons/NiagaraRibbonRibbonUVParamCalculation.usf
--- a/Engine/Plugins/FX/Niagara/Shaders/Private/Ribbons/NiagaraRibbonRibbonUVParamCalculation.usf
+++ b/Engine/Plugins/FX/Niagara/Shaders/Private/Ribbons/NiagaraRibbonRibbonUVParamCalculation.usf
@@ -210,7 +210,7 @@ void GenerateRibbonUVParams(uint3 DispatchThreadId : SV_DispatchThreadID)
 {
 #if HAS_RIBBON_ID
 	const int RibbonIndex = DispatchThreadId.x;
-	const int TotalNumRibbons = CommandBuffer[CommandBufferOffset * VERTEX_GEN_OUTPUT_DATA_STRIDE + VERTEX_GEN_OUTPUT_DATA_TOTAL_NUM_RIBBONS_OFFSET];
+	const int TotalNumRibbons = CommandBuffer[CommandBufferOffset * GetVertexGenOutputDataStride() + VERTEX_GEN_OUTPUT_DATA_TOTAL_NUM_RIBBONS_OFFSET];
 	
 	if (RibbonIndex < TotalNumRibbons)
 	{
diff --git a/Engine/Plugins/FX/Niagara/Shaders/Private/Ribbons/NiagaraRibbonVertexReductionFinalization.usf b/Engine/Plugins/FX/Niagara/Shaders/Private/Ribbons/NiagaraRibbonVertexReductionFinalization.usf
--- a/Engine/Plugins/FX/Niagara/Shaders/Private/Ribbons/NiagaraRibbonVertexReductionFinalization.usf
+++ b/Engine/Plugins/FX/Niagara/Shaders/Private/Ribbons/NiagaraRibbonVertexReductionFinalization.usf
@@ -14,10 +14,8 @@ Buffer<uint> MultiRibbonIndices;
 // Index of the final segment this particle starts
 Buffer<uint> Segments;
 
-#if RIBBONS_WANTS_AUTOMATIC_TESSELLATION
 // Tessellation reduction stats if we're computing them
 Buffer<float> TessellationStats;
-#endif
 
 StructuredBuffer<FRibbonAccumulationValues> AccumulationBuffer;
 
@@ -45,7 +43,7 @@ int FinalizationThreadBlockSize;
 void OutputRibbonStats(uint3 DispatchThreadId : SV_DispatchThreadID)
 {
 	const uint TotalNumParticles = GetTotalNumParticles();
-	const uint OutputCommandBufferOffset = OutputCommandBufferIndex * VERTEX_GEN_OUTPUT_DATA_STRIDE;
+	const uint OutputCommandBufferOffset = OutputCommandBufferIndex * GetVertexGenOutputDataStride();
 
 	// Early out if we have no particles
 	if ( TotalNumParticles == 0 )
@@ -57,18 +55,19 @@ void OutputRibbonStats(uint3 DispatchThreadId : SV_DispatchThreadID)
 		OutputCommandBuffer[OutputCommandBufferOffset + VERTEX_GEN_OUTPUT_DATA_TOTAL_NUM_RIBBONS_OFFSET] = 0;
 		OutputCommandBuffer[OutputCommandBufferOffset + VERTEX_GEN_OUTPUT_DATA_TOTAL_NUM_SEGMENTS_OFFSET] = 0;
 		OutputCommandBuffer[OutputCommandBufferOffset + VERTEX_GEN_OUTPUT_DATA_TOTAL_NUM_RIBBONS_OFFSET] = 0;
-	#if RIBBONS_WANTS_AUTOMATIC_TESSELLATION
-		OutputCommandBuffer[OutputCommandBufferOffset + VERTEX_GEN_OUTPUT_DATA_TESSELLATION_ANGLE] = asuint(0.0f);
-		OutputCommandBuffer[OutputCommandBufferOffset + VERTEX_GEN_OUTPUT_DATA_TESSELLATION_CURVATURE] = asuint(0.0f);
-		OutputCommandBuffer[OutputCommandBufferOffset + VERTEX_GEN_OUTPUT_DATA_TESSELLATION_TWIST_ANGLE] = asuint(0.0f);
-		OutputCommandBuffer[OutputCommandBufferOffset + VERTEX_GEN_OUTPUT_DATA_TESSELLATION_TWIST_CURVATURE] = asuint(0.0f);
-		OutputCommandBuffer[OutputCommandBufferOffset + VERTEX_GEN_OUTPUT_DATA_TESSELLATION_TOTAL_SEGMENT_LENGTH] = asuint(0.0f);
-		OutputCommandBuffer[OutputCommandBufferOffset + VERTEX_GEN_OUTPUT_DATA_TESSELLATION_CURRENT_FRAME_TOTAL_SEGMENT_LENGTH] = asuint(0.0f);
-		OutputCommandBuffer[OutputCommandBufferOffset + VERTEX_GEN_OUTPUT_DATA_TESSELLATION_CURRENT_FRAME_AVERAGE_SEGMENT_LENGTH] = asuint(0.0f);
-		OutputCommandBuffer[OutputCommandBufferOffset + VERTEX_GEN_OUTPUT_DATA_TESSELLATION_CURRENT_FRAME_AVERAGE_SEGMENT_ANGLE] = asuint(0.0f);
-		OutputCommandBuffer[OutputCommandBufferOffset + VERTEX_GEN_OUTPUT_DATA_TESSELLATION_CURRENT_FRAME_AVERAGE_TWIST_ANGLE] = asuint(0.0f);
-		OutputCommandBuffer[OutputCommandBufferOffset + VERTEX_GEN_OUTPUT_DATA_TESSELLATION_CURRENT_FRAME_AVERAGE_WIDTH] = asuint(0.0f);
-	#endif
+		if (RibbonWantsAutomaticTessellation())
+		{
+			OutputCommandBuffer[OutputCommandBufferOffset + VERTEX_GEN_OUTPUT_DATA_TESSELLATION_ANGLE] = asuint(0.0f);
+			OutputCommandBuffer[OutputCommandBufferOffset + VERTEX_GEN_OUTPUT_DATA_TESSELLATION_CURVATURE] = asuint(0.0f);
+			OutputCommandBuffer[OutputCommandBufferOffset + VERTEX_GEN_OUTPUT_DATA_TESSELLATION_TWIST_ANGLE] = asuint(0.0f);
+			OutputCommandBuffer[OutputCommandBufferOffset + VERTEX_GEN_OUTPUT_DATA_TESSELLATION_TWIST_CURVATURE] = asuint(0.0f);
+			OutputCommandBuffer[OutputCommandBufferOffset + VERTEX_GEN_OUTPUT_DATA_TESSELLATION_TOTAL_SEGMENT_LENGTH] = asuint(0.0f);
+			OutputCommandBuffer[OutputCommandBufferOffset + VERTEX_GEN_OUTPUT_DATA_TESSELLATION_CURRENT_FRAME_TOTAL_SEGMENT_LENGTH] = asuint(0.0f);
+			OutputCommandBuffer[OutputCommandBufferOffset + VERTEX_GEN_OUTPUT_DATA_TESSELLATION_CURRENT_FRAME_AVERAGE_SEGMENT_LENGTH] = asuint(0.0f);
+			OutputCommandBuffer[OutputCommandBufferOffset + VERTEX_GEN_OUTPUT_DATA_TESSELLATION_CURRENT_FRAME_AVERAGE_SEGMENT_ANGLE] = asuint(0.0f);
+			OutputCommandBuffer[OutputCommandBufferOffset + VERTEX_GEN_OUTPUT_DATA_TESSELLATION_CURRENT_FRAME_AVERAGE_TWIST_ANGLE] = asuint(0.0f);
+			OutputCommandBuffer[OutputCommandBufferOffset + VERTEX_GEN_OUTPUT_DATA_TESSELLATION_CURRENT_FRAME_AVERAGE_WIDTH] = asuint(0.0f);
+		}
 		return;
 	}
 
@@ -130,74 +129,73 @@ void OutputRibbonStats(uint3 DispatchThreadId : SV_DispatchThreadID)
 		// Total num ribbons
 		OutputCommandBuffer[OutputCommandBufferOffset + VERTEX_GEN_OUTPUT_DATA_TOTAL_NUM_RIBBONS_OFFSET] = TotalNumRibbons;
 
-#if RIBBONS_WANTS_AUTOMATIC_TESSELLATION
-#if RIBBON_HAS_TWIST
-		const uint TessellationBaseOffset = LastParticleIndex * 5;
-#else
-		const uint TessellationBaseOffset = LastParticleIndex * 3;
-#endif
-		float TotalSegmentLength = TessellationStats[TessellationBaseOffset + 0];
-		float AverageSegmentLength = TessellationStats[TessellationBaseOffset + 1];
-		float AverageSegmentAngle = TessellationStats[TessellationBaseOffset + 2];
-#if RIBBON_HAS_TWIST
-		float AverageTwistAngle = TessellationStats[TessellationBaseOffset + 3];
-		float AverageWidth = TessellationStats[TessellationBaseOffset + 4];
-#else
-		float AverageTwistAngle = 0;
-		float AverageWidth = 0;
-#endif
+		if (RibbonWantsAutomaticTessellation())
+		{
+			const uint TessellationDataStride = GetTessellationDataStride();
+			const uint TessellationBaseOffset = LastParticleIndex * TessellationDataStride;
+			float TotalSegmentLength = TessellationStats[TessellationBaseOffset + 0];
+			float AverageSegmentLength = TessellationStats[TessellationBaseOffset + 1];
+			float AverageSegmentAngle = TessellationStats[TessellationBaseOffset + 2];
+			float AverageTwistAngle = 0;
+			float AverageWidth = 0;
+			if (RibbonHasTwist())
+			{
+				AverageTwistAngle = TessellationStats[TessellationBaseOffset + 3];
+				AverageWidth = TessellationStats[TessellationBaseOffset + 4];
+			}
 		
-		// Generate smoothed tessellation information
-		float TessellationAngle = 0.0f;
-		float TessellationCurvature = 0.0f;
-		float TessellationTwistAngle = 0.0f;
-		float TessellationTwistCurvature = 0.0f;
-		float TessellationTotalSegmentLength = 0.0f;
+			// Generate smoothed tessellation information
+			float TessellationAngle = 0.0f;
+			float TessellationCurvature = 0.0f;
+			float TessellationTwistAngle = 0.0f;
+			float TessellationTwistCurvature = 0.0f;
+			float TessellationTotalSegmentLength = 0.0f;
 		
-		if (TotalSegmentLength > 0.0)
-		{				
-			// Read Tessellation Smmothing Data (accumulated between frames)
-			TessellationAngle = asfloat(OutputCommandBuffer[OutputCommandBufferOffset + VERTEX_GEN_OUTPUT_DATA_TESSELLATION_ANGLE]);
-			TessellationCurvature = asfloat(OutputCommandBuffer[OutputCommandBufferOffset + VERTEX_GEN_OUTPUT_DATA_TESSELLATION_CURVATURE]);
-			TessellationTwistAngle = asfloat(OutputCommandBuffer[OutputCommandBufferOffset + VERTEX_GEN_OUTPUT_DATA_TESSELLATION_TWIST_ANGLE]);
-			TessellationTwistCurvature = asfloat(OutputCommandBuffer[OutputCommandBufferOffset + VERTEX_GEN_OUTPUT_DATA_TESSELLATION_TWIST_CURVATURE]);
-			TessellationTotalSegmentLength = asfloat(OutputCommandBuffer[OutputCommandBufferOffset + VERTEX_GEN_OUTPUT_DATA_TESSELLATION_TOTAL_SEGMENT_LENGTH]);		
-
-			// Blend the result between the last frame tessellation factors and the current frame base on the total length of all segments.
-			// This is only used to increase the tessellation value of the current frame data to prevent glitches where tessellation is significantly changin between frames.
-			const float OneOverTotalSegmentLength = 1.f / max(1.0f, TotalSegmentLength);
-			const float AveragingFactor = TessellationTotalSegmentLength / (TotalSegmentLength + TessellationTotalSegmentLength);
-			TessellationTotalSegmentLength = TotalSegmentLength;
+			if (TotalSegmentLength > 0.0)
+			{
+				// Read Tessellation Smmothing Data (accumulated between frames)
+				TessellationAngle = asfloat(OutputCommandBuffer[OutputCommandBufferOffset + VERTEX_GEN_OUTPUT_DATA_TESSELLATION_ANGLE]);
+				TessellationCurvature = asfloat(OutputCommandBuffer[OutputCommandBufferOffset + VERTEX_GEN_OUTPUT_DATA_TESSELLATION_CURVATURE]);
+				TessellationTwistAngle = asfloat(OutputCommandBuffer[OutputCommandBufferOffset + VERTEX_GEN_OUTPUT_DATA_TESSELLATION_TWIST_ANGLE]);
+				TessellationTwistCurvature = asfloat(OutputCommandBuffer[OutputCommandBufferOffset + VERTEX_GEN_OUTPUT_DATA_TESSELLATION_TWIST_CURVATURE]);
+				TessellationTotalSegmentLength = asfloat(OutputCommandBuffer[OutputCommandBufferOffset + VERTEX_GEN_OUTPUT_DATA_TESSELLATION_TOTAL_SEGMENT_LENGTH]);
+
+				// Blend the result between the last frame tessellation factors and the current frame base on the total length of all segments.
+				// This is only used to increase the tessellation value of the current frame data to prevent glitches where tessellation is significantly changin between frames.
+				const float OneOverTotalSegmentLength = 1.f / max(1.0f, TotalSegmentLength);
+				const float AveragingFactor = TessellationTotalSegmentLength / (TotalSegmentLength + TessellationTotalSegmentLength);
+				TessellationTotalSegmentLength = TotalSegmentLength;
 			
-			AverageSegmentAngle *= OneOverTotalSegmentLength;
-			AverageSegmentLength *= OneOverTotalSegmentLength;
-			const float AverageSegmentCurvature = AverageSegmentLength / max(UE_SMALL_NUMBER, abs(sin(AverageSegmentAngle)));
+				AverageSegmentAngle *= OneOverTotalSegmentLength;
+				AverageSegmentLength *= OneOverTotalSegmentLength;
+				const float AverageSegmentCurvature = AverageSegmentLength / max(UE_SMALL_NUMBER, abs(sin(AverageSegmentAngle)));
 
-			TessellationAngle = lerp(AverageSegmentAngle, max(TessellationAngle, AverageSegmentAngle), AveragingFactor);
-			TessellationCurvature = lerp(AverageSegmentCurvature, max(TessellationCurvature, AverageSegmentCurvature), AveragingFactor);
+				TessellationAngle = lerp(AverageSegmentAngle, max(TessellationAngle, AverageSegmentAngle), AveragingFactor);
+				TessellationCurvature = lerp(AverageSegmentCurvature, max(TessellationCurvature, AverageSegmentCurvature), AveragingFactor);
 			
-#if RIBBON_HAS_TWIST				
-			AverageTwistAngle *= OneOverTotalSegmentLength;
-			AverageWidth *= OneOverTotalSegmentLength;
-
-			TessellationTwistAngle = lerp(AverageTwistAngle, max(TessellationTwistAngle, AverageTwistAngle), AveragingFactor);
-			TessellationTwistCurvature = lerp(AverageWidth, max(TessellationTwistCurvature, AverageWidth), AveragingFactor);		
-#endif
+				if (RibbonHasTwist())
+				{
+					AverageTwistAngle *= OneOverTotalSegmentLength;
+					AverageWidth *= OneOverTotalSegmentLength;
+
+					TessellationTwistAngle = lerp(AverageTwistAngle, max(TessellationTwistAngle, AverageTwistAngle), AveragingFactor);
+					TessellationTwistCurvature = lerp(AverageWidth, max(TessellationTwistCurvature, AverageWidth), AveragingFactor);
+				}
+			}
+
+			// Output Tessellation Smmothing Data (accumulated between frames)
+			OutputCommandBuffer[OutputCommandBufferOffset + VERTEX_GEN_OUTPUT_DATA_TESSELLATION_ANGLE] = asuint(TessellationAngle);
+			OutputCommandBuffer[OutputCommandBufferOffset + VERTEX_GEN_OUTPUT_DATA_TESSELLATION_CURVATURE] = asuint(TessellationCurvature);
+			OutputCommandBuffer[OutputCommandBufferOffset + VERTEX_GEN_OUTPUT_DATA_TESSELLATION_TWIST_ANGLE] = asuint(TessellationTwistAngle);
+			OutputCommandBuffer[OutputCommandBufferOffset + VERTEX_GEN_OUTPUT_DATA_TESSELLATION_TWIST_CURVATURE] = asuint(TessellationTwistCurvature);
+			OutputCommandBuffer[OutputCommandBufferOffset + VERTEX_GEN_OUTPUT_DATA_TESSELLATION_TOTAL_SEGMENT_LENGTH] = asuint(TessellationTotalSegmentLength);
+
+			// Write out current frame tessellation stats
+			OutputCommandBuffer[OutputCommandBufferOffset + VERTEX_GEN_OUTPUT_DATA_TESSELLATION_CURRENT_FRAME_TOTAL_SEGMENT_LENGTH] = asuint(TotalSegmentLength);
+			OutputCommandBuffer[OutputCommandBufferOffset + VERTEX_GEN_OUTPUT_DATA_TESSELLATION_CURRENT_FRAME_AVERAGE_SEGMENT_LENGTH] = asuint(AverageSegmentLength);
+			OutputCommandBuffer[OutputCommandBufferOffset + VERTEX_GEN_OUTPUT_DATA_TESSELLATION_CURRENT_FRAME_AVERAGE_SEGMENT_ANGLE] = asuint(AverageSegmentAngle);
+			OutputCommandBuffer[OutputCommandBufferOffset + VERTEX_GEN_OUTPUT_DATA_TESSELLATION_CURRENT_FRAME_AVERAGE_TWIST_ANGLE] = asuint(AverageTwistAngle);
+			OutputCommandBuffer[OutputCommandBufferOffset + VERTEX_GEN_OUTPUT_DATA_TESSELLATION_CURRENT_FRAME_AVERAGE_WIDTH] = asuint(AverageWidth);
 		}
-
-		// Output Tessellation Smmothing Data (accumulated between frames)
-		OutputCommandBuffer[OutputCommandBufferOffset + VERTEX_GEN_OUTPUT_DATA_TESSELLATION_ANGLE] = asuint(TessellationAngle);
-		OutputCommandBuffer[OutputCommandBufferOffset + VERTEX_GEN_OUTPUT_DATA_TESSELLATION_CURVATURE] = asuint(TessellationCurvature);
-		OutputCommandBuffer[OutputCommandBufferOffset + VERTEX_GEN_OUTPUT_DATA_TESSELLATION_TWIST_ANGLE] = asuint(TessellationTwistAngle);
-		OutputCommandBuffer[OutputCommandBufferOffset + VERTEX_GEN_OUTPUT_DATA_TESSELLATION_TWIST_CURVATURE] = asuint(TessellationTwistCurvature);
-		OutputCommandBuffer[OutputCommandBufferOffset + VERTEX_GEN_OUTPUT_DATA_TESSELLATION_TOTAL_SEGMENT_LENGTH] = asuint(TessellationTotalSegmentLength);
-
-		// Write out current frame tessellation stats
-		OutputCommandBuffer[OutputCommandBufferOffset + VERTEX_GEN_OUTPUT_DATA_TESSELLATION_CURRENT_FRAME_TOTAL_SEGMENT_LENGTH] = asuint(TotalSegmentLength);
-		OutputCommandBuffer[OutputCommandBufferOffset + VERTEX_GEN_OUTPUT_DATA_TESSELLATION_CURRENT_FRAME_AVERAGE_SEGMENT_LENGTH] = asuint(AverageSegmentLength);
-		OutputCommandBuffer[OutputCommandBufferOffset + VERTEX_GEN_OUTPUT_DATA_TESSELLATION_CURRENT_FRAME_AVERAGE_SEGMENT_ANGLE] = asuint(AverageSegmentAngle);
-		OutputCommandBuffer[OutputCommandBufferOffset + VERTEX_GEN_OUTPUT_DATA_TESSELLATION_CURRENT_FRAME_AVERAGE_TWIST_ANGLE] = asuint(AverageTwistAngle);
-		OutputCommandBuffer[OutputCommandBufferOffset + VERTEX_GEN_OUTPUT_DATA_TESSELLATION_CURRENT_FRAME_AVERAGE_WIDTH] = asuint(AverageWidth);
-#endif //RIBBONS_WANTS_AUTOMATIC_TESSELLATION
 	}
 }
diff --git a/Engine/Plugins/FX/Niagara/Shaders/Private/Ribbons/NiagaraRibbonVertexReductionInitialization.usf b/Engine/Plugins/FX/Niagara/Shaders/Private/Ribbons/NiagaraRibbonVertexReductionInitialization.usf
--- a/Engine/Plugins/FX/Niagara/Shaders/Private/Ribbons/NiagaraRibbonVertexReductionInitialization.usf
+++ b/Engine/Plugins/FX/Niagara/Shaders/Private/Ribbons/NiagaraRibbonVertexReductionInitialization.usf
@@ -116,26 +116,28 @@ void VertexGenInitialize(uint3 DispatchThreadId : SV_DispatchThreadID)
 				TangentAndDistance = float4(0,0,0,0);
 			}
 
-#if RIBBONS_WANTS_AUTOMATIC_TESSELLATION
-			// TotalLength
-			OutputAccumulationBuffer[RawParticleID].TessTotalLength = CurrToNextSize;
-			// AverageSegmentLength
-			OutputAccumulationBuffer[RawParticleID].TessAvgSegmentLength = CurrToNextSize * CurrToNextSize;
-			// AverageSegmentAngle
-			OutputAccumulationBuffer[RawParticleID].TessAvgSegmentAngle = CurrToNextSize * AcosFast(dot(LastToCurrVec, CurrToNextVec));
-
-#if RIBBON_HAS_TWIST
-			const float LastTwist = SafeGetFloat(TwistDataOffset, PrevParticleID, 0);
-			const float ThisTwist = SafeGetFloat(TwistDataOffset, ParticleID, 0);
-
-			//const float LastWidth = SafeGetFloat(WidthDataOffset, PrevParticleID, 1);
-			const float ThisWidth = SafeGetFloat(WidthDataOffset, ParticleID, 1);
-			// AverageTwistAngle
-			OutputAccumulationBuffer[RawParticleID].TessTwistAvgAngle = CurrToNextSize * abs(ThisTwist - LastTwist);
-			// AverageWidth
-			OutputAccumulationBuffer[RawParticleID].TessTwistAvgWidth = CurrToNextSize * ThisWidth;
-#endif // RIBBON_HAS_TWIST
-#endif // RIBBONS_WANTS_AUTOMATIC_TESSELLATION
+			if (RibbonWantsAutomaticTessellation())
+			{
+				// TotalLength
+				OutputAccumulationBuffer[RawParticleID].TessTotalLength = CurrToNextSize;
+				// AverageSegmentLength
+				OutputAccumulationBuffer[RawParticleID].TessAvgSegmentLength = CurrToNextSize * CurrToNextSize;
+				// AverageSegmentAngle
+				OutputAccumulationBuffer[RawParticleID].TessAvgSegmentAngle = CurrToNextSize * AcosFast(dot(LastToCurrVec, CurrToNextVec));
+
+				if (RibbonHasTwist())
+				{
+					const float LastTwist = SafeGetFloat(TwistDataOffset, PrevParticleID, 0);
+					const float ThisTwist = SafeGetFloat(TwistDataOffset, ParticleID, 0);
+
+					//const float LastWidth = SafeGetFloat(WidthDataOffset, PrevParticleID, 1);
+					const float ThisWidth = SafeGetFloat(WidthDataOffset, ParticleID, 1);
+					// AverageTwistAngle
+					OutputAccumulationBuffer[RawParticleID].TessTwistAvgAngle = CurrToNextSize * abs(ThisTwist - LastTwist);
+					// AverageWidth
+					OutputAccumulationBuffer[RawParticleID].TessTwistAvgWidth = CurrToNextSize * ThisWidth;
+				}
+			}
 		}
 		else if (bHasNextParticle)
 		{
diff --git a/Engine/Plugins/FX/Niagara/Shaders/Private/Ribbons/NiagaraRibbonVertexReductionPropagation.usf b/Engine/Plugins/FX/Niagara/Shaders/Private/Ribbons/NiagaraRibbonVertexReductionPropagation.usf
--- a/Engine/Plugins/FX/Niagara/Shaders/Private/Ribbons/NiagaraRibbonVertexReductionPropagation.usf
+++ b/Engine/Plugins/FX/Niagara/Shaders/Private/Ribbons/NiagaraRibbonVertexReductionPropagation.usf
@@ -21,11 +21,9 @@ RWBuffer<uint> OutputMultiRibbonIndices;
 Buffer<uint> InputSegments;
 RWBuffer<uint> OutputSegments;
 
-#if RIBBONS_WANTS_AUTOMATIC_TESSELLATION
 // Tessellation reduction stats if we're computing them
 Buffer<float> InputTessellationStats;
 RWBuffer<float> OutputTessellationStats;
-#endif
 
 int PrefixScanStride;
 
@@ -75,28 +73,26 @@ void VertexGenPrefixSumPropagation(uint3 DispatchThreadId : SV_DispatchThreadID)
 			}
 #endif
 
-#if RIBBONS_WANTS_AUTOMATIC_TESSELLATION
-#if RIBBON_HAS_TWIST
-			const uint TessellationBaseOffset = RawParticleID * 5;
-			const uint TessellationSourceBaseOffset = RawSourceParticleID * 5;
-#else
-			const uint TessellationBaseOffset = RawParticleID * 3;
-			const uint TessellationSourceBaseOffset = RawSourceParticleID * 3;
-#endif
-			// TotalLength
-			OutputTessellationStats[TessellationBaseOffset + 0] = InputTessellationStats[TessellationBaseOffset + 0] + InputTessellationStats[TessellationSourceBaseOffset + 0];
-			// AverageSegmentLength
-			OutputTessellationStats[TessellationBaseOffset + 1] = InputTessellationStats[TessellationBaseOffset + 1] + InputTessellationStats[TessellationSourceBaseOffset + 1];
-			// AverageSegmentAngle
-			OutputTessellationStats[TessellationBaseOffset + 2] = InputTessellationStats[TessellationBaseOffset + 2] + InputTessellationStats[TessellationSourceBaseOffset + 2];
-#if RIBBON_HAS_TWIST
-			// AverageTwistAngle
-			OutputTessellationStats[TessellationBaseOffset + 3] = InputTessellationStats[TessellationBaseOffset + 3] + InputTessellationStats[TessellationSourceBaseOffset + 3];
-			// AverageWidth
-			OutputTessellationStats[TessellationBaseOffset + 4] = InputTessellationStats[TessellationBaseOffset + 4] + InputTessellationStats[TessellationSourceBaseOffset + 4];
-#endif
-#endif
-			
+			if (RibbonWantsAutomaticTessellation())
+			{
+				const uint TessellationDataStride = GetTessellationDataStride();
+				const uint TessellationBaseOffset = RawParticleID * TessellationDataStride;
+				const uint TessellationSourceBaseOffset = RawSourceParticleID * TessellationDataStride;
+				// TotalLength
+				OutputTessellationStats[TessellationBaseOffset + 0] = InputTessellationStats[TessellationBaseOffset + 0] + InputTessellationStats[TessellationSourceBaseOffset + 0];
+				// AverageSegmentLength
+				OutputTessellationStats[TessellationBaseOffset + 1] = InputTessellationStats[TessellationBaseOffset + 1] + InputTessellationStats[TessellationSourceBaseOffset + 1];
+				// AverageSegmentAngle
+				OutputTessellationStats[TessellationBaseOffset + 2] = InputTessellationStats[TessellationBaseOffset + 2] + InputTessellationStats[TessellationSourceBaseOffset + 2];
+
+				if (RibbonHasTwist())
+				{
+					// AverageTwistAngle
+					OutputTessellationStats[TessellationBaseOffset + 3] = InputTessellationStats[TessellationBaseOffset + 3] + InputTessellationStats[TessellationSourceBaseOffset + 3];
+					// AverageWidth
+					OutputTessellationStats[TessellationBaseOffset + 4] = InputTessellationStats[TessellationBaseOffset + 4] + InputTessellationStats[TessellationSourceBaseOffset + 4];
+				}
+			}
 		}
 		else
 		{			
@@ -112,26 +108,26 @@ void VertexGenPrefixSumPropagation(uint3 DispatchThreadId : SV_DispatchThreadID)
 			OutputTangentsAndDistances[RawParticleID * 4 + 3] = InputTangentsAndDistances[RawParticleID * 4 + 3];
 
 
-#if RIBBONS_WANTS_AUTOMATIC_TESSELLATION
-#if RIBBON_HAS_TWIST
-			const uint TessellationBaseOffset = RawParticleID * 5;
-#else
-			const uint TessellationBaseOffset = RawParticleID * 3;
-#endif
-			// TotalLength
-			OutputTessellationStats[TessellationBaseOffset + 0] = InputTessellationStats[TessellationBaseOffset + 0];
-			// AverageSegmentLength
-			OutputTessellationStats[TessellationBaseOffset + 1] = InputTessellationStats[TessellationBaseOffset + 1];
-			// AverageSegmentAngle
-			OutputTessellationStats[TessellationBaseOffset + 2] = InputTessellationStats[TessellationBaseOffset + 2];
+			if (RibbonWantsAutomaticTessellation())
+			{
+				const uint TessellationDataStride = GetTessellationDataStride();
+				const uint TessellationBaseOffset = RawParticleID * TessellationDataStride;
+
+				// TotalLength
+				OutputTessellationStats[TessellationBaseOffset + 0] = InputTessellationStats[TessellationBaseOffset + 0];
+				// AverageSegmentLength
+				OutputTessellationStats[TessellationBaseOffset + 1] = InputTessellationStats[TessellationBaseOffset + 1];
+				// AverageSegmentAngle
+				OutputTessellationStats[TessellationBaseOffset + 2] = InputTessellationStats[TessellationBaseOffset + 2];
 			
-#if RIBBON_HAS_TWIST
-			// AverageTwistAngle
-			OutputTessellationStats[TessellationBaseOffset + 3] = InputTessellationStats[TessellationBaseOffset + 3];
-			// AverageWidth
-			OutputTessellationStats[TessellationBaseOffset + 4] = InputTessellationStats[TessellationBaseOffset + 4];
-#endif
-#endif
+				if (RibbonHasTwist())
+				{
+					// AverageTwistAngle
+					OutputTessellationStats[TessellationBaseOffset + 3] = InputTessellationStats[TessellationBaseOffset + 3];
+					// AverageWidth
+					OutputTessellationStats[TessellationBaseOffset + 4] = InputTessellationStats[TessellationBaseOffset + 4];
+				}
+			}
 		}
 	}
 }
diff --git a/Engine/Plugins/FX/Niagara/Shaders/Shared/NiagaraRibbonShared.h b/Engine/Plugins/FX/Niagara/Shaders/Shared/NiagaraRibbonShared.h
--- a/Engine/Plugins/FX/Niagara/Shaders/Shared/NiagaraRibbonShared.h
+++ b/Engine/Plugins/FX/Niagara/Shaders/Shared/NiagaraRibbonShared.h
@@ -12,13 +12,29 @@ struct FRibbonAccumulationValues
 	float RibbonDistance;
 	uint SegmentCount;
 	uint MultiRibbonCount;			// Only valid when HAS_RIBBON_ID
-	float TessTotalLength;			// Only valid when RIBBONS_WANTS_AUTOMATIC_TESSELLATION
-	float TessAvgSegmentLength;		// Only valid when RIBBONS_WANTS_AUTOMATIC_TESSELLATION
-	float TessAvgSegmentAngle;		// Only valid when RIBBONS_WANTS_AUTOMATIC_TESSELLATION
-	float TessTwistAvgAngle;		// Only valid when RIBBONS_WANTS_AUTOMATIC_TESSELLATION && RIBBON_HAS_TWIST
-	float TessTwistAvgWidth;		// Only valid when RIBBONS_WANTS_AUTOMATIC_TESSELLATION && RIBBON_HAS_TWIST
+	float TessTotalLength;			// Only valid when RibbonWantsAutomaticTessellation()
+	float TessAvgSegmentLength;		// Only valid when RibbonWantsAutomaticTessellation()
+	float TessAvgSegmentAngle;		// Only valid when RibbonWantsAutomaticTessellation()
+	float TessTwistAvgAngle;		// Only valid when RibbonWantsAutomaticTessellation() && RibbonHasTwist()
+	float TessTwistAvgWidth;		// Only valid when RibbonWantsAutomaticTessellation() && RibbonHasTwist()
 };
 
+struct ENiagaraGpuTessellationMode
+{
+	enum Type
+	{
+		Disabled,
+		Automatic,
+		Constant,
+	};
+};
+
+#define NIAGARA_RIBBON_TESSELLATION_FULL_STATS_NUM_ELEMENTS		5
+#define NIAGARA_RIBBON_TESSELLATION_NO_TWIST_STATS_NUM_ELEMENTS	3
+
+#define VERTEX_GEN_OUTPUT_DATA_AUTO_TESSELLATION_NUM_ELEMENTS	15
+#define VERTEX_GEN_OUTPUT_DATA_NUM_ELEMENTS						5
+
 #ifndef GPU_SIMULATION
 	#undef uint
 #endif
diff --git a/Engine/Plugins/FX/Niagara/Source/Niagara/Private/NiagaraRendererRibbons.cpp b/Engine/Plugins/FX/Niagara/Source/Niagara/Private/NiagaraRendererRibbons.cpp
--- a/Engine/Plugins/FX/Niagara/Source/Niagara/Private/NiagaraRendererRibbons.cpp
+++ b/Engine/Plugins/FX/Niagara/Source/Niagara/Private/NiagaraRendererRibbons.cpp
@@ -186,71 +186,44 @@ static FAutoConsoleVariableRef CNiagaraRibbonForceIndex32(
 	ECVF_Default
 );
 
-// max absolute error 9.0x10^-3
-// Eberly's polynomial degree 1 - respect bounds
-// input [-1, 1] and output [0, PI]
-FORCEINLINE float AcosFast(float InX)
+namespace NiagaraRendererRibbonsPrivate
 {
-	float X = FMath::Abs(InX);
-	float Res = -0.156583f * X + (0.5f * PI);
-	Res *= FMath::Sqrt(FMath::Max(0.f, 1.0f - X));
-	return (InX >= 0) ? Res : PI - Res;
-}
-
-// Calculates the number of bits needed to store a maximum value
-FORCEINLINE uint32 CalculateBitsForRange(uint32 Range)
-{
-	return FMath::CeilToInt(FMath::Loge(static_cast<float>(Range)) / FMath::Loge(static_cast<float>(2)));
-}
-
-// Generates the mask to remove unecessary bits after a range of bits
-FORCEINLINE uint32 CalculateBitMask(uint32 NumBits)
-{
-	return static_cast<uint32>(static_cast<uint64>(0xFFFFFFFF) >> (32 - NumBits));
-}
+	// max absolute error 9.0x10^-3
+	// Eberly's polynomial degree 1 - respect bounds
+	// input [-1, 1] and output [0, PI]
+	static float AcosFast(float InX)
+	{
+		float X = FMath::Abs(InX);
+		float Res = -0.156583f * X + (0.5f * PI);
+		Res *= FMath::Sqrt(FMath::Max(0.f, 1.0f - X));
+		return (InX >= 0) ? Res : PI - Res;
+	}
 
-struct FTessellationStatsEntry
-{
-	static constexpr int32 NumElements = 5;
-	
-	float TotalLength;
-	float AverageSegmentLength;
-	float AverageSegmentAngle;
-	float AverageTwistAngle;
-	float AverageWidth;
-};
+	// Calculates the number of bits needed to store a maximum value
+	static uint32 CalculateBitsForRange(uint32 Range)
+	{
+		return FMath::CeilToInt(FMath::Loge(static_cast<float>(Range)) / FMath::Loge(static_cast<float>(2)));
+	}
 
-struct FTessellationStatsEntryNoTwist
-{
-	static constexpr int32 NumElements = 3;
-	
-	float TotalLength;
-	float AverageSegmentLength;
-	float AverageSegmentAngle;
-};
+	// Generates the mask to remove unecessary bits after a range of bits
+	static uint32 CalculateBitMask(uint32 NumBits)
+	{
+		return static_cast<uint32>(static_cast<uint64>(0xFFFFFFFF) >> (32 - NumBits));
+	}
 
-struct FNiagaraRibbonCommandBufferLayout
-{
-	static constexpr int32 NumElements = 15;
-	
-	uint32 FinalizationIndirectArgsXDim;
-	uint32 FinalizationIndirectArgsYDim;
-	uint32 FinalizationIndirectArgsZDim;
-	uint32 NumSegments;
-	uint32 NumRibbons;
-	
-	float Tessellation_Angle;
-	float Tessellation_Curvature;
-	float Tessellation_TwistAngle;
-	float Tessellation_TwistCurvature;
-	float Tessellation_TotalLength;
-
-	float TessCurrentFrame_TotalLength;
-	float TessCurrentFrame_AverageSegmentLength;
-	float TessCurrentFrame_AverageSegmentAngle;
-	float TessCurrentFrame_AverageTwistAngle;
-	float TessCurrentFrame_AverageWidth;
-};
+	static ENiagaraGpuTessellationMode::Type GetGpuTessellationMode(const FNiagaraRibbonGenerationConfig& GenerationConfig)
+	{
+		if (GenerationConfig.WantsAutomaticTessellation())
+		{
+			return ENiagaraGpuTessellationMode::Automatic;
+		}
+		else if (GenerationConfig.WantsConstantTessellation())
+		{
+			return ENiagaraGpuTessellationMode::Constant;
+		}
+		return ENiagaraGpuTessellationMode::Disabled;
+	}
+} //namespace NiagaraRendererRibbonsPrivate
 
 // This data must match INDEX_GEN_INDIRECT_ARGS_STRIDE in NiagaraRibbonCommon.ush
 // Be careful if we ever allocate more than 1 of these as ExecuteIndirect arguments have boundary restrictions on some platforms
@@ -529,7 +502,7 @@ void FNiagaraRibbonVertexBuffers::InitializeOrUpdateBuffers(FRHICommandListBase&
 		MultiRibbonIndicesBuffer.Allocate(RHICmdList, GenerationConfig.HasRibbonIDs() ? TotalParticles : 0, MaxAllocationCount, InitialBufferAccessFlags, false);
 		RibbonLookupTableBuffer.Allocate(RHICmdList, TotalRibbons * FRibbonMultiRibbonInfoBufferEntry::NumElements, MaxAllocatedRibbons * FRibbonMultiRibbonInfoBufferEntry::NumElements, InitialBufferAccessFlags, false);
 		SegmentsBuffer.Allocate(RHICmdList, TotalParticles, MaxAllocationCount, InitialBufferAccessFlags, false);
-		bJustCreatedCommandBuffer |= GPUComputeCommandBuffer.Allocate(RHICmdList, FNiagaraRibbonCommandBufferLayout::NumElements, FNiagaraRibbonCommandBufferLayout::NumElements, InitialBufferAccessFlags | ERHIAccess::IndirectArgs, false, EBufferUsageFlags::DrawIndirect);
+		bJustCreatedCommandBuffer |= GPUComputeCommandBuffer.Allocate(RHICmdList, VERTEX_GEN_OUTPUT_DATA_AUTO_TESSELLATION_NUM_ELEMENTS, VERTEX_GEN_OUTPUT_DATA_AUTO_TESSELLATION_NUM_ELEMENTS, InitialBufferAccessFlags | ERHIAccess::IndirectArgs, false, EBufferUsageFlags::DrawIndirect);
 	}
 	else
 	{		
@@ -555,6 +528,8 @@ struct FNiagaraRibbonGPUInitComputeBuffers
 
 	void InitOrUpdateBuffers(FRHICommandListBase& RHICmdList, int32 NeededSize, bool bWantsMultiRibbon, bool bWantsTessellation, bool bWantsTessellationTwist)
 	{
+		using namespace NiagaraRendererRibbonsPrivate;
+
 		// TODO: Downsize these when we haven't needed the size for a bit
 		constexpr ERHIAccess InitialAccess = ERHIAccess::SRVMask;
 		
@@ -563,7 +538,7 @@ struct FNiagaraRibbonGPUInitComputeBuffers
 			SortBuffer.Initialize(RHICmdList, TEXT("NiagarGPUInit-SortedIndices"), sizeof(uint32), NeededSize, EPixelFormat::PF_R32_UINT, InitialAccess | ERHIAccess::VertexOrIndexBuffer);
 		}
 
-		const uint32 TessellationBufferSize = NeededSize * (bWantsTessellation? (bWantsTessellationTwist? FTessellationStatsEntry::NumElements : FTessellationStatsEntryNoTwist::NumElements) : 0);
+		const uint32 TessellationBufferSize = NeededSize * (bWantsTessellation ? (bWantsTessellationTwist ? NIAGARA_RIBBON_TESSELLATION_FULL_STATS_NUM_ELEMENTS : NIAGARA_RIBBON_TESSELLATION_NO_TWIST_STATS_NUM_ELEMENTS) : 0);
 		if (TransientTessellationStats.NumBytes < TessellationBufferSize * sizeof(float))
 		{
 			TransientTessellationStats.Initialize(RHICmdList, TEXT("NiagaraGPUInit-Tessellation-0"), sizeof(float), TessellationBufferSize, EPixelFormat::PF_R32_FLOAT, InitialAccess, BUF_Static);
@@ -1257,6 +1232,8 @@ void FNiagaraRendererRibbons::GetDynamicRayTracingInstances(FRayTracingInstanceC
 
 void FNiagaraRendererRibbons::GenerateShapeStateMultiPlane(FNiagaraRibbonShapeGeometryData& State, int32 MultiPlaneCount, int32 WidthSegmentationCount, bool bEnableAccurateGeometry, bool bUseMaterialBackfaceCulling)
 {
+	using namespace NiagaraRendererRibbonsPrivate;
+
 	State.Shape = ENiagaraRibbonShapeMode::MultiPlane;
 	State.bDisableBackfaceCulling = !bEnableAccurateGeometry && !bUseMaterialBackfaceCulling;
 	State.bShouldFlipNormalToView = !bEnableAccurateGeometry;
@@ -1323,6 +1300,8 @@ void FNiagaraRendererRibbons::GenerateShapeStateMultiPlane(FNiagaraRibbonShapeGe
 
 void FNiagaraRendererRibbons::GenerateShapeStateTube(FNiagaraRibbonShapeGeometryData& State, int32 TubeSubdivisions, bool bUseMaterialBackfaceCulling)
 {
+	using namespace NiagaraRendererRibbonsPrivate;
+
 	State.Shape = ENiagaraRibbonShapeMode::Tube;
 	State.bDisableBackfaceCulling = !bUseMaterialBackfaceCulling;
 	State.bShouldFlipNormalToView = false;
@@ -1352,6 +1331,8 @@ void FNiagaraRendererRibbons::GenerateShapeStateTube(FNiagaraRibbonShapeGeometry
 
 void FNiagaraRendererRibbons::GenerateShapeStateCustom(FNiagaraRibbonShapeGeometryData& State, const TArray<FNiagaraRibbonShapeCustomVertex>& CustomVertices, bool bUseMaterialBackfaceCulling)
 {
+	using namespace NiagaraRendererRibbonsPrivate;
+
 	State.Shape = ENiagaraRibbonShapeMode::Custom;
 	State.bDisableBackfaceCulling = !bUseMaterialBackfaceCulling;
 	State.bShouldFlipNormalToView = false;
@@ -1391,6 +1372,8 @@ void FNiagaraRendererRibbons::GenerateShapeStateCustom(FNiagaraRibbonShapeGeomet
 
 void FNiagaraRendererRibbons::GenerateShapeStatePlane(FNiagaraRibbonShapeGeometryData& State, int32 WidthSegmentationCount, bool bUseMaterialBackfaceCulling)
 {
+	using namespace NiagaraRendererRibbonsPrivate;
+
 	State.Shape = ENiagaraRibbonShapeMode::Plane;
 	State.bDisableBackfaceCulling = !bUseMaterialBackfaceCulling;
 	State.bShouldFlipNormalToView = false;
@@ -1547,6 +1530,8 @@ void FNiagaraRendererRibbons::CalculateUVScaleAndOffsets(const FNiagaraRibbonUVS
 template<bool bWantsTessellation, bool bHasTwist, bool bWantsMultiRibbon>
 void FNiagaraRendererRibbons::GenerateVertexBufferForRibbonPart(const FNiagaraGenerationInputDataCPUAccessors& CPUData, TConstArrayView<uint32> RibbonIndices, uint32 RibbonIndex, FNiagaraRibbonCPUGeneratedVertexData& OutputData) const
 {
+	using namespace NiagaraRendererRibbonsPrivate;
+
 	TArray<uint32>& SegmentData = OutputData.SegmentData;
 	
 	const FNiagaraDataSetReaderFloat<FNiagaraPosition>& PosData = CPUData.PosData;	
@@ -1907,6 +1892,8 @@ void FNiagaraRendererRibbons::GenerateVertexBufferCPU(const FNiagaraGenerationIn
 
 int32 FNiagaraRendererRibbons::CalculateTessellationFactor(const FNiagaraSceneProxy* SceneProxy, const FSceneView* View, const FVector& ViewOriginForDistanceCulling) const
 {
+	using namespace NiagaraRendererRibbonsPrivate;
+
 	bool bUseConstantFactor = false;
 	int32 TessellationFactor = GNiagaraRibbonMaxTessellation;
 	float TessellationMinAngle = GNiagaraRibbonTessellationAngle;
@@ -1962,6 +1949,8 @@ int32 FNiagaraRendererRibbons::CalculateTessellationFactor(const FNiagaraScenePr
 FNiagaraIndexGenerationInput FNiagaraRendererRibbons::CalculateIndexBufferConfiguration(const TSharedPtr<FNiagaraRibbonCPUGeneratedVertexData>& GeneratedVertices, const FNiagaraDataBuffer* SourceParticleData,
 	const FNiagaraSceneProxy* SceneProxy, const FSceneView* View, const FVector& ViewOriginForDistanceCulling, bool bShouldUseGPUInitIndices, bool bIsGPUSim) const
 {
+	using namespace NiagaraRendererRibbonsPrivate;
+
 	FNiagaraIndexGenerationInput IndexGenInput;
 
 	IndexGenInput.ViewDistance = SceneProxy->GetProxyDynamicData().LODDistanceOverride >= 0.0f ? SceneProxy->GetProxyDynamicData().LODDistanceOverride : float(SceneProxy->GetBounds().ComputeSquaredDistanceFromBoxToPoint(ViewOriginForDistanceCulling));
@@ -2171,7 +2160,7 @@ void FNiagaraRendererRibbons::SetupPerViewUniformBuffer(FNiagaraIndexGenerationI
 	PerViewUniformParameters.ColorDataOffset = VFVariables[ENiagaraRibbonVFLayout::Color].GetGPUOffset();
 	PerViewUniformParameters.WidthDataOffset = VFVariables[ENiagaraRibbonVFLayout::Width].GetGPUOffset();
 	PerViewUniformParameters.PrevWidthDataOffset = SelectVFComponent(VFVariables, ENiagaraRibbonVFLayout::PrevRibbonWidth, ENiagaraRibbonVFLayout::Width);
-	PerViewUniformParameters.TwistDataOffset = VFVariables[ENiagaraRibbonVFLayout::Twist].GetGPUOffset();
+	PerViewUniformParameters.TwistDataOffset = GenerationConfig.HasTwist() ? VFVariables[ENiagaraRibbonVFLayout::Twist].GetGPUOffset() : -1;
 	PerViewUniformParameters.PrevTwistDataOffset = SelectVFComponent(VFVariables, ENiagaraRibbonVFLayout::PrevRibbonTwist, ENiagaraRibbonVFLayout::PrevRibbonTwist);
 	PerViewUniformParameters.NormalizedAgeDataOffset = VFVariables[ENiagaraRibbonVFLayout::NormalizedAge].GetGPUOffset();
 	PerViewUniformParameters.MaterialRandomDataOffset = VFVariables[ENiagaraRibbonVFLayout::MaterialRandom].GetGPUOffset();
@@ -2282,6 +2271,8 @@ inline void FNiagaraRendererRibbons::SetupMeshBatchAndCollectorResourceForView(F
 
 void FNiagaraRendererRibbons::InitializeViewIndexBuffersGPU(FRHICommandListImmediate& RHICmdList, FNiagaraGpuComputeDispatchInterface* ComputeDispatchInterface, const FNiagaraRibbonGPUInitParameters& GpuInitParameters, const FNiagaraRibbonRenderingFrameViewResources* RenderingViewResources) const
 {
+	using namespace NiagaraRendererRibbonsPrivate;
+
 	SCOPE_CYCLE_COUNTER(STAT_NiagaraRenderRibbonsGenIndiciesGPU);
 	
 	if (!RenderingViewResources->IndirectDrawBuffer->Buffer.IsValid())
@@ -2293,11 +2284,7 @@ void FNiagaraRendererRibbons::InitializeViewIndexBuffersGPU(FRHICommandListImmed
 
 	SCOPED_DRAW_EVENT(RHICmdList, NiagaraRenderRibbonsGenIndiciesGPU);
 	{
-		FNiagaraRibbonCreateIndexBufferParamsCS::FPermutationDomain PermutationVector;
-		PermutationVector.Set<FRibbonWantsAutomaticTessellation>(GenerationConfig.WantsAutomaticTessellation());
-		PermutationVector.Set<FRibbonWantsConstantTessellation>(GenerationConfig.WantsConstantTessellation());
-			
-		TShaderMapRef<FNiagaraRibbonCreateIndexBufferParamsCS> ComputeShader(GetGlobalShaderMap(GMaxRHIFeatureLevel), PermutationVector);
+		TShaderMapRef<FNiagaraRibbonCreateIndexBufferParamsCS> ComputeShader(GetGlobalShaderMap(GMaxRHIFeatureLevel));
 
 		FNiagaraRibbonInitializeIndices Params;
 		FMemory::Memzero(Params);
@@ -2317,6 +2304,7 @@ void FNiagaraRendererRibbons::InitializeViewIndexBuffersGPU(FRHICommandListImmed
 		Params.IndexGenThreadSize = FNiagaraRibbonComputeCommon::IndexGenThreadSize;
 		Params.TrianglesPerSegment = ShapeState.TrianglesPerSegment;
 
+		Params.RibbonTessellationMode = GetGpuTessellationMode(GenerationConfig);
 		Params.ViewDistance = RenderingViewResources->IndexGenerationSettings.ViewDistance;
 		Params.LODDistanceFactor = RenderingViewResources->IndexGenerationSettings.LODDistanceFactor;
 		Params.TessellationMode = static_cast<uint32>(TessellationConfig.TessellationMode);
@@ -2348,9 +2336,6 @@ void FNiagaraRendererRibbons::InitializeViewIndexBuffersGPU(FRHICommandListImmed
 		PermutationVector.Set<FRibbonHasFullRibbonID>(GenerationConfig.HasFullRibbonIDs());
 		PermutationVector.Set<FRibbonHasRibbonID>(GenerationConfig.HasSimpleRibbonIDs());
 
-		// This switches the index gen from a unrolled limited loop for performance to a full loop that can handle anything thrown at it
-		PermutationVector.Set<FRibbonHasHighSliceComplexity>(ShapeState.TrianglesPerSegment > 32);
-		
 		TShaderMapRef<FNiagaraRibbonCreateIndexBufferCS> ComputeShader(GetGlobalShaderMap(GMaxRHIFeatureLevel), PermutationVector);
 		// const int TotalNumInvocations = Params.TotalNumParticles * IndicesConfig.SubSegmentCount;
 		// const uint32 NumThreadGroups = FMath::DivideAndRoundUp<uint32>(TotalNumInvocations, FNiagaraRibbonComputeCommon::IndexGenThreadSize);
@@ -2487,6 +2472,8 @@ void FNiagaraRendererRibbons::InitializeVertexBuffersResources(FRHICommandListBa
 
 FRibbonComputeUniformParameters FNiagaraRendererRibbons::SetupComputeVertexGenParams(FNiagaraGpuComputeDispatchInterface* ComputeDispatchInterface, const TSharedPtr<FNiagaraRibbonRenderingFrameResources>& RenderingResources, const FNiagaraRibbonGPUInitParameters& GpuInitParameters) const
 {
+	using namespace NiagaraRendererRibbonsPrivate;
+
 	FRibbonComputeUniformParameters CommonParams;
 	FMemory::Memzero(CommonParams);
 
@@ -2540,8 +2527,11 @@ FRibbonComputeUniformParameters FNiagaraRendererRibbons::SetupComputeVertexGenPa
 	CommonParams.FacingDataOffset = bShouldDoFacing ? VFVariables[ENiagaraRibbonVFLayout::Facing].GetGPUOffset() : -1;
 	CommonParams.PrevFacingDataOffset = bShouldDoFacing ? VFVariables[ENiagaraRibbonVFLayout::PrevRibbonFacing].GetGPUOffset() : -1;
 
+	CommonParams.RibbonLinkOrderIsFloat = bGpuRibbonLinkIsFloat;
 	CommonParams.RibbonLinkOrderDataOffset = GpuRibbonLinkOrderOffset;
 
+	CommonParams.RibbonTessellationMode = GetGpuTessellationMode(GenerationConfig);
+
 	CommonParams.U0DistributionMode = static_cast<int32>(UV0Settings.DistributionMode);
 	CommonParams.U1DistributionMode = static_cast<int32>(UV1Settings.DistributionMode);
 
@@ -2581,7 +2571,6 @@ void FNiagaraRendererRibbons::InitializeVertexBuffersGPU(FRHICommandListImmediat
 		FNiagaraRibbonSortPhase1CS::FPermutationDomain PermutationVector;
 		PermutationVector.Set<FRibbonHasFullRibbonID>(GenerationConfig.HasFullRibbonIDs());
 		PermutationVector.Set<FRibbonHasRibbonID>(GenerationConfig.HasSimpleRibbonIDs());
-		PermutationVector.Set<FRibbonLinkIsFloat>(bGpuRibbonLinkIsFloat);
 		
 		TShaderMapRef<FNiagaraRibbonSortPhase1CS> BubbleSortShader(GetGlobalShaderMap(GMaxRHIFeatureLevel), PermutationVector);
 		TShaderMapRef<FNiagaraRibbonSortPhase2CS> MergeSortShader(GetGlobalShaderMap(GMaxRHIFeatureLevel), PermutationVector);
@@ -2650,25 +2639,16 @@ void FNiagaraRendererRibbons::InitializeVertexBuffersGPU(FRHICommandListImmediat
 		FNiagaraRibbonVertexReductionInitializationCS::FPermutationDomain InitPermutationVector;
 		InitPermutationVector.Set<FRibbonHasFullRibbonID>(GenerationConfig.HasFullRibbonIDs());
 		InitPermutationVector.Set<FRibbonHasRibbonID>(GenerationConfig.HasSimpleRibbonIDs());
-		InitPermutationVector.Set<FRibbonWantsAutomaticTessellation>(GenerationConfig.WantsAutomaticTessellation());
-		InitPermutationVector.Set<FRibbonWantsConstantTessellation>(GenerationConfig.WantsConstantTessellation());
-		InitPermutationVector.Set<FRibbonHasTwist>(GenerationConfig.HasTwist());
 		TShaderMapRef<FNiagaraRibbonVertexReductionInitializationCS> ReductionInitializationShader(GetGlobalShaderMap(GMaxRHIFeatureLevel), InitPermutationVector);
 
 		FNiagaraRibbonAggregationStepCS::FPermutationDomain StepPermutationVector;
 		StepPermutationVector.Set<FRibbonHasFullRibbonID>(GenerationConfig.HasFullRibbonIDs());
 		StepPermutationVector.Set<FRibbonHasRibbonID>(GenerationConfig.HasSimpleRibbonIDs());
-		StepPermutationVector.Set<FRibbonWantsAutomaticTessellation>(GenerationConfig.WantsAutomaticTessellation());
-		StepPermutationVector.Set<FRibbonWantsConstantTessellation>(GenerationConfig.WantsConstantTessellation());
-		StepPermutationVector.Set<FRibbonHasTwist>(GenerationConfig.HasTwist());
 		TShaderMapRef<FNiagaraRibbonAggregationStepCS> AggregationStepShader(GetGlobalShaderMap(GMaxRHIFeatureLevel), StepPermutationVector);
 
 		FNiagaraRibbonAggregationApplyCS::FPermutationDomain ApplyPermutationVector;
 		ApplyPermutationVector.Set<FRibbonHasFullRibbonID>(GenerationConfig.HasFullRibbonIDs());
 		ApplyPermutationVector.Set<FRibbonHasRibbonID>(GenerationConfig.HasSimpleRibbonIDs());
-		ApplyPermutationVector.Set<FRibbonWantsAutomaticTessellation>(GenerationConfig.WantsAutomaticTessellation());
-		ApplyPermutationVector.Set<FRibbonWantsConstantTessellation>(GenerationConfig.WantsConstantTessellation());
-		ApplyPermutationVector.Set<FRibbonHasTwist>(GenerationConfig.HasTwist());
 		TShaderMapRef<FNiagaraRibbonAggregationApplyCS> AggregationApplyShader(GetGlobalShaderMap(GMaxRHIFeatureLevel), ApplyPermutationVector);
 
 		// repeatedly run the step shader to perform the aggregation
@@ -2801,9 +2781,6 @@ void FNiagaraRendererRibbons::InitializeVertexBuffersGPU(FRHICommandListImmediat
 			FNiagaraRibbonVertexReductionFinalizeCS::FPermutationDomain PermutationVector;
 			PermutationVector.Set<FRibbonHasFullRibbonID>(GenerationConfig.HasFullRibbonIDs());
 			PermutationVector.Set<FRibbonHasRibbonID>(GenerationConfig.HasSimpleRibbonIDs());
-			PermutationVector.Set<FRibbonWantsAutomaticTessellation>(GenerationConfig.WantsAutomaticTessellation());
-			PermutationVector.Set<FRibbonWantsConstantTessellation>(GenerationConfig.WantsConstantTessellation());
-			PermutationVector.Set<FRibbonHasTwist>(GenerationConfig.HasTwist());
 			
 			TShaderMapRef<FNiagaraRibbonVertexReductionFinalizeCS> ComputeShader(GetGlobalShaderMap(GMaxRHIFeatureLevel), PermutationVector);
 
@@ -2857,8 +2834,6 @@ void FNiagaraRendererRibbons::InitializeVertexBuffersGPU(FRHICommandListImmediat
 				FNiagaraRibbonUVParamCalculationCS::FPermutationDomain PermutationVector;
 				PermutationVector.Set<FRibbonHasFullRibbonID>(GenerationConfig.HasFullRibbonIDs());
 				PermutationVector.Set<FRibbonHasRibbonID>(GenerationConfig.HasSimpleRibbonIDs());
-				PermutationVector.Set<FRibbonWantsAutomaticTessellation>(GenerationConfig.WantsAutomaticTessellation());
-				PermutationVector.Set<FRibbonWantsConstantTessellation>(GenerationConfig.WantsConstantTessellation());
 
 				TShaderMapRef<FNiagaraRibbonUVParamCalculationCS> ComputeShader(GetGlobalShaderMap(GMaxRHIFeatureLevel), PermutationVector);
 
@@ -2867,7 +2842,7 @@ void FNiagaraRendererRibbons::InitializeVertexBuffersGPU(FRHICommandListImmediat
 					FRHITransitionInfo(VertexBuffers.RibbonLookupTableBuffer.Buffer, ERHIAccess::UAVCompute, ERHIAccess::UAVCompute),
 					FRHITransitionInfo(VertexBuffers.TangentsAndDistancesBuffer.Buffer, ERHIAccess::SRVMask, ERHIAccess::UAVCompute)});
 
-				FComputeShaderUtils::DispatchIndirect(RHICmdList, ComputeShader, FinalizeParams, VertexBuffers.GPUComputeCommandBuffer.Buffer, CommandBufferOffset * FNiagaraRibbonCommandBufferLayout::NumElements);
+				FComputeShaderUtils::DispatchIndirect(RHICmdList, ComputeShader, FinalizeParams, VertexBuffers.GPUComputeCommandBuffer.Buffer, CommandBufferOffset * VERTEX_GEN_OUTPUT_DATA_AUTO_TESSELLATION_NUM_ELEMENTS);
 
 				RHICmdList.Transition({
 					FRHITransitionInfo(VertexBuffers.TangentsAndDistancesBuffer.Buffer, ERHIAccess::UAVCompute, ERHIAccess::SRVMask),
diff --git a/Engine/Plugins/FX/Niagara/Source/NiagaraShader/Public/NiagaraRibbonCompute.h b/Engine/Plugins/FX/Niagara/Source/NiagaraShader/Public/NiagaraRibbonCompute.h
--- a/Engine/Plugins/FX/Niagara/Source/NiagaraShader/Public/NiagaraRibbonCompute.h
+++ b/Engine/Plugins/FX/Niagara/Source/NiagaraShader/Public/NiagaraRibbonCompute.h
@@ -26,7 +26,6 @@ BEGIN_SHADER_PARAMETER_STRUCT(FRibbonComputeUniformParameters, NIAGARASHADER_API
 
 	// Int bindings
 	SHADER_PARAMETER(int, RibbonIdDataOffset)	
-	SHADER_PARAMETER(int, RibbonLinkOrderDataOffset)
 
 	// Float bindings
 	SHADER_PARAMETER(int, PositionDataOffset)
@@ -54,19 +53,15 @@ BEGIN_SHADER_PARAMETER_STRUCT(FRibbonComputeUniformParameters, NIAGARASHADER_API
 	SHADER_PARAMETER(int, U0DistributionMode)
 	SHADER_PARAMETER(int, U1DistributionMode)
 
+	SHADER_PARAMETER(int, RibbonLinkOrderIsFloat)
+	SHADER_PARAMETER(int, RibbonLinkOrderDataOffset)
+
+	SHADER_PARAMETER(int, RibbonTessellationMode)	
 END_SHADER_PARAMETER_STRUCT()
 
 
 class FRibbonHasFullRibbonID : SHADER_PERMUTATION_BOOL("RIBBONID_IS_NIAGARAID");
 class FRibbonHasRibbonID : SHADER_PERMUTATION_BOOL("RIBBONID_IS_INT");
-class FRibbonLinkIsFloat : SHADER_PERMUTATION_BOOL("RIBBONLINK_IS_FLOAT");
-
-class FRibbonWantsConstantTessellation : SHADER_PERMUTATION_BOOL("RIBBONS_WANTS_CONSTANT_TESSELLATION");
-class FRibbonWantsAutomaticTessellation : SHADER_PERMUTATION_BOOL("RIBBONS_WANTS_AUTOMATIC_TESSELLATION");
-
-class FRibbonHasTwist : SHADER_PERMUTATION_BOOL("RIBBON_HAS_TWIST");
-
-class FRibbonHasHighSliceComplexity : SHADER_PERMUTATION_BOOL("RIBBON_HAS_HIGH_SLICE_COMPLEXITY");
 
 struct FNiagaraRibbonComputeCommon
 {	
@@ -101,7 +96,7 @@ class FNiagaraRibbonSortPhase1CS : public FGlobalShader
 	static constexpr uint32 BubbleSortGroupWidth = 32;
 
 	using FParameters = FRibbonOrderSortParameters;
-	using FPermutationDomain = TShaderPermutationDomain<FRibbonHasFullRibbonID, FRibbonHasRibbonID, FRibbonLinkIsFloat>;
+	using FPermutationDomain = TShaderPermutationDomain<FRibbonHasFullRibbonID, FRibbonHasRibbonID>;
 
 	static void ModifyCompilationEnvironment(const FGlobalShaderPermutationParameters& Parameters, FShaderCompilerEnvironment& OutEnvironment);
 };
@@ -117,7 +112,7 @@ class FNiagaraRibbonSortPhase2CS : public FGlobalShader
 	static constexpr uint32 ThreadGroupSize = 64;
 
 	using FParameters = FRibbonOrderSortParameters;
-	using FPermutationDomain = TShaderPermutationDomain<FRibbonHasFullRibbonID, FRibbonHasRibbonID, FRibbonLinkIsFloat>;
+	using FPermutationDomain = TShaderPermutationDomain<FRibbonHasFullRibbonID, FRibbonHasRibbonID>;
 
 	static void ModifyCompilationEnvironment(const FGlobalShaderPermutationParameters& Parameters, FShaderCompilerEnvironment& OutEnvironment);
 };
@@ -143,7 +138,7 @@ class FNiagaraRibbonVertexReductionInitializationCS : public FGlobalShader
 	DECLARE_EXPORTED_GLOBAL_SHADER(FNiagaraRibbonVertexReductionInitializationCS, NIAGARASHADER_API);
 	SHADER_USE_PARAMETER_STRUCT(FNiagaraRibbonVertexReductionInitializationCS, FGlobalShader);
 	
-	using FPermutationDomain = TShaderPermutationDomain<FRibbonHasFullRibbonID, FRibbonHasRibbonID, FRibbonWantsConstantTessellation, FRibbonWantsAutomaticTessellation, FRibbonHasTwist>;
+	using FPermutationDomain = TShaderPermutationDomain<FRibbonHasFullRibbonID, FRibbonHasRibbonID>;
 	using FParameters = FNiagaraRibbonVertexReductionParameters;
 	
 	static void ModifyCompilationEnvironment(const FGlobalShaderPermutationParameters& Parameters, FShaderCompilerEnvironment& OutEnvironment);
@@ -168,7 +163,7 @@ class FNiagaraRibbonVertexReductionFinalizeCS : public FGlobalShader
 	DECLARE_EXPORTED_GLOBAL_SHADER(FNiagaraRibbonVertexReductionFinalizeCS, NIAGARASHADER_API);
 	SHADER_USE_PARAMETER_STRUCT(FNiagaraRibbonVertexReductionFinalizeCS, FGlobalShader);
 	
-	using FPermutationDomain = TShaderPermutationDomain<FRibbonHasFullRibbonID, FRibbonHasRibbonID, FRibbonWantsConstantTessellation, FRibbonWantsAutomaticTessellation, FRibbonHasTwist>;
+	using FPermutationDomain = TShaderPermutationDomain<FRibbonHasFullRibbonID, FRibbonHasRibbonID>;
 	using FParameters = FNiagaraRibbonVertexReductionFinalizationParameters;
 
 	static void ModifyCompilationEnvironment(const FGlobalShaderPermutationParameters& Parameters, FShaderCompilerEnvironment& OutEnvironment);
@@ -202,7 +197,7 @@ class FNiagaraRibbonUVParamCalculationCS : public FGlobalShader
 	DECLARE_EXPORTED_GLOBAL_SHADER(FNiagaraRibbonUVParamCalculationCS, NIAGARASHADER_API);
 	SHADER_USE_PARAMETER_STRUCT(FNiagaraRibbonUVParamCalculationCS, FGlobalShader);
 	
-	using FPermutationDomain = TShaderPermutationDomain<FRibbonHasFullRibbonID, FRibbonHasRibbonID, FRibbonWantsConstantTessellation, FRibbonWantsAutomaticTessellation>;
+	using FPermutationDomain = TShaderPermutationDomain<FRibbonHasFullRibbonID, FRibbonHasRibbonID>;
 	using FParameters = FNiagaraRibbonVertexFinalizationParameters;
 	
 	static void ModifyCompilationEnvironment(const FGlobalShaderPermutationParameters& Parameters, FShaderCompilerEnvironment& OutEnvironment);
@@ -222,6 +217,7 @@ BEGIN_SHADER_PARAMETER_STRUCT(FNiagaraRibbonInitializeIndices, NIAGARASHADER_API
 	SHADER_PARAMETER(uint32, IndexGenThreadSize)
 	SHADER_PARAMETER(uint32, TrianglesPerSegment)
 
+	SHADER_PARAMETER(int, RibbonTessellationMode)
 	SHADER_PARAMETER(float, ViewDistance)
 	SHADER_PARAMETER(float, LODDistanceFactor)
 	SHADER_PARAMETER(uint32, TessellationMode)
@@ -241,7 +237,6 @@ class FNiagaraRibbonCreateIndexBufferParamsCS : public FGlobalShader
 	DECLARE_EXPORTED_GLOBAL_SHADER(FNiagaraRibbonCreateIndexBufferParamsCS, NIAGARASHADER_API);
 	SHADER_USE_PARAMETER_STRUCT(FNiagaraRibbonCreateIndexBufferParamsCS, FGlobalShader);
 	
-	using FPermutationDomain = TShaderPermutationDomain<FRibbonWantsConstantTessellation, FRibbonWantsAutomaticTessellation>;
 	using FParameters = FNiagaraRibbonInitializeIndices;
 
 	static void ModifyCompilationEnvironment(const FGlobalShaderPermutationParameters& Parameters, FShaderCompilerEnvironment& OutEnvironment);
@@ -281,7 +276,7 @@ class FNiagaraRibbonCreateIndexBufferCS : public FGlobalShader
 	DECLARE_EXPORTED_GLOBAL_SHADER(FNiagaraRibbonCreateIndexBufferCS, NIAGARASHADER_API);
 	SHADER_USE_PARAMETER_STRUCT(FNiagaraRibbonCreateIndexBufferCS, FGlobalShader);
 	
-	using FPermutationDomain = TShaderPermutationDomain<FRibbonHasFullRibbonID, FRibbonHasRibbonID, FRibbonHasHighSliceComplexity>;;
+	using FPermutationDomain = TShaderPermutationDomain<FRibbonHasFullRibbonID, FRibbonHasRibbonID>;
 	using FParameters = FNiagaraRibbonGenerateIndices;
 
 	static void ModifyCompilationEnvironment(const FGlobalShaderPermutationParameters& Parameters, FShaderCompilerEnvironment& OutEnvironment);
@@ -300,7 +295,7 @@ class FNiagaraRibbonAggregationStepCS : public FGlobalShader
 	DECLARE_EXPORTED_GLOBAL_SHADER(FNiagaraRibbonAggregationStepCS, NIAGARASHADER_API);
 	SHADER_USE_PARAMETER_STRUCT(FNiagaraRibbonAggregationStepCS, FGlobalShader);
 
-	using FPermutationDomain = TShaderPermutationDomain<FRibbonHasFullRibbonID, FRibbonHasRibbonID, FRibbonWantsConstantTessellation, FRibbonWantsAutomaticTessellation, FRibbonHasTwist>;
+	using FPermutationDomain = TShaderPermutationDomain<FRibbonHasFullRibbonID, FRibbonHasRibbonID>;
 	using FParameters = FNiagaraRibbonAggregationStepParameters;
 
 	static void ModifyCompilationEnvironment(const FGlobalShaderPermutationParameters& Parameters, FShaderCompilerEnvironment& OutEnvironment);
@@ -321,7 +316,7 @@ class FNiagaraRibbonAggregationApplyCS : public FGlobalShader
 	DECLARE_EXPORTED_GLOBAL_SHADER(FNiagaraRibbonAggregationApplyCS, NIAGARASHADER_API);
 	SHADER_USE_PARAMETER_STRUCT(FNiagaraRibbonAggregationApplyCS, FGlobalShader);
 
-	using FPermutationDomain = TShaderPermutationDomain<FRibbonHasFullRibbonID, FRibbonHasRibbonID, FRibbonWantsConstantTessellation, FRibbonWantsAutomaticTessellation, FRibbonHasTwist>;
+	using FPermutationDomain = TShaderPermutationDomain<FRibbonHasFullRibbonID, FRibbonHasRibbonID>;
 	using FParameters = FNiagaraRibbonAggregationApplyParameters;
 
 	static void ModifyCompilationEnvironment(const FGlobalShaderPermutationParameters& Parameters, FShaderCompilerEnvironment& OutEnvironment);
