- Fix a crash when collapsing the stack and deleting a dynamic input that is used

#13897

diff --git a/Engine/Plugins/FX/Niagara/Source/NiagaraEditor/Private/ViewModels/Stack/NiagaraStackFunctionInput.cpp b/Engine/Plugins/FX/Niagara/Source/NiagaraEditor/Private/ViewModels/Stack/NiagaraStackFunctionInput.cpp
--- a/Engine/Plugins/FX/Niagara/Source/NiagaraEditor/Private/ViewModels/Stack/NiagaraStackFunctionInput.cpp
+++ b/Engine/Plugins/FX/Niagara/Source/NiagaraEditor/Private/ViewModels/Stack/NiagaraStackFunctionInput.cpp
@@ -700,6 +700,10 @@ FText UNiagaraStackFunctionInput::GetCollapsedStateText() const
 				}
 				CollapsedTextCache = FText::Format(InputValues.DynamicNode->GetScriptData()->CollapsedViewFormat, Arguments);
 			}
+			else
+			{
+				CollapsedTextCache = FText();
+			}
 			break;
 		case EValueMode::Linked:
 		case EValueMode::LinkedShared:
