- Fix for casting shadows as masked from Niagara not working

#tests Make a sprite / mesh / ribbon renderer, ensure shadow maps are enabled and test casting a translucent material as masked works as expected

diff --git a/Engine/Plugins/FX/Niagara/Source/Niagara/Private/NiagaraDebugHud.cpp b/Engine/Plugins/FX/Niagara/Source/Niagara/Private/NiagaraDebugHud.cpp
--- a/Engine/Plugins/FX/Niagara/Source/Niagara/Private/NiagaraDebugHud.cpp
+++ b/Engine/Plugins/FX/Niagara/Source/Niagara/Private/NiagaraDebugHud.cpp
@@ -830,7 +830,7 @@ namespace NiagaraDebugLocal
 							RenderProperties->GetUsedMaterials(&EmitterInstance.Get(), UsedMaterials);
 							for (UMaterialInterface* Material : UsedMaterials)
 							{
-								if (Material && !IsTranslucentBlendMode(*Material))
+								if (Material && (!IsTranslucentBlendMode(*Material) || Material->GetCastDynamicShadowAsMasked()))
 								{
 									bSupportsThisFrameData = false;
 									break;
diff --git a/Engine/Plugins/FX/Niagara/Source/Niagara/Private/NiagaraRenderer.cpp b/Engine/Plugins/FX/Niagara/Source/Niagara/Private/NiagaraRenderer.cpp
--- a/Engine/Plugins/FX/Niagara/Source/Niagara/Private/NiagaraRenderer.cpp
+++ b/Engine/Plugins/FX/Niagara/Source/Niagara/Private/NiagaraRenderer.cpp
@@ -597,17 +597,17 @@ void FNiagaraRenderer::ProcessMaterialParameterBindings(const FNiagaraRendererMa
 	}
 }
 
-bool FNiagaraRenderer::IsViewRenderingOpaqueOnly(const FSceneView* View, bool bCastsVolumetricTranslucentShadow)
+bool FNiagaraRenderer::IsViewRenderingOpaqueOnly(const FSceneView* View, bool bCanCastTranslucentShadows)
 {
 	const bool bShadowView = View->GetDynamicMeshElementsShadowCullFrustum() != nullptr;
-	return View->CustomRenderPass || (bShadowView && !bCastsVolumetricTranslucentShadow);
+	return View->CustomRenderPass || (bShadowView && !bCanCastTranslucentShadows);
 }
 
-bool FNiagaraRenderer::AreViewsRenderingOpaqueOnly(const TArray<const FSceneView*>& Views, int32 ViewVisibilityMask, bool bCastsVolumetricTranslucentShadow)
+bool FNiagaraRenderer::AreViewsRenderingOpaqueOnly(const TArray<const FSceneView*>& Views, int32 ViewVisibilityMask, bool bCanCastTranslucentShadows)
 {
 	for ( int32 i=0; i < Views.Num(); ++i )
 	{
-		if (((ViewVisibilityMask & (1 << i)) != 0) && !IsViewRenderingOpaqueOnly(Views[i], bCastsVolumetricTranslucentShadow))
+		if (((ViewVisibilityMask & (1 << i)) != 0) && !IsViewRenderingOpaqueOnly(Views[i], bCanCastTranslucentShadows))
 		{
 			return false;
 		}
diff --git a/Engine/Plugins/FX/Niagara/Source/Niagara/Private/NiagaraRendererMeshes.cpp b/Engine/Plugins/FX/Niagara/Source/Niagara/Private/NiagaraRendererMeshes.cpp
--- a/Engine/Plugins/FX/Niagara/Source/Niagara/Private/NiagaraRendererMeshes.cpp
+++ b/Engine/Plugins/FX/Niagara/Source/Niagara/Private/NiagaraRendererMeshes.cpp
@@ -337,8 +337,8 @@ void FNiagaraRendererMeshes::PrepareParticleMeshRenderData(FRHICommandListBase&
 	// If these conditions change please update the DebugHUD display also to reflect it
 	// Note: SceneCaptures will use latent data as GpuReadyTickStage < CurrentParticleData->GetGPUDataReadyStage()
 	//       For main pass scene captures we exclude the batches if they are translucent
-	const bool bIsExcludedViewMode = AllowDebugViewmodes() &&
-		(ViewFamily.EngineShowFlags.Wireframe || ViewFamily.EngineShowFlags.ShaderComplexity || ViewFamily.EngineShowFlags.ShaderComplexityWithQuadOverdraw);
+	const bool bIsExcludedViewMode = AllowDebugViewmodes() && (ViewFamily.EngineShowFlags.Wireframe || ViewFamily.EngineShowFlags.ShaderComplexity || ViewFamily.EngineShowFlags.ShaderComplexityWithQuadOverdraw);
+
 	ParticleMeshRenderData.bIsGpuLowLatencyTranslucency =
 		bGpuLowLatencyTranslucency &&
 		GpuReadyTickStage >= CurrentParticleData->GetGPUDataReadyStage() &&
@@ -350,6 +350,8 @@ void FNiagaraRendererMeshes::PrepareParticleMeshRenderData(FRHICommandListBase&
 
 	ParticleMeshRenderData.bHasTranslucentMaterials = false;
 	ParticleMeshRenderData.bAllTranslucentMaterials = ParticleMeshRenderData.DynamicDataMesh->Materials.Num() > 0;
+	ParticleMeshRenderData.bCanCastTranslucentShadows = false;
+
 	for (FMaterialRenderProxy* MaterialProxy : ParticleMeshRenderData.DynamicDataMesh->Materials)
 	{
 		check(MaterialProxy);
@@ -358,13 +360,15 @@ void FNiagaraRendererMeshes::PrepareParticleMeshRenderData(FRHICommandListBase&
 
 		ParticleMeshRenderData.bHasTranslucentMaterials |= bTranslucent;
 		ParticleMeshRenderData.bAllTranslucentMaterials &= bTranslucent;
+		ParticleMeshRenderData.bCanCastTranslucentShadows |= SceneProxy->CastsVolumetricTranslucentShadow() || Material.GetCastDynamicShadowAsMasked();
 
 		// If even one material can cause the mesh to render before FFXSystemInterface::PostRenderOpaque, we cannot use low latency data
 		ParticleMeshRenderData.bIsGpuLowLatencyTranslucency =
 			ParticleMeshRenderData.bIsGpuLowLatencyTranslucency &&
 			!UE::FXRenderingUtils::CanMaterialRenderBeforeFXPostOpaque(ViewFamily, *SceneProxy, Material);
 	}
-	
+	ParticleMeshRenderData.bIsGpuLowLatencyTranslucency = ParticleMeshRenderData.bIsGpuLowLatencyTranslucency && !ParticleMeshRenderData.bCanCastTranslucentShadows;
+
 	ParticleMeshRenderData.SourceParticleData = ParticleMeshRenderData.DynamicDataMesh->GetParticleDataToRender(RHICmdList, ParticleMeshRenderData.bIsGpuLowLatencyTranslucency);
 	
 	// Anything to render?
@@ -1395,7 +1399,7 @@ void FNiagaraRendererMeshes::GetDynamicMeshElements(const TArray<const FSceneVie
 		return;
 	}
 
-	if (ParticleMeshRenderData.bAllTranslucentMaterials && AreViewsRenderingOpaqueOnly(Views, VisibilityMap, SceneProxy->CastsVolumetricTranslucentShadow()))
+	if (ParticleMeshRenderData.bAllTranslucentMaterials && AreViewsRenderingOpaqueOnly(Views, VisibilityMap, ParticleMeshRenderData.bCanCastTranslucentShadows))
 	{
 		return;
 	}
@@ -1554,9 +1558,11 @@ void FNiagaraRendererMeshes::GetDynamicMeshElements(const TArray<const FSceneVie
 						const bool bNeedsPrevTransform = !bTranslucent || Material.IsTranslucencyWritingVelocity();
 						if (bTranslucent)
 						{
+							const bool bCanCastTranslucentShadows = SceneProxy->CastsVolumetricTranslucentShadow() || Material.GetCastDynamicShadowAsMasked();
+
 							// If we are rendering opaque only we can skip this batch
 							//-OPT: If we only have opaque materials we can skip earlier however due to RemappedMaterialIndex potentially being invalid this is tricky
-							if (IsViewRenderingOpaqueOnly(View, SceneProxy->CastsVolumetricTranslucentShadow()))
+							if (IsViewRenderingOpaqueOnly(View, bCanCastTranslucentShadows))
 							{
 								continue;
 							}
@@ -1715,7 +1721,7 @@ void FNiagaraRendererMeshes::GetDynamicRayTracingInstances(FRayTracingInstanceCo
 				continue;
 			}
 
-			const uint32 RemappedMaterialIndex = MeshData.MaterialRemapTable[Section.MaterialIndex];
+			const uint32 RemappedMaterialIndex = MeshData.MaterialRemapTable.IsValidIndex(Section.MaterialIndex) ? MeshData.MaterialRemapTable[Section.MaterialIndex] : INDEX_NONE;
 			if (!ParticleMeshRenderData.DynamicDataMesh->Materials.IsValidIndex(RemappedMaterialIndex))
 			{
 				// This should never occur. Otherwise, the section data changed since initialization
diff --git a/Engine/Plugins/FX/Niagara/Source/Niagara/Private/NiagaraRendererRibbons.cpp b/Engine/Plugins/FX/Niagara/Source/Niagara/Private/NiagaraRendererRibbons.cpp
--- a/Engine/Plugins/FX/Niagara/Source/Niagara/Private/NiagaraRendererRibbons.cpp
+++ b/Engine/Plugins/FX/Niagara/Source/Niagara/Private/NiagaraRendererRibbons.cpp
@@ -882,8 +882,13 @@ void FNiagaraRendererRibbons::GetDynamicMeshElements(const TArray<const FSceneVi
 		}
 	}
 
-	const bool bTranslucentMaterial = DynamicData->Material && IsTranslucentBlendMode(DynamicData->Material->GetIncompleteMaterialWithFallback(FeatureLevel));
-	if (bTranslucentMaterial && AreViewsRenderingOpaqueOnly(Views, VisibilityMap, SceneProxy->CastsVolumetricTranslucentShadow()))
+	// Material should never be nullptr, it will point to the Default Material if it was null from the game thread
+	check(DynamicData->Material);
+	const FMaterial& Material = DynamicData->Material->GetIncompleteMaterialWithFallback(FeatureLevel);
+
+	const bool bTranslucentMaterial = IsTranslucentBlendMode(Material);
+	const bool bCanCastTranslucentShadows = SceneProxy->CastsVolumetricTranslucentShadow() || Material.GetCastDynamicShadowAsMasked();
+	if (bTranslucentMaterial && AreViewsRenderingOpaqueOnly(Views, VisibilityMap, bCanCastTranslucentShadows))
 	{
 		return;
 	}
@@ -933,7 +938,7 @@ void FNiagaraRendererRibbons::GetDynamicMeshElements(const TArray<const FSceneVi
 
 			// If we are rendering opaque only we can skip this batch
 			//-OPT: If we only have opaque materials we can skip earlier however due to RemappedMaterialIndex potentially being invalid this is tricky
-			if (bTranslucentMaterial && IsViewRenderingOpaqueOnly(View, SceneProxy->CastsVolumetricTranslucentShadow()))
+			if (bTranslucentMaterial && IsViewRenderingOpaqueOnly(View, bCanCastTranslucentShadows))
 			{
 				continue;
 			}
diff --git a/Engine/Plugins/FX/Niagara/Source/Niagara/Private/NiagaraRendererSprites.cpp b/Engine/Plugins/FX/Niagara/Source/Niagara/Private/NiagaraRendererSprites.cpp
--- a/Engine/Plugins/FX/Niagara/Source/Niagara/Private/NiagaraRendererSprites.cpp
+++ b/Engine/Plugins/FX/Niagara/Source/Niagara/Private/NiagaraRendererSprites.cpp
@@ -251,6 +251,7 @@ void FNiagaraRendererSprites::PrepareParticleSpriteRenderData(FRHICommandListBas
 	const FMaterial& Material = MaterialRenderProxy->GetIncompleteMaterialWithFallback(FeatureLevel);
 	ParticleSpriteRenderData.BlendMode = Material.GetBlendMode();
 	ParticleSpriteRenderData.bHasTranslucentMaterials = IsTranslucentBlendMode(Material);
+	ParticleSpriteRenderData.bCanCastTranslucentShadows = SceneProxy->CastsVolumetricTranslucentShadow() || Material.GetCastDynamicShadowAsMasked();
 
 	// If these conditions change please update the DebugHUD display also to reflect it
 	bool bLowLatencyTranslucencyEnabled =
@@ -258,7 +259,7 @@ void FNiagaraRendererSprites::PrepareParticleSpriteRenderData(FRHICommandListBas
 		bGpuLowLatencyTranslucency &&
 		Material.GetMaterialDomain() == MD_Surface &&
 		GpuReadyTickStage >= CurrentParticleData->GetGPUDataReadyStage() &&
-		!SceneProxy->CastsVolumetricTranslucentShadow() &&
+		!ParticleSpriteRenderData.bCanCastTranslucentShadows &&
 		ViewFamilySupportLowLatencyTranslucency(ViewFamily);
 
 	if (bLowLatencyTranslucencyEnabled && SceneProxy->ShouldRenderCustomDepth())
@@ -978,7 +979,7 @@ void FNiagaraRendererSprites::GetDynamicMeshElements(const TArray<const FSceneVi
 		return;
 	}
 
-	if (ParticleSpriteRenderData.bHasTranslucentMaterials && AreViewsRenderingOpaqueOnly(Views, VisibilityMap, SceneProxy->CastsVolumetricTranslucentShadow()))
+	if (ParticleSpriteRenderData.bHasTranslucentMaterials && AreViewsRenderingOpaqueOnly(Views, VisibilityMap, ParticleSpriteRenderData.bCanCastTranslucentShadows))
 	{
 		return;
 	}
@@ -1001,7 +1002,7 @@ void FNiagaraRendererSprites::GetDynamicMeshElements(const TArray<const FSceneVi
 			}
 
 			// Scene captures that are rendered in with the regular views will run depth only so we can skip building the batches
-			if (ParticleSpriteRenderData.bHasTranslucentMaterials && IsViewRenderingOpaqueOnly(View, SceneProxy->CastsVolumetricTranslucentShadow()))
+			if (ParticleSpriteRenderData.bHasTranslucentMaterials && IsViewRenderingOpaqueOnly(View, ParticleSpriteRenderData.bCanCastTranslucentShadows))
 			{
 				continue;
 			}
diff --git a/Engine/Plugins/FX/Niagara/Source/Niagara/Public/NiagaraRenderer.h b/Engine/Plugins/FX/Niagara/Source/Niagara/Public/NiagaraRenderer.h
--- a/Engine/Plugins/FX/Niagara/Source/Niagara/Public/NiagaraRenderer.h
+++ b/Engine/Plugins/FX/Niagara/Source/Niagara/Public/NiagaraRenderer.h
@@ -91,8 +91,8 @@ public:
 	virtual bool IsMaterialValid(const UMaterialInterface* Mat)const { return Mat != nullptr; }
 
 	// Determine if we are rendering into an opaque only view
-	static NIAGARA_API bool IsViewRenderingOpaqueOnly(const FSceneView* View, bool bCastsVolumetricTranslucentShadow);
-	static NIAGARA_API bool AreViewsRenderingOpaqueOnly(const TArray<const FSceneView*>& Views, int32 ViewVisibilityMask, bool bCastsVolumetricTranslucentShadow);
+	static NIAGARA_API bool IsViewRenderingOpaqueOnly(const FSceneView* View, bool bCanCastTranslucentShadows);
+	static NIAGARA_API bool AreViewsRenderingOpaqueOnly(const TArray<const FSceneView*>& Views, int32 ViewVisibilityMask, bool bCanCastTranslucentShadows);
 
 	UE_DEPRECATED(5.4, "SortIndices method should be replaced with SortAndCullIndices")
 	static NIAGARA_API void SortIndices(const struct FNiagaraGPUSortInfo& SortInfo, const FNiagaraRendererVariableInfo& SortVariable, const FNiagaraDataBuffer& Buffer, FGlobalDynamicReadBuffer::FAllocation& OutIndices);
diff --git a/Engine/Plugins/FX/Niagara/Source/Niagara/Public/NiagaraRendererMeshes.h b/Engine/Plugins/FX/Niagara/Source/Niagara/Public/NiagaraRendererMeshes.h
--- a/Engine/Plugins/FX/Niagara/Source/Niagara/Public/NiagaraRendererMeshes.h
+++ b/Engine/Plugins/FX/Niagara/Source/Niagara/Public/NiagaraRendererMeshes.h
@@ -51,6 +51,7 @@ protected:
 		bool 							bUseGPUScene = false;
 		bool							bHasTranslucentMaterials = false;
 		bool							bAllTranslucentMaterials = false;
+		bool							bCanCastTranslucentShadows = false;
 		bool							bSortCullOnGpu = false;
 		bool							bNeedsSort = false;
 		bool							bNeedsCull = false;
diff --git a/Engine/Plugins/FX/Niagara/Source/Niagara/Public/NiagaraRendererSprites.h b/Engine/Plugins/FX/Niagara/Source/Niagara/Public/NiagaraRendererSprites.h
--- a/Engine/Plugins/FX/Niagara/Source/Niagara/Public/NiagaraRendererSprites.h
+++ b/Engine/Plugins/FX/Niagara/Source/Niagara/Public/NiagaraRendererSprites.h
@@ -45,6 +45,7 @@ private:
 
 		EBlendMode							BlendMode = BLEND_Opaque;
 		bool								bHasTranslucentMaterials = false;
+		bool								bCanCastTranslucentShadows = false;
 		bool								bSortCullOnGpu = false;
 		bool								bNeedsSort = false;
 		bool								bNeedsCull = false;
