Niagara - Data interface routing fixes and expanded support.
+ When collecting data interfaces from scripts pre and post translate, associate them with a key struct with complete source information instead of just a name to avoid ambibuity and string manipulation.
+ Remove redundant data from the cached default data interfaces struct and rename it to clarify it's purpose.
+ Create a new struct for parameter references which includes emitter handle information and uses flags for namespaces which are extracted from the parameter name at translation time to avoid string wrangling at runtime.
+ Track the DI souce at a more granular level when translating so that the resolve can be done more reliably.
+ Refactor the DI resolve algorithm to better encapsulate handling data from one script at a time and to better track parameter references for emitters, to prevent sharing of script-local DIs, and to enable sharing of unresolved DIs create at resolve time.
+ Refactor testing to be more granular and added additional test cases.

+ Fixes an issue where a single data interface written to multiple parameters might end up with the wrong name in the execution parameter store.
+ Fixes an issue where a single data interfaces written to the same parameter in multiple modules would generate invalid placeholder data interfaces.
+ Adds support for binding different source data interface parameters to the same target data interface parameter if all of the source data interfaces resolve to the same instance.

#virtualized

diff --git a/Engine/Plugins/FX/Niagara/Source/Niagara/Classes/NiagaraConstants.h b/Engine/Plugins/FX/Niagara/Source/Niagara/Classes/NiagaraConstants.h
--- a/Engine/Plugins/FX/Niagara/Source/Niagara/Classes/NiagaraConstants.h
+++ b/Engine/Plugins/FX/Niagara/Source/Niagara/Classes/NiagaraConstants.h
@@ -250,7 +250,6 @@ struct FNiagaraConstants
 	static NIAGARA_API const FString StackContextNamespaceString;
 	static NIAGARA_API const FString SystemNamespaceString;
 	static NIAGARA_API const FString UserNamespaceString;
-	static NIAGARA_API const FString InternalNamespaceString;
 	static NIAGARA_API const FString PreviousNamespaceString;
 
 	static NIAGARA_API const int32 MaxCategoryNameLength;
diff --git a/Engine/Plugins/FX/Niagara/Source/Niagara/Classes/NiagaraScript.h b/Engine/Plugins/FX/Niagara/Source/Niagara/Classes/NiagaraScript.h
--- a/Engine/Plugins/FX/Niagara/Source/Niagara/Classes/NiagaraScript.h
+++ b/Engine/Plugins/FX/Niagara/Source/Niagara/Classes/NiagaraScript.h
@@ -1248,7 +1248,7 @@ public:
 	NIAGARA_API FOnPropertyChanged& OnPropertyChanged();
 
 	/** External call used to identify the values for a successful VM script compilation. OnVMScriptCompiled will be issued in this case.*/
-	void SetVMCompilationResults(const FNiagaraVMExecutableDataId& InCompileId, FNiagaraVMExecutableData& InScriptVM, const FString& EmitterUniqueName, const TMap<FName, UNiagaraDataInterface*>& ObjectNameMap, bool ApplyRapidIterationParameters);
+	void SetVMCompilationResults(const FNiagaraVMExecutableDataId& InCompileId, FNiagaraVMExecutableData& InScriptVM, const TMap<FNiagaraGraphInputDataInterfaceKey, UNiagaraDataInterface*>& GraphInputDataInterfaces, bool ApplyRapidIterationParameters);
 
 	void SetComputeCompilationResults(
 		const ITargetPlatform* TargetPlatform,
@@ -1289,8 +1289,8 @@ public:
 	NIAGARA_API TArray<TObjectPtr<UNiagaraParameterCollection>>& GetCachedParameterCollectionReferences();
 
 #if WITH_EDITORONLY_DATA
-	TArray<FNiagaraScriptDataInterfaceInfo>& GetCachedDefaultDataInterfaces() { return CachedDefaultDataInterfaces; }
-	TConstArrayView<FNiagaraScriptDataInterfaceInfo> GetCachedDefaultDataInterfaces() const { return MakeArrayView(CachedDefaultDataInterfaces); }
+	TArray<FNiagaraScriptDataInterfaceDefaultInfo>& GetCachedGraphDefaultDataInterfaces() { return CachedGraphDefaultDataInterfaces; }
+	TConstArrayView<FNiagaraScriptDataInterfaceDefaultInfo> GetCachedGraphDefaultDataInterfaces() const { return MakeArrayView(CachedGraphDefaultDataInterfaces); }
 
 	TArrayView<FNiagaraScriptUObjectCompileInfo> GetCachedDefaultUObjects() { return CachedDefaultUObjects; }
 	TConstArrayView<FNiagaraScriptUObjectCompileInfo> GetCachedDefaultUObjects() const { return CachedDefaultUObjects; }
@@ -1396,7 +1396,7 @@ public:
 
 #if WITH_EDITORONLY_DATA
 	UPROPERTY()
-	TArray<FNiagaraScriptDataInterfaceInfo> CachedDefaultDataInterfaces;
+	TArray<FNiagaraScriptDataInterfaceDefaultInfo> CachedGraphDefaultDataInterfaces;
 #endif
 
 #if WITH_EDITORONLY_DATA
diff --git a/Engine/Plugins/FX/Niagara/Source/Niagara/Classes/NiagaraScriptSourceBase.h b/Engine/Plugins/FX/Niagara/Source/Niagara/Classes/NiagaraScriptSourceBase.h
--- a/Engine/Plugins/FX/Niagara/Source/Niagara/Classes/NiagaraScriptSourceBase.h
+++ b/Engine/Plugins/FX/Niagara/Source/Niagara/Classes/NiagaraScriptSourceBase.h
@@ -58,7 +58,7 @@ public:
 	virtual ~FNiagaraCompileRequestDuplicateDataBase() {}
 	virtual bool IsDuplicateDataFor(UNiagaraSystem* InSystem, UNiagaraEmitter* InEmitter, UNiagaraScript* InScript) const = 0;
 	virtual void GetDuplicatedObjects(TArray<UObject*>& Objects) = 0;
-	virtual const TMap<FName, UNiagaraDataInterface*>& GetObjectNameMap() = 0;
+	virtual const TMap<FNiagaraGraphInputDataInterfaceKey, UNiagaraDataInterface*>& GetGraphInputDataInterfaces() = 0;
 	virtual const UNiagaraScriptSourceBase* GetScriptSource() const =  0;
 	virtual int32 GetDependentRequestCount() const = 0;
 	virtual TSharedPtr<FNiagaraCompileRequestDuplicateDataBase, ESPMode::ThreadSafe> GetDependentRequest(int32 Index) = 0;
@@ -136,7 +136,7 @@ class UNiagaraScriptSourceBase : public UObject
 	virtual void ReportAnalyticsData(FNiagaraScriptSourceAnalytics& InData) const {};
 #endif
 
-	virtual TMap<FName, UNiagaraDataInterface*> ComputeObjectNameMap(UNiagaraSystem& System, ENiagaraScriptUsage Usage, FGuid UsageId, FString EmitterUniqueName) const { return TMap<FName, UNiagaraDataInterface*>(); }
+	virtual TMap<FNiagaraGraphInputDataInterfaceKey, UNiagaraDataInterface*> CollectGraphInputDataInterfaces(UNiagaraSystem& System, ENiagaraScriptUsage Usage, FGuid UsageId, FString EmitterUniqueName) const { return TMap<FNiagaraGraphInputDataInterfaceKey, UNiagaraDataInterface*>(); }
 
 	virtual FGuid GetCompileBaseId(ENiagaraScriptUsage InUsage, const FGuid& InUsageId) const { return FGuid(); }
 
diff --git a/Engine/Plugins/FX/Niagara/Source/Niagara/Private/NiagaraAsyncCompile.cpp b/Engine/Plugins/FX/Niagara/Source/Niagara/Private/NiagaraAsyncCompile.cpp
--- a/Engine/Plugins/FX/Niagara/Source/Niagara/Private/NiagaraAsyncCompile.cpp
+++ b/Engine/Plugins/FX/Niagara/Source/Niagara/Private/NiagaraAsyncCompile.cpp
@@ -880,40 +880,19 @@ void FNiagaraActiveCompilationDefault::Apply(const FNiagaraQueryCompilationOptio
 			TSharedPtr<FNiagaraVMExecutableData> ExeData = EmitterCompiledScriptPair.CompileResults;
 			UNiagaraScript* CompiledScript = EmitterCompiledScriptPair.CompiledScript;
 
-			// generate the ObjectNameMap from the source (from the duplicated data if available).
-			TMap<FName, UNiagaraDataInterface*> ObjectNameMap;
+			// generate the GraphInputDataInterfaces from the source (from the duplicated data if available).
+			TMap<FNiagaraGraphInputDataInterfaceKey, UNiagaraDataInterface*> GraphInputDataInterfaces;
 			if (AsyncTask->ComputedPrecompileDuplicateData.IsValid())
 			{
-				if (const UNiagaraScriptSourceBase* ScriptSource = AsyncTask->ComputedPrecompileDuplicateData->GetScriptSource())
-				{
-					ObjectNameMap = ScriptSource->ComputeObjectNameMap(*Options.System, CompiledScript->GetUsage(), CompiledScript->GetUsageId(), AsyncTask->UniqueEmitterName);
-				}
-				else
-				{
-					checkf(false, TEXT("Failed to get ScriptSource from ComputedPrecompileDuplicateData"));
-				}
-
-				// we need to replace the DI in the above generated map with the duplicates that we've created as a part of the duplication process
-				TMap<FName, UNiagaraDataInterface*> DuplicatedObjectNameMap = AsyncTask->ComputedPrecompileDuplicateData->GetObjectNameMap();
-				for (auto ObjectMapIt = ObjectNameMap.CreateIterator(); ObjectMapIt; ++ObjectMapIt)
-				{
-					if (UNiagaraDataInterface** Replacement = DuplicatedObjectNameMap.Find(ObjectMapIt.Key()))
-					{
-						ObjectMapIt.Value() = *Replacement;
-					}
-					else
-					{
-						ObjectMapIt.RemoveCurrent();
-					}
-				}
+				GraphInputDataInterfaces = AsyncTask->ComputedPrecompileDuplicateData->GetGraphInputDataInterfaces();
 			}
 			else
 			{
 				const UNiagaraScriptSourceBase* ScriptSource = CompiledScript->GetLatestSource();
-				ObjectNameMap = ScriptSource->ComputeObjectNameMap(*Options.System, CompiledScript->GetUsage(), CompiledScript->GetUsageId(), AsyncTask->UniqueEmitterName);
+				GraphInputDataInterfaces = ScriptSource->CollectGraphInputDataInterfaces(*Options.System, CompiledScript->GetUsage(), CompiledScript->GetUsageId(), AsyncTask->UniqueEmitterName);
 			}
 
-			EmitterCompiledScriptPair.CompiledScript->SetVMCompilationResults(EmitterCompiledScriptPair.CompileId, *ExeData, AsyncTask->UniqueEmitterName, ObjectNameMap, bEvaluateParametersPending);
+			EmitterCompiledScriptPair.CompiledScript->SetVMCompilationResults(EmitterCompiledScriptPair.CompileId, *ExeData, GraphInputDataInterfaces, bEvaluateParametersPending);
 
 			// Synchronize the variables that we actually encountered during precompile so that we can expose them to the end user.
 			FNiagaraUserRedirectionParameterStore& ExposedParameters = Options.System->GetExposedParameters();
diff --git a/Engine/Plugins/FX/Niagara/Source/Niagara/Private/NiagaraCommon.cpp b/Engine/Plugins/FX/Niagara/Source/Niagara/Private/NiagaraCommon.cpp
--- a/Engine/Plugins/FX/Niagara/Source/Niagara/Private/NiagaraCommon.cpp
+++ b/Engine/Plugins/FX/Niagara/Source/Niagara/Private/NiagaraCommon.cpp
@@ -1188,31 +1188,10 @@ FNiagaraVariable FNiagaraUtilities::ConvertVariableToRapidIterationConstantName(
 	return Var;
 }
 
-bool FNiagaraScriptDataInterfaceInfo::IsUserDataInterface() const
-{
-	TStringBuilder<128> NameBuilder;
-	Name.ToString(NameBuilder);
-	return FCString::Strnicmp(NameBuilder.ToString(), PARAM_MAP_USER_STR, 5) == 0;
-}
-
 bool FNiagaraScriptResolvedDataInterfaceInfo::NeedsPerInstanceBinding() const
 {
-	FNameBuilder NameBuilder;
-	Name.ToString(NameBuilder);
-	return NameBuilder.ToView().StartsWith(PARAM_MAP_USER_STR) 
-		|| (ResolvedDataInterface != nullptr && ResolvedDataInterface->PerInstanceDataSize() > 0);
-}
-
-bool FNiagaraScriptDataInterfaceCompileInfo::CanExecuteOnTarget(ENiagaraSimTarget SimTarget) const
-{
-	// Note that this can be called on non-game threads. We ensure that the data interface CDO object is already in existence at application init time.
-	UNiagaraDataInterface* Obj = GetDefaultDataInterface();
-	if (Obj)
-	{
-		return Obj->CanExecuteOnTarget(SimTarget);
-	}
-	UE_LOG(LogNiagara, Error, TEXT("Failed to call CanExecuteOnTarget for DataInterface \"%s\". Perhaps missing a plugin for your project?"), *Name.ToString());
-	return false;
+	return ResolvedParameterReference.HasFlags(ENiagaraCompileParameterReferenceFlags::User) ||
+		(ResolvedDataInterface != nullptr && ResolvedDataInterface->PerInstanceDataSize() > 0);
 }
 
 UNiagaraDataInterface* FNiagaraScriptDataInterfaceCompileInfo::GetDefaultDataInterface() const
@@ -1236,23 +1215,6 @@ UNiagaraDataInterface* FNiagaraScriptDataInterfaceCompileInfo::GetDefaultDataInt
 	return nullptr;
 }
 
-bool FNiagaraScriptDataInterfaceCompileInfo::NeedsPerInstanceBinding()const
-{
-	FNameBuilder NameBuilder;
-	Name.ToString(NameBuilder);
-	if (NameBuilder.ToView().StartsWith(TEXT("User.")))
-	{
-		return true;
-	}
-
-	UNiagaraDataInterface* Obj = GetDefaultDataInterface();
-	if (Obj && Obj->PerInstanceDataSize() > 0)
-	{
-		return true;
-	}
-	return false;
-}
-
 bool FNiagaraScriptDataInterfaceCompileInfo::MatchesClass(const UClass* InClass) const
 {
 	UNiagaraDataInterface* Obj = GetDefaultDataInterface();
diff --git a/Engine/Plugins/FX/Niagara/Source/Niagara/Private/NiagaraCompilationTypes.cpp b/Engine/Plugins/FX/Niagara/Source/Niagara/Private/NiagaraCompilationTypes.cpp
--- a/Engine/Plugins/FX/Niagara/Source/Niagara/Private/NiagaraCompilationTypes.cpp
+++ b/Engine/Plugins/FX/Niagara/Source/Niagara/Private/NiagaraCompilationTypes.cpp
@@ -87,13 +87,15 @@ bool operator==(const FVMExternalFunctionBindingInfo& Lhs, const FVMExternalFunc
 
 bool operator==(const FNiagaraScriptDataInterfaceCompileInfo& Lhs, const FNiagaraScriptDataInterfaceCompileInfo& Rhs)
 {
-	return Lhs.Name == Rhs.Name
+	return Lhs.Id == Rhs.Id
+		&& Lhs.Name == Rhs.Name
 		&& Lhs.UserPtrIdx == Rhs.UserPtrIdx
 		&& Lhs.Type == Rhs.Type
 		&& Lhs.RegisteredFunctions == Rhs.RegisteredFunctions
-		&& Lhs.RegisteredParameterMapRead == Rhs.RegisteredParameterMapRead
-		&& Lhs.RegisteredParameterMapWrite == Rhs.RegisteredParameterMapWrite
-		&& Lhs.bIsPlaceholder == Rhs.bIsPlaceholder;
+		&& Lhs.SourceGraphInputDataInterfaceKey == Rhs.SourceGraphInputDataInterfaceKey
+		&& Lhs.SourceExternalParameter == Rhs.SourceExternalParameter
+		&& Lhs.SourceScriptDIIndex == Rhs.SourceScriptDIIndex
+		&& Lhs.TargetParameter == Rhs.TargetParameter;
 }
 
 bool operator==(const FNiagaraShaderScriptParametersMetadata& Lhs, const FNiagaraShaderScriptParametersMetadata& Rhs)
@@ -334,7 +336,7 @@ public:
 					FNiagaraVMExecutableDataId UpdatedCompileId;
 					TargetScript->ComputeVMCompilationId(UpdatedCompileId, FGuid());
 
-					TMap<FName, UNiagaraDataInterface*> ObjectNameMap;
+					TMap<FNiagaraGraphInputDataInterfaceKey, UNiagaraDataInterface*> GraphInputDataInterfaces;
 
 					// The original implementation would generate DI references from the compilation data (unless things were pulled from
 					// the DDC).  We will always pull the data from the target scripts so that we can avoid any weird caching issues with
@@ -343,15 +345,14 @@ public:
 					// when the compilation results are applied.
 					if (const UNiagaraScriptSourceBase* ScriptSource = TargetScript->GetLatestSource())
 					{
-						ObjectNameMap = ScriptSource->ComputeObjectNameMap(*Options.System, TargetScript->GetUsage(), TargetScript->GetUsageId(), ScriptCompileData.UniqueEmitterName);
+						GraphInputDataInterfaces = ScriptSource->CollectGraphInputDataInterfaces(*Options.System, TargetScript->GetUsage(), TargetScript->GetUsageId(), ScriptCompileData.UniqueEmitterName);
 					}
 
 					constexpr bool bApplyRapidIterationParameters = false;
 					TargetScript->SetVMCompilationResults(
 						UpdatedCompileId,
 						*ScriptCompileData.ExeData,
-						ScriptCompileData.UniqueEmitterName,
-						ObjectNameMap,
+						GraphInputDataInterfaces,
 						bApplyRapidIterationParameters);
 
 					if (!ScriptCompileData.CompiledShaders.IsEmpty())
diff --git a/Engine/Plugins/FX/Niagara/Source/Niagara/Private/NiagaraConstants.cpp b/Engine/Plugins/FX/Niagara/Source/Niagara/Private/NiagaraConstants.cpp
--- a/Engine/Plugins/FX/Niagara/Source/Niagara/Private/NiagaraConstants.cpp
+++ b/Engine/Plugins/FX/Niagara/Source/Niagara/Private/NiagaraConstants.cpp
@@ -62,7 +62,6 @@ const FString FNiagaraConstants::RapidIterationParametersNamespaceString(TEXT("C
 const FString FNiagaraConstants::StackContextNamespaceString(TEXT("StackContext"));
 const FString FNiagaraConstants::SystemNamespaceString(TEXT("System"));
 const FString FNiagaraConstants::UserNamespaceString(TEXT("User"));
-const FString FNiagaraConstants::InternalNamespaceString(TEXT("__INTERNAL__"));
 const FString FNiagaraConstants::PreviousNamespaceString(TEXT("Previous"));
 
 const FName FNiagaraConstants::InputScopeName(TEXT("Input"));
diff --git a/Engine/Plugins/FX/Niagara/Source/Niagara/Private/NiagaraDataInterfaceEmitterBinding.cpp b/Engine/Plugins/FX/Niagara/Source/Niagara/Private/NiagaraDataInterfaceEmitterBinding.cpp
--- a/Engine/Plugins/FX/Niagara/Source/Niagara/Private/NiagaraDataInterfaceEmitterBinding.cpp
+++ b/Engine/Plugins/FX/Niagara/Source/Niagara/Private/NiagaraDataInterfaceEmitterBinding.cpp
@@ -5,6 +5,7 @@
 #include "NiagaraEmitterInstance.h"
 #include "NiagaraScriptRuntimeData.h"
 #include "NiagaraSystem.h"
+#include "NiagaraSystemImpl.h"
 #include "NiagaraSystemInstance.h"
 
 #include UE_INLINE_GENERATED_CPP_BY_NAME(NiagaraDataInterfaceEmitterBinding)
@@ -27,7 +28,7 @@ FNiagaraEmitterInstance* FNiagaraDataInterfaceEmitterBinding::Resolve(const FNia
 		// If we came from an emitter script we need to look over the default data interfaces to find ourself
 		else
 		{
-			FString SourceEmitterName;
+			FGuid SourceEmitterHandleId;
 			UNiagaraSystem* NiagaraSystem = SystemInstance->GetSystem();
 			for (UNiagaraScript* NiagaraScript : { NiagaraSystem->GetSystemUpdateScript(), NiagaraSystem->GetSystemSpawnScript() })
 			{
@@ -38,26 +39,23 @@ FNiagaraEmitterInstance* FNiagaraDataInterfaceEmitterBinding::Resolve(const FNia
 					{
 						continue;
 					}
-					const FNiagaraScriptResolvedDataInterfaceInfo* MatchingResolvedDataInterface = ScriptRuntimeData->GetResolvedDataInterfaces().FindByPredicate(
-						[&DataInterface](const FNiagaraScriptResolvedDataInterfaceInfo& ResolvedDataInterface) { return ResolvedDataInterface.ResolvedDataInterface == DataInterface; });
-					if (MatchingResolvedDataInterface != nullptr)
+					const FNiagaraScriptResolvedDataInterfaceInfo* MatchingResolvedDataInterfaceInfo = ScriptRuntimeData->GetResolvedDataInterfaces().FindByPredicate([&DataInterface]
+						(const FNiagaraScriptResolvedDataInterfaceInfo& ResolvedDataInterfaceInfo) { return ResolvedDataInterfaceInfo.ResolvedDataInterface == DataInterface; });
+					if (MatchingResolvedDataInterfaceInfo != nullptr)
 					{
-						SourceEmitterName = MatchingResolvedDataInterface->ResolvedSourceEmitterName;
+						SourceEmitterHandleId = MatchingResolvedDataInterfaceInfo->OwningEmitterHandleId;
 						break;
 					}
 				}
 			}
 
-			if (SourceEmitterName.IsEmpty() == false)
+			if (SourceEmitterHandleId.IsValid())
 			{
 				for (const FNiagaraEmitterInstanceRef& EmitterInstance : SystemInstance->GetEmitters())
 				{
-					if (const UNiagaraEmitter* CachedEmitter = EmitterInstance->GetEmitter())
+					if (EmitterInstance->GetEmitterHandle().GetId() == SourceEmitterHandleId)
 					{
-						if (CachedEmitter->GetUniqueEmitterName() == SourceEmitterName)
-						{
-							return &EmitterInstance.Get();
-						}
+						return &EmitterInstance.Get();
 					}
 				}
 			}
@@ -125,83 +123,51 @@ const FNiagaraEmitterHandle* FNiagaraDataInterfaceEmitterBinding::ResolveHandle(
 		}
 		else
 		{
-			// Try to find this data interface in one of the emitter's particle scripts.
-			for (const FNiagaraEmitterHandle& EmitterHandle : OwnerSystem->GetEmitterHandles())
+			// Try to find this data interface in one of the scripts resolved data interfaces.
+			const FNiagaraEmitterHandle* FoundEmitterHandle = nullptr;
+			OwnerSystem->ForEachScriptWithOwningContext([DataInterface, &FoundEmitterHandle](const UNiagaraSystem* OwningSystem, const FNiagaraEmitterHandle* OwningEmitterHandle, UNiagaraScript* Script)
 			{
-				FVersionedNiagaraEmitterData* EmitterData = EmitterHandle.GetEmitterData();
-				if (EmitterHandle.GetIsEnabled() && EmitterData)
+				if (FoundEmitterHandle != nullptr)
 				{
-					bool bFoundDI = false;
-					bool bIsEmitterNamespace = false;
-					EmitterData->ForEachScript(
-						[OwnerSystem, &EmitterHandle, DataInterface, &bFoundDI, &bIsEmitterNamespace](const UNiagaraScript* Script)
-						{
-							const FNiagaraScriptRuntimeData* ScriptRuntimeData = OwnerSystem->GetScriptRuntimeData(&EmitterHandle, *Script).Get();
-							if (ScriptRuntimeData == nullptr)
-							{
-								return;
-							}
-							for (const FNiagaraScriptResolvedDataInterfaceInfo& ResolvedDI : ScriptRuntimeData->GetResolvedDataInterfaces())
-							{
-								if (ResolvedDI.ResolvedDataInterface == DataInterface)
-								{
-									bFoundDI = true;
-									bIsEmitterNamespace = FNiagaraVariableBase::IsInNameSpace(FNiagaraConstants::EmitterNamespaceString, ResolvedDI.CompileName);
-									return;
-								}
-							}
-						}
-					);
-
-					if (bFoundDI)
-					{
-						return bIsEmitterNamespace ? &EmitterHandle : nullptr;
-					}
+					return;
 				}
-			}
-			// Check the system scripts since they contain the compiled emitter data.
-			TArray<UNiagaraScript*> SystemScripts = { OwnerSystem->GetSystemSpawnScript(), OwnerSystem->GetSystemUpdateScript() };
-			for (UNiagaraScript* SystemScript : SystemScripts)
-			{
-				const FNiagaraScriptRuntimeData* ScriptRuntimeData = OwnerSystem->GetScriptRuntimeData(nullptr, *SystemScript).Get();
+
+				FGuid EmitterHandleId = OwningEmitterHandle != nullptr ? OwningEmitterHandle->GetId() : FGuid();
+				const FNiagaraScriptRuntimeData* ScriptRuntimeData = OwningSystem->GetScriptRuntimeData(EmitterHandleId, Script->GetUsage(), Script->GetUsageId()).Get();
 				if (ScriptRuntimeData == nullptr)
 				{
-					continue;
+					return;
 				}
+
 				for (const FNiagaraScriptResolvedDataInterfaceInfo& ResolvedDI : ScriptRuntimeData->GetResolvedDataInterfaces())
 				{
 					if (ResolvedDI.ResolvedDataInterface == DataInterface)
 					{
-						if (ResolvedDI.ResolvedSourceEmitterName.IsEmpty() == false)
+						if (OwningEmitterHandle != nullptr)
+						{
+							FoundEmitterHandle = OwningEmitterHandle;
+						}
+						else if (ResolvedDI.OwningEmitterHandleId.IsValid())
 						{
-							return OwnerSystem->GetEmitterHandles().FindByPredicate([&ResolvedDI](const FNiagaraEmitterHandle& EmitterHandle)
-								{ return EmitterHandle.GetUniqueInstanceName() == ResolvedDI.ResolvedSourceEmitterName; });
+							FoundEmitterHandle = OwningSystem->GetEmitterHandles().FindByPredicate([&ResolvedDI](const FNiagaraEmitterHandle& EmitterHandle)
+								{ return EmitterHandle.GetId() == ResolvedDI.OwningEmitterHandleId; });
+						}
+						if (FoundEmitterHandle != nullptr)
+						{
+							break;
 						}
 					}
 				}
+			});
+			if (FoundEmitterHandle != nullptr)
+			{
+				return FoundEmitterHandle;
 			}
 		}
 	}
 	else if (BindingMode == ENiagaraDataInterfaceEmitterBindingMode::Other)
 	{
-		if (!EmitterName.IsNone())
-		{
-			FNameBuilder EmitterNameString;
-			EmitterName.ToString(EmitterNameString);
-			FStringView EmitterNameStringView = EmitterNameString.ToView();
-
-			for (const FNiagaraEmitterHandle& EmitterHandle : OwnerSystem->GetEmitterHandles())
-			{
-				if (UNiagaraEmitter* NiagaraEmitter = EmitterHandle.GetInstance().Emitter)
-				{
-					//-TODO: UniqueEmitterName should probably be a FName?
-					if (EmitterNameStringView.Equals(NiagaraEmitter->GetUniqueEmitterName(), ESearchCase::IgnoreCase))
-					{
-						return &EmitterHandle;
-					}
-				}
-			}
-		}
+		return ResolveHandleOtherBinding(OwnerSystem);
 	}
 	return nullptr;
 }
@@ -216,21 +182,33 @@ const FNiagaraEmitterHandle* FNiagaraDataInterfaceEmitterBinding::ResolveHandle(
 	}
 	else if (BindingMode == ENiagaraDataInterfaceEmitterBindingMode::Other)
 	{
-		if (!EmitterName.IsNone())
-		{
-			FNameBuilder EmitterNameString;
-			EmitterName.ToString(EmitterNameString);
-			FStringView EmitterNameStringView = EmitterNameString.ToView();
+		return ResolveHandleOtherBinding(OwnerSystem);
+	}
+	return nullptr;
+}
 
-			for (const FNiagaraEmitterHandle& EmitterHandle : OwnerSystem->GetEmitterHandles())
+UNiagaraEmitter* FNiagaraDataInterfaceEmitterBinding::Resolve(const UNiagaraDataInterface* DataInterface) const
+{
+	const FNiagaraEmitterHandle* EmitterHandle = ResolveHandle(DataInterface);
+	return EmitterHandle ? EmitterHandle->GetInstance().Emitter : nullptr;
+}
+
+const FNiagaraEmitterHandle* FNiagaraDataInterfaceEmitterBinding::ResolveHandleOtherBinding(const UNiagaraSystem* OwnerSystem) const
+{
+	if (!EmitterName.IsNone())
+	{
+		FNameBuilder EmitterNameString;
+		EmitterName.ToString(EmitterNameString);
+		FStringView EmitterNameStringView = EmitterNameString.ToView();
+
+		for (const FNiagaraEmitterHandle& EmitterHandle : OwnerSystem->GetEmitterHandles())
+		{
+			if (UNiagaraEmitter* NiagaraEmitter = EmitterHandle.GetInstance().Emitter)
 			{
-				if (UNiagaraEmitter* NiagaraEmitter = EmitterHandle.GetInstance().Emitter)
+				//-TODO: UniqueEmitterName should probably be a FName?
+				if (EmitterNameStringView.Equals(NiagaraEmitter->GetUniqueEmitterName(), ESearchCase::IgnoreCase))
 				{
-					//-TODO: UniqueEmitterName should probably be a FName?
-					if (EmitterNameStringView.Equals(NiagaraEmitter->GetUniqueEmitterName(), ESearchCase::IgnoreCase))
-					{
-						return &EmitterHandle;
-					}
+					return &EmitterHandle;
 				}
 			}
 		}
@@ -238,11 +216,6 @@ const FNiagaraEmitterHandle* FNiagaraDataInterfaceEmitterBinding::ResolveHandle(
 	return nullptr;
 }
 
-UNiagaraEmitter* FNiagaraDataInterfaceEmitterBinding::Resolve(const UNiagaraDataInterface* DataInterface) const
-{
-	const FNiagaraEmitterHandle* EmitterHandle = ResolveHandle(DataInterface);
-	return EmitterHandle ? EmitterHandle->GetInstance().Emitter : nullptr;
-}
 
 FString FNiagaraDataInterfaceEmitterBinding::ResolveUniqueName(const UNiagaraDataInterface* DataInterface) const
 {
diff --git a/Engine/Plugins/FX/Niagara/Source/Niagara/Private/NiagaraDataInterfaceGrid2DCollectionReader.cpp b/Engine/Plugins/FX/Niagara/Source/Niagara/Private/NiagaraDataInterfaceGrid2DCollectionReader.cpp
--- a/Engine/Plugins/FX/Niagara/Source/Niagara/Private/NiagaraDataInterfaceGrid2DCollectionReader.cpp
+++ b/Engine/Plugins/FX/Niagara/Source/Niagara/Private/NiagaraDataInterfaceGrid2DCollectionReader.cpp
@@ -81,21 +81,24 @@ bool UNiagaraDataInterfaceGrid2DCollectionReader::InitPerInstanceData(void* PerI
 		{
 			TConstArrayView<FNiagaraScriptResolvedDataInterfaceInfo> DataInterfaceInfo = ExecContext->GetScriptRuntimeData()->GetResolvedDataInterfaces();
 			const TArray<UNiagaraDataInterface*>& DataInterfaces = ExecContext->CombinedParamStore.GetDataInterfaces();
+			if (DataInterfaceInfo.Num() != DataInterfaces.Num())
+			{
+				UE_LOG(LogNiagara, Error, TEXT("Grid2DCollectionReader failed to bind to collection for Emitter: %s and DataInterface: %s.  Resolved data interface count didn't match exec context data interface count."), *this->EmitterName, *this->DIName);
+				return false;
+			}
 
 			FString FullName = FString("Emitter.") + DIName;
-			int Index = 0;
-
-			// #todo(dmp): we are looking at the UObjects that define the DIs here 
-
-			for (UNiagaraDataInterface* Interface : DataInterfaces)
+			for (int32 Index = 0; Index < DataInterfaces.Num(); ++Index)
 			{
-				if (DataInterfaceInfo[Index].CompileName.GetPlainNameString() == FullName)
+				UNiagaraDataInterface* Interface = DataInterfaces[Index];
+				const FNiagaraScriptResolvedDataInterfaceInfo& Info = DataInterfaceInfo[Index];
+				if (Info.TargetParameterReference.ParameterFullName.GetPlainNameString() == FullName ||
+					Info.ResolvedParameterReference.ParameterFullName.GetPlainNameString() == FullName)
 				{
-					GridInterface = static_cast<UNiagaraDataInterfaceGrid2DCollection*>(Interface);
+					GridInterface = Cast<UNiagaraDataInterfaceGrid2DCollection>(Interface);
 					ProxyToUse = Interface->GetProxy();
 					break;
 				}
-				++Index;
 			}
 		}
 
diff --git a/Engine/Plugins/FX/Niagara/Source/Niagara/Private/NiagaraDataInterfaceGrid3DCollectionReader.cpp b/Engine/Plugins/FX/Niagara/Source/Niagara/Private/NiagaraDataInterfaceGrid3DCollectionReader.cpp
--- a/Engine/Plugins/FX/Niagara/Source/Niagara/Private/NiagaraDataInterfaceGrid3DCollectionReader.cpp
+++ b/Engine/Plugins/FX/Niagara/Source/Niagara/Private/NiagaraDataInterfaceGrid3DCollectionReader.cpp
@@ -80,21 +80,24 @@ bool UNiagaraDataInterfaceGrid3DCollectionReader::InitPerInstanceData(void* PerI
 		{
 			TConstArrayView<FNiagaraScriptResolvedDataInterfaceInfo> DataInterfaceInfo = ExecContext->GetScriptRuntimeData()->GetResolvedDataInterfaces();
 			const TArray<UNiagaraDataInterface*>& DataInterfaces = ExecContext->CombinedParamStore.GetDataInterfaces();
+			if (DataInterfaceInfo.Num() != DataInterfaces.Num())
+			{
+				UE_LOG(LogNiagara, Error, TEXT("Grid3DCollectionReader failed to bind to collection for Emitter: %s and DataInterface: %s.  Resolved data interface count didn't match exec context data interface count."), *this->EmitterName, *this->DIName);
+				return false;
+			}
 
 			FString FullName = FString("Emitter.") + DIName;
-			int Index = 0;
-
-			// #todo(dmp): we are looking at the UObjects that define the DIs here 
-			
-			for (UNiagaraDataInterface* Interface : DataInterfaces)
+			for (int32 Index = 0; Index < DataInterfaces.Num(); ++Index)
 			{
-				if (DataInterfaceInfo[Index].CompileName.GetPlainNameString() == FullName)
+				UNiagaraDataInterface* Interface = DataInterfaces[Index];
+				const FNiagaraScriptResolvedDataInterfaceInfo& Info = DataInterfaceInfo[Index];
+				if (Info.TargetParameterReference.ParameterFullName.GetPlainNameString() == FullName ||
+					Info.ResolvedParameterReference.ParameterFullName.GetPlainNameString() == FullName)
 				{
-					GridInterface = static_cast<UNiagaraDataInterfaceGrid3DCollection*>(Interface);
+					GridInterface = Cast<UNiagaraDataInterfaceGrid3DCollection>(Interface);
 					ProxyToUse = Interface->GetProxy();
 					break;
 				}
-				++Index;
 			}
 		}
 
diff --git a/Engine/Plugins/FX/Niagara/Source/Niagara/Private/NiagaraDataInterfaceParticleRead.cpp b/Engine/Plugins/FX/Niagara/Source/Niagara/Private/NiagaraDataInterfaceParticleRead.cpp
--- a/Engine/Plugins/FX/Niagara/Source/Niagara/Private/NiagaraDataInterfaceParticleRead.cpp
+++ b/Engine/Plugins/FX/Niagara/Source/Niagara/Private/NiagaraDataInterfaceParticleRead.cpp
@@ -8,6 +8,7 @@
 #include "NiagaraGPUSystemTick.h"
 #include "NiagaraRenderer.h"
 #include "NiagaraSimStageData.h"
+#include "NiagaraSystemImpl.h"
 #include "NiagaraSystemInstance.h"
 #include "NiagaraShaderParametersBuilder.h"
 
@@ -2407,20 +2408,24 @@ void UNiagaraDataInterfaceParticleRead::GetFeedback(UNiagaraSystem* Asset, UNiag
 		return false;
 	}();
 
-
-	for (const auto Script : Scripts)
+	Asset->ForEachScriptWithRuntimeData([&](const UNiagaraScript* Script, const FNiagaraScriptRuntimeData* ScriptRuntimeData)
 	{
-		const TArray<FNiagaraScriptDataInterfaceInfo>& CachedDefaultDIs = Script->GetCachedDefaultDataInterfaces();
+		const TArray<FNiagaraScriptDataInterfaceCompileInfo>& CompileInfos = Script->GetVMExecutableData().DataInterfaceInfo;
+		TConstArrayView<FNiagaraScriptResolvedDataInterfaceInfo> ResolvedInfos = ScriptRuntimeData->GetResolvedDataInterfaces();
+		if (CompileInfos.Num() != ResolvedInfos.Num())
+		{
+			return;
+		}
 
-		for (int32 Idx = 0; Idx < Script->GetVMExecutableData().DataInterfaceInfo.Num(); Idx++)
+		for (int32 Idx = 0; Idx < CompileInfos.Num(); Idx++)
 		{
-			const auto& DIInfo = Script->GetVMExecutableData().DataInterfaceInfo[Idx];
-			if (!DIInfo.MatchesClass(GetClass()))
+			const auto& DIInfo = CompileInfos[Idx];
+			if (DIInfo.Type.GetClass() != GetClass())
 			{
 				continue;
 			}
 
-			if (!CachedDefaultDIs.IsValidIndex(Idx) || CachedDefaultDIs[Idx].DataInterface == nullptr || !CachedDefaultDIs[Idx].DataInterface->Equals(this))
+			if (ResolvedInfos[Idx].ResolvedDataInterface == nullptr || !ResolvedInfos[Idx].ResolvedDataInterface->Equals(this))
 			{
 				continue;
 			}
@@ -2488,7 +2493,7 @@ void UNiagaraDataInterfaceParticleRead::GetFeedback(UNiagaraSystem* Asset, UNiag
 				}
 			}
 		}
-	}
+	});
 
 	// If we found persistent ID functions in use and the target emitter isn't set to expose them, trigger a fixable warning.
 	FVersionedNiagaraEmitterData* EmitterData = FoundSourceEmitter.GetEmitterData();
diff --git a/Engine/Plugins/FX/Niagara/Source/Niagara/Private/NiagaraDataInterfaceVectorField.cpp b/Engine/Plugins/FX/Niagara/Source/Niagara/Private/NiagaraDataInterfaceVectorField.cpp
--- a/Engine/Plugins/FX/Niagara/Source/Niagara/Private/NiagaraDataInterfaceVectorField.cpp
+++ b/Engine/Plugins/FX/Niagara/Source/Niagara/Private/NiagaraDataInterfaceVectorField.cpp
@@ -291,7 +291,7 @@ void UNiagaraDataInterfaceVectorField::GetFeedback(UNiagaraSystem* InAsset, UNia
 					{
 						// We assume that if the properties match or we are referencing an external variable whose name is in the list of candidates found in the prior search, it's a valid match for us.
 						if (ResolvedDIs.IsValidIndex(Idx) && ResolvedDIs[Idx].ResolvedDataInterface != nullptr &&
-							(ResolvedDIs[Idx].ResolvedDataInterface->Equals(this) || DIAliases.Contains(ResolvedDIs[Idx].CompileName)))
+							(ResolvedDIs[Idx].ResolvedDataInterface->Equals(this) || DIAliases.Contains(ResolvedDIs[Idx].ResolvedVariable.GetName())))
 						{
 							FVersionedNiagaraEmitter OuterEmitter = Script->GetOuterEmitter();
 							if (OuterEmitter.Emitter && OuterEmitter.GetEmitterData() && (OuterEmitter.GetEmitterData()->SimTarget == ENiagaraSimTarget::CPUSim || Script->IsSystemScript(Script->Usage)))
diff --git a/Engine/Plugins/FX/Niagara/Source/Niagara/Private/NiagaraEmitterInstanceImpl.cpp b/Engine/Plugins/FX/Niagara/Source/Niagara/Private/NiagaraEmitterInstanceImpl.cpp
--- a/Engine/Plugins/FX/Niagara/Source/Niagara/Private/NiagaraEmitterInstanceImpl.cpp
+++ b/Engine/Plugins/FX/Niagara/Source/Niagara/Private/NiagaraEmitterInstanceImpl.cpp
@@ -15,6 +15,7 @@
 #include "NiagaraSimulationStageBase.h"
 #include "NiagaraComponentSettings.h"
 #include "NiagaraSystemSimulation.h"
+#include "NiagaraScriptRuntimeData.h"
 
 DECLARE_DWORD_COUNTER_STAT(TEXT("Num Custom Events"), STAT_NiagaraNumCustomEvents, STATGROUP_Niagara);
 
@@ -663,30 +664,26 @@ void FNiagaraEmitterInstanceImpl::UnbindParameters(bool bExternalOnly)
 {
 	if (bExternalOnly && !IsDisabled())
 	{
-		for (UNiagaraParameterCollection* Collection : SpawnExecContext.Script->GetCachedParameterCollectionReferences())
-		{
-			if (UNiagaraParameterCollectionInstance* NPCInst = ParentSystemInstance->GetParameterCollectionInstance(Collection))
-			{
-				NPCInst->GetParameterStore().Unbind(&SpawnExecContext.Parameters);
-			}
-		}
-		for (UNiagaraParameterCollection* Collection : UpdateExecContext.Script->GetCachedParameterCollectionReferences())
+		ParentSystemInstance->UnbindScriptParameterCollections(SpawnExecContext.Script, SpawnExecContext.Parameters);
+		ParentSystemInstance->UnbindScriptParameterCollections(UpdateExecContext.Script, UpdateExecContext.Parameters);
+
+		FVersionedNiagaraEmitterData* EmitterData = VersionedEmitter.GetEmitterData();
+		if (EmitterData->SimTarget == ENiagaraSimTarget::GPUComputeSim)
 		{
-			if (UNiagaraParameterCollectionInstance* NPCInst = ParentSystemInstance->GetParameterCollectionInstance(Collection))
+			ParentSystemInstance->UnbindScriptParameterCollections(SpawnExecContext.Script, GPUExecContext->CombinedParamStore);
+			ParentSystemInstance->UnbindScriptParameterCollections(UpdateExecContext.Script, GPUExecContext->CombinedParamStore);
+			for (const UNiagaraSimulationStageBase* SimStage : EmitterData->GetSimulationStages())
 			{
-				NPCInst->GetParameterStore().Unbind(&UpdateExecContext.Parameters);
+				if (SimStage->bEnabled)
+				{
+					ParentSystemInstance->UnbindScriptParameterCollections(SimStage->Script, GPUExecContext->CombinedParamStore);
+				}
 			}
 		}
 
 		for (FNiagaraScriptExecutionContext& EventContext : GetEventExecutionContexts())
 		{
-			for (UNiagaraParameterCollection* Collection : EventContext.Script->GetCachedParameterCollectionReferences())
-			{
-				if (UNiagaraParameterCollectionInstance* NPCInst = ParentSystemInstance->GetParameterCollectionInstance(Collection))
-				{
-					NPCInst->GetParameterStore().Unbind(&EventContext.Parameters);
-				}
-			}
+			ParentSystemInstance->UnbindScriptParameterCollections(EventContext.Script, EventContext.Parameters);
 		}
 
 		ParentSystemInstance->UnbindFromParameterStore(RendererBindings);
@@ -716,67 +713,26 @@ void FNiagaraEmitterInstanceImpl::BindParameters(bool bExternalOnly)
 		return;
 	}
 
-	auto BindToParameterCollection = [&](UNiagaraParameterCollection* Collection, FNiagaraParameterStore& DestStore)
-	{
-		if (Collection)
-		{
-			if (UNiagaraParameterCollectionInstance* Inst = ParentSystemInstance->GetParameterCollectionInstance(Collection))
-			{
-				Inst->GetParameterStore().Bind(&DestStore);
-			}
-			else
-			{
-				UE_LOG(LogNiagara, Error, TEXT("Emitter attempting to bind to a null Parameter Collection Instance.\nEmitter:%s\nCollection:%s"), VersionedEmitter.Emitter ? *VersionedEmitter.Emitter->GetPathName() : TEXT("null emitter!"), *Collection->GetPathName());
-			}
-		}
-		else
-		{
-			UE_LOG(LogNiagara, Error, TEXT("Emitter attempting to bind to a null Parameter Collection.\nEmitter:%s"), VersionedEmitter.Emitter ? *VersionedEmitter.Emitter->GetPathName() : TEXT("null emitter!" ));
-		}
-	};
-
-	for (UNiagaraParameterCollection* Collection : SpawnExecContext.Script->GetCachedParameterCollectionReferences())
-	{
-		BindToParameterCollection(Collection, SpawnExecContext.Parameters);
-	}
-	for (UNiagaraParameterCollection* Collection : UpdateExecContext.Script->GetCachedParameterCollectionReferences())
-	{
-		BindToParameterCollection(Collection, UpdateExecContext.Parameters);
-	}
+	ParentSystemInstance->BindScriptParameterCollections(SpawnExecContext.Script, SpawnExecContext.GetScriptRuntimeData(), SpawnExecContext.Parameters);
+	ParentSystemInstance->BindScriptParameterCollections(UpdateExecContext.Script, UpdateExecContext.GetScriptRuntimeData(), UpdateExecContext.Parameters);
 
 	FVersionedNiagaraEmitterData* EmitterData = VersionedEmitter.GetEmitterData();
 	if (EmitterData->SimTarget == ENiagaraSimTarget::GPUComputeSim)
 	{
-		for (UNiagaraParameterCollection* Collection : SpawnExecContext.Script->GetCachedParameterCollectionReferences())
-		{
-			BindToParameterCollection(Collection, GPUExecContext->CombinedParamStore);
-		}
-		for (UNiagaraParameterCollection* Collection : UpdateExecContext.Script->GetCachedParameterCollectionReferences())
+		ParentSystemInstance->BindScriptParameterCollections(SpawnExecContext.Script, GPUExecContext->GetScriptRuntimeData(), GPUExecContext->CombinedParamStore);
+		ParentSystemInstance->BindScriptParameterCollections(UpdateExecContext.Script, GPUExecContext->GetScriptRuntimeData(), GPUExecContext->CombinedParamStore);
+		for (const UNiagaraSimulationStageBase* SimStage : EmitterData->GetSimulationStages())
 		{
-			BindToParameterCollection(Collection, GPUExecContext->CombinedParamStore);
+			if (SimStage->bEnabled)
+			{
+				ParentSystemInstance->BindScriptParameterCollections(SimStage->Script, GPUExecContext->GetScriptRuntimeData(), GPUExecContext->CombinedParamStore);
+			}
 		}
 	}
 
 	for (FNiagaraScriptExecutionContext& EventContext : GetEventExecutionContexts())
 	{
-		for (UNiagaraParameterCollection* Collection : EventContext.Script->GetCachedParameterCollectionReferences())
-		{
-			BindToParameterCollection(Collection, EventContext.Parameters);
-		}
-	}
-
-	FNiagaraScriptInstanceParameterStore& TargetParamStore = EmitterData->SimTarget == ENiagaraSimTarget::GPUComputeSim ? GPUExecContext->CombinedParamStore : UpdateExecContext.Parameters;
-	for (const UNiagaraSimulationStageBase* SimStage : EmitterData->GetSimulationStages())
-	{
-		if (SimStage->bEnabled == false)
-		{
-			continue;
-		}
-
-		for (UNiagaraParameterCollection* Collection : SimStage->Script->GetCachedParameterCollectionReferences())
-		{
-			BindToParameterCollection(Collection, TargetParamStore);
-		}
+		ParentSystemInstance->BindScriptParameterCollections(EventContext.Script, EventContext.GetScriptRuntimeData(), EventContext.Parameters);
 	}
 
 	if (!bExternalOnly)
@@ -1461,8 +1417,6 @@ void FNiagaraEmitterInstanceImpl::Tick(float DeltaSeconds)
 			}
 		}
 
-		//GPUExecContext.UpdateInterfaces = VersionedEmitter->UpdateScriptProps.Script->GetCachedDefaultDataInterfaces();
-
 		// copy over the constants for the render thread
 		//
 		if (GbDumpParticleData || System->bDumpDebugEmitterInfo)
diff --git a/Engine/Plugins/FX/Niagara/Source/Niagara/Private/NiagaraResolveDIHelpers.cpp b/Engine/Plugins/FX/Niagara/Source/Niagara/Private/NiagaraResolveDIHelpers.cpp
--- a/Engine/Plugins/FX/Niagara/Source/Niagara/Private/NiagaraResolveDIHelpers.cpp
+++ b/Engine/Plugins/FX/Niagara/Source/Niagara/Private/NiagaraResolveDIHelpers.cpp
@@ -9,6 +9,7 @@
 #include "NiagaraScript.h"
 #include "NiagaraSimulationStageBase.h"
 #include "NiagaraSystem.h"
+#include "NiagaraSystemImpl.h"
 
 #if WITH_EDITORONLY_DATA
 
@@ -16,477 +17,520 @@
 
 namespace FNiagaraResolveDIHelpers
 {
-	struct FDataInterfaceSourceEmitterNamePair
+	struct FScriptDefaultData
 	{
-		UNiagaraDataInterface* DataInterface;
-		FName SourceName;
-		FString SourceEmitterName;
+		TMap<FNiagaraScriptDataInterfaceCompileId, UNiagaraDataInterface*> InternalDefaultMap;
+		TMap<FNiagaraCompileParameterReference, UNiagaraDataInterface*> ParameterDefaultMap;
 	};
 
-	void CollectDIReadsAssignmentsAndBindingsForScript(
-		const UNiagaraScript* TargetScript,
-		const UNiagaraSystem* System,
-		const FString& EmitterName,
-		TMap<FNiagaraVariableBase, FDataInterfaceSourceEmitterNamePair>& OutVariableReadMap,
-		TMap<FNiagaraVariableBase, FDataInterfaceSourceEmitterNamePair>& OutVariableAssignmentMap,
-		TMap<FNiagaraVariableBase, FNiagaraVariableBase>& OutVariableBindingMap,
-		TArray<FText>& OutErrorMessages)
+	/** Collects all valid DI instances from a set of data interface compile infos and maps them by compile id or parameter reference,
+		depending on the source. */
+	void CollectScriptInternalDefaultDIs(
+		const TArray<FNiagaraScriptDataInterfaceCompileInfo>& ScriptDataInterfaceCompileInfo,
+		TConstArrayView<FNiagaraScriptDataInterfaceDefaultInfo> ScriptCachedDefaultDataInterfaces,
+		FScriptDefaultData& OutScriptDefaultData)
 	{
-		if (TargetScript == nullptr)
-		{
-			return;
-		}
-
-		for (const FNiagaraScriptDataInterfaceInfo& CachedDefaultDataInterface : TargetScript->GetCachedDefaultDataInterfaces())
+		for (const FNiagaraScriptDataInterfaceCompileInfo& DataInterfaceCompileInfo : ScriptDataInterfaceCompileInfo)
 		{
-			if (CachedDefaultDataInterface.RegisteredParameterMapWrite != NAME_None)
+			// First collect the default data interfaces defined in the graphs.
+			const FNiagaraScriptDataInterfaceDefaultInfo* DataInterfaceDefaultInfo = ScriptCachedDefaultDataInterfaces.FindByPredicate(
+				[&DataInterfaceCompileInfo](const FNiagaraScriptDataInterfaceDefaultInfo& CachedDefaultInfo)
+				{ return CachedDefaultInfo.Id == DataInterfaceCompileInfo.Id && CachedDefaultInfo.DataInterface != nullptr; });
+			if (DataInterfaceDefaultInfo != nullptr)
 			{
-				FNiagaraVariable WriteVariable(CachedDefaultDataInterface.Type, CachedDefaultDataInterface.RegisteredParameterMapWrite);
-				if (EmitterName.IsEmpty() == false)
+				OutScriptDefaultData.InternalDefaultMap.Add(DataInterfaceCompileInfo.Id, DataInterfaceDefaultInfo->DataInterface);
+				if (DataInterfaceCompileInfo.TargetParameter.IsValid())
 				{
-					WriteVariable = FNiagaraUtilities::ResolveAliases(WriteVariable, FNiagaraAliasContext().ChangeEmitterNameToEmitter(EmitterName));
+					// NOTE: We shouldn't need to collect these here since it should be handled in the resolve, but this is needed to handle legacy 
+					// content where an update script defines a data interface that's used in spawn.  This should be handled by fail if not 
+					// set for data interfaces if that's ever implemented.
+					OutScriptDefaultData.ParameterDefaultMap.Add(DataInterfaceCompileInfo.TargetParameter, DataInterfaceDefaultInfo->DataInterface);
 				}
+			}
+		}
+	}
 
-				// Handle reads and bindings.
-				if (CachedDefaultDataInterface.RegisteredParameterMapRead != NAME_None)
-				{
-					FNiagaraVariable ReadVariable(CachedDefaultDataInterface.Type, CachedDefaultDataInterface.RegisteredParameterMapRead);
-					if (EmitterName.IsEmpty() == false)
-					{
-						ReadVariable = FNiagaraUtilities::ResolveAliases(ReadVariable, FNiagaraAliasContext().ChangeEmitterToEmitterName(EmitterName));
-					}
-
-					FDataInterfaceSourceEmitterNamePair* CurrentRead = OutVariableReadMap.Find(ReadVariable);
-					if (CurrentRead == nullptr)
-					{
-						FDataInterfaceSourceEmitterNamePair& NewRead = OutVariableReadMap.Add(ReadVariable);
-						NewRead.DataInterface = CachedDefaultDataInterface.DataInterface;
-						NewRead.SourceName = CachedDefaultDataInterface.Name;
-						NewRead.SourceEmitterName = CachedDefaultDataInterface.SourceEmitterName;
-					}
-
-					if (ReadVariable != WriteVariable)
-					{
-						FNiagaraVariableBase* CurrentBinding = OutVariableBindingMap.Find(WriteVariable);
-						if (CurrentBinding != nullptr)
-						{
-							if (*CurrentBinding != ReadVariable)
-							{
-								OutErrorMessages.Add(FText::Format(
-									LOCTEXT("MultipleBindingsFormat", "A data interface parameter was the target of multiple bindings in a single system.  The data interface used in the simulation may be incorrect.  Target Parameter: {0} First Read: {1} Current Read: {2}"),
-									FText::FromName(WriteVariable.GetName()),
-									FText::FromName(CurrentBinding->GetName()),
-									FText::FromName(ReadVariable.GetName())));
-							}
-						}
-						else
-						{
-							{
-								OutVariableBindingMap.Add(WriteVariable, ReadVariable);
-							}
-						}
-					}
-				}
-		
-				// Handle assignments
-				FDataInterfaceSourceEmitterNamePair* CurrentAssignment = OutVariableAssignmentMap.Find(WriteVariable);
-				if (CurrentAssignment != nullptr)
-				{
-					if (CurrentAssignment->DataInterface != CachedDefaultDataInterface.DataInterface)
-					{
-						OutErrorMessages.Add(FText::Format(
-							LOCTEXT("MultipleAssignmentsFormat", "A data interface parameter was the target of an assignment multiple times in a single system.  The data interface used in the simulation may be incorrect.  Target Parameter: {0} First Assignment: {1} Current Assignment: {2}"),
-							FText::FromName(WriteVariable.GetName()),
-							FText::FromName(CurrentAssignment->SourceName),
-							FText::FromName(CachedDefaultDataInterface.Name)));
-					}
-				}
-				else
-				{
-					FDataInterfaceSourceEmitterNamePair& NewAssignment = OutVariableAssignmentMap.Add(WriteVariable);
-					NewAssignment.DataInterface = CachedDefaultDataInterface.DataInterface;
-					NewAssignment.SourceName = CachedDefaultDataInterface.Name;
-					NewAssignment.SourceEmitterName = CachedDefaultDataInterface.SourceEmitterName;
-				}
+	/** Creates default instances of DIs for NPC data interfaces reads.  The NPC data interfaces are not available at compile time and will likely
+		be replaced at runtime, but we need valid instances for safety and to avoid creating them at runtime if they're missing. */
+	void CreateScriptNPCDefaults(UObject& OwnerForCreatedDIs, const TArray<FNiagaraScriptDataInterfaceCompileInfo>& ScriptDataInterfaceCompileInfo,
+		TMap<FNiagaraCompileParameterReference, UNiagaraDataInterface*>& OutSystemExternalParameterDefaultMap)
+	{
+		for (const FNiagaraScriptDataInterfaceCompileInfo& DataInterfaceCompileInfo : ScriptDataInterfaceCompileInfo)
+		{
+			if (DataInterfaceCompileInfo.SourceExternalParameter.HasFlags(ENiagaraCompileParameterReferenceFlags::NPC) &&
+				OutSystemExternalParameterDefaultMap.Contains(DataInterfaceCompileInfo.SourceExternalParameter) == false)
+			{
+				UNiagaraDataInterface* NPCDefault = NewObject<UNiagaraDataInterface>(&OwnerForCreatedDIs, DataInterfaceCompileInfo.Type.GetClass());
+				OutSystemExternalParameterDefaultMap.Add(DataInterfaceCompileInfo.SourceExternalParameter, NPCDefault);
 			}
 		}
 	}
 
-	void CollectDIReadsAssignmentsAndBindings(
-		const UNiagaraSystem* System,
-		TMap<FGuid, TMap<FNiagaraVariableBase, FDataInterfaceSourceEmitterNamePair>>& OutEmitterIdToVariableReadsMap,
-		TMap<FGuid, TMap<FNiagaraVariableBase, FDataInterfaceSourceEmitterNamePair>>& OutEmitterIdToVariableAssignmentsMap,
-		TMap<FGuid, TMap<FNiagaraVariableBase, FNiagaraVariableBase>>& OutEmitterIdToVariableBindingsMap,
+	/** Collects default DI instances for all scripts in a system. */
+	void CollectScriptDIResolveDataForSystem(
+		const UNiagaraSystem& System,
+		UObject& OwnerForCreatedDIs,
+		TMap<FGuid, FScriptDefaultData>& OutScriptDefaultDataMap,
+		TMap<FNiagaraCompileParameterReference, UNiagaraDataInterface*>& OutSystemExternalParameterDefaultMap,
 		TArray<FText>& OutErrorMessages)
 	{
-		TMap<FNiagaraVariableBase, FDataInterfaceSourceEmitterNamePair> VariableReadMap;
-		TMap<FNiagaraVariableBase, FDataInterfaceSourceEmitterNamePair> VariableAssignmentMap;
-		TMap<FNiagaraVariableBase, FNiagaraVariableBase> VariableBindingMap;
-
-		// Add user parameters as assignments.
-		const FNiagaraParameterStore& InstanceParameters = System->GetSystemCompiledData().InstanceParamStore;
+		// First add the user parameters as external defaults.
+		const FNiagaraParameterStore& InstanceParameters = System.GetSystemCompiledData().InstanceParamStore;
 		for (const FNiagaraVariableWithOffset& UserParameterWithOffset : InstanceParameters.ReadParameterVariables())
 		{
 			if (UserParameterWithOffset.IsDataInterface())
 			{
+				FNiagaraCompileParameterReference UserParameterReference = FNiagaraCompileParameterReference(
+					UserParameterWithOffset.GetName(), FGuid(), ENiagaraCompileParameterReferenceFlags::User);
 				UNiagaraDataInterface* DataInterface = InstanceParameters.GetDataInterface(UserParameterWithOffset.Offset);
-				if (DataInterface != nullptr)
+				if (DataInterface == nullptr)
 				{
-					FDataInterfaceSourceEmitterNamePair& NewAssignment = VariableAssignmentMap.Add(UserParameterWithOffset);
-					NewAssignment.DataInterface = DataInterface;
-					NewAssignment.SourceEmitterName = FString();
+					DataInterface = NewObject<UNiagaraDataInterface>(&OwnerForCreatedDIs, UserParameterWithOffset.GetType().GetClass());
 				}
+				OutSystemExternalParameterDefaultMap.Add(UserParameterReference, DataInterface);
 			}
 		}
 
-		// Collect system and emitter scripts.
-		CollectDIReadsAssignmentsAndBindingsForScript(System->GetSystemSpawnScript(), System, FString(), VariableReadMap, VariableAssignmentMap, VariableBindingMap, OutErrorMessages);
-		CollectDIReadsAssignmentsAndBindingsForScript(System->GetSystemUpdateScript(), System, FString(), VariableReadMap, VariableAssignmentMap, VariableBindingMap, OutErrorMessages);
-		OutEmitterIdToVariableReadsMap.Add(FGuid(), VariableReadMap);
-		OutEmitterIdToVariableAssignmentsMap.Add(FGuid(), VariableAssignmentMap);
-		OutEmitterIdToVariableBindingsMap.Add(FGuid(), VariableBindingMap);
-
-		// Collect emitter scripts.
-		for (const FNiagaraEmitterHandle& EmitterHandle : System->GetEmitterHandles())
+		auto HandleScript = [&OwnerForCreatedDIs,&OutScriptDefaultDataMap,&OutSystemExternalParameterDefaultMap]
+			(const FNiagaraEmitterHandle* OwningEmitterHandle, const UNiagaraScript* Script, FScriptDefaultData& ScriptDefaultData)
 		{
-			if (EmitterHandle.GetIsEnabled() == false)
-			{
-				continue;
-			}
+			CollectScriptInternalDefaultDIs(Script->GetVMExecutableData().DataInterfaceInfo, Script->GetCachedGraphDefaultDataInterfaces(), ScriptDefaultData);
+			CreateScriptNPCDefaults(OwnerForCreatedDIs, Script->GetVMExecutableData().DataInterfaceInfo, OutSystemExternalParameterDefaultMap);
+		};
+
+		FScriptDefaultData& SystemScriptDefaultData = OutScriptDefaultDataMap.Add(FGuid());
+		HandleScript(nullptr, System.GetSystemSpawnScript(), SystemScriptDefaultData);
+		HandleScript(nullptr, System.GetSystemUpdateScript(), SystemScriptDefaultData);
 
-			FVersionedNiagaraEmitterData* VersionedNiagaraEmitterData = EmitterHandle.GetEmitterData();
+		for (const FNiagaraEmitterHandle& EmitterHandle : System.GetEmitterHandles())
+		{
+			FVersionedNiagaraEmitterData* VersionedNiagaraEmitterData = EmitterHandle.GetIsEnabled() ? EmitterHandle.GetEmitterData() : nullptr;
 			if (VersionedNiagaraEmitterData == nullptr)
 			{
 				continue;
 			}
 
-			TMap<FNiagaraVariableBase, FDataInterfaceSourceEmitterNamePair> EmitterVariableReadMap = VariableReadMap;
-			TMap<FNiagaraVariableBase, FDataInterfaceSourceEmitterNamePair> EmitterVariableAssignmentMap = VariableAssignmentMap;
-			TMap<FNiagaraVariableBase, FNiagaraVariableBase> EmitterVariableBindingMap = VariableBindingMap;
-
-			CollectDIReadsAssignmentsAndBindingsForScript(VersionedNiagaraEmitterData->SpawnScriptProps.Script, System, EmitterHandle.GetUniqueInstanceName(), EmitterVariableReadMap, EmitterVariableAssignmentMap, EmitterVariableBindingMap, OutErrorMessages);
-			CollectDIReadsAssignmentsAndBindingsForScript(VersionedNiagaraEmitterData->UpdateScriptProps.Script, System, EmitterHandle.GetUniqueInstanceName(), EmitterVariableReadMap, EmitterVariableAssignmentMap, EmitterVariableBindingMap, OutErrorMessages);
-
+			FScriptDefaultData& EmitterScriptDefaultData = OutScriptDefaultDataMap.Add(EmitterHandle.GetId());
+			HandleScript(&EmitterHandle, VersionedNiagaraEmitterData->SpawnScriptProps.Script, EmitterScriptDefaultData);
+			HandleScript(&EmitterHandle, VersionedNiagaraEmitterData->UpdateScriptProps.Script, EmitterScriptDefaultData);
 			for (const FNiagaraEventScriptProperties& EventHandler : VersionedNiagaraEmitterData->GetEventHandlers())
 			{
-				CollectDIReadsAssignmentsAndBindingsForScript(EventHandler.Script, System, EmitterHandle.GetUniqueInstanceName(), EmitterVariableReadMap, EmitterVariableAssignmentMap, EmitterVariableBindingMap, OutErrorMessages);
+				HandleScript(&EmitterHandle, EventHandler.Script, EmitterScriptDefaultData);
 			}
-
-			for (const UNiagaraSimulationStageBase* SimulationStage : VersionedNiagaraEmitterData->GetSimulationStages())
+			if (VersionedNiagaraEmitterData->SimTarget == ENiagaraSimTarget::GPUComputeSim)
 			{
-				if (SimulationStage == nullptr || SimulationStage->bEnabled == false)
-				{
-					continue;
-				}
-				CollectDIReadsAssignmentsAndBindingsForScript(SimulationStage->Script, System, EmitterHandle.GetUniqueInstanceName(), EmitterVariableReadMap, EmitterVariableAssignmentMap, EmitterVariableBindingMap, OutErrorMessages);
+				HandleScript(&EmitterHandle, VersionedNiagaraEmitterData->GetGPUComputeScript(), EmitterScriptDefaultData);
 			}
-
-			OutEmitterIdToVariableReadsMap.Add(EmitterHandle.GetId(), EmitterVariableReadMap);
-			OutEmitterIdToVariableAssignmentsMap.Add(EmitterHandle.GetId(), EmitterVariableAssignmentMap);
-			OutEmitterIdToVariableBindingsMap.Add(EmitterHandle.GetId(), EmitterVariableBindingMap);
 		}
 	}
 
-	void CompactBindings(const TMap<FNiagaraVariableBase, FNiagaraVariableBase>& InVariableBindingMap, TMap<FNiagaraVariableBase, FNiagaraVariableBase>& OutCompactedVariableBindingMap, TArray<FText>& OutErrorMessages)
+	/** Initializes the runtime compiled data map for all scrips in the system. */
+	void ResetResolvedDIsForSystem(const UNiagaraSystem& System, TMap<FNiagaraScriptDataKey, FNiagaraScriptRuntimeCompiledData>& OutScriptRuntimeCompiledDataMap)
 	{
-		TArray<FNiagaraVariableBase> SourceVariables;
-		InVariableBindingMap.GetKeys(SourceVariables);
-		TMap<FNiagaraVariableBase, FNiagaraVariableBase> CompactedVariableBindingMap;
-		for (const FNiagaraVariableBase& SourceVariable : SourceVariables)
+		System.ForEachScriptWithOwningContext([&OutScriptRuntimeCompiledDataMap]
+			(const UNiagaraSystem* OwningSystem, const FNiagaraEmitterHandle* OwningEmitterHandle, UNiagaraScript* Script)
 		{
-			FNiagaraVariableBase BoundVariable;
-			const FNiagaraVariableBase* CurrentBoundVariable = &SourceVariable;
-			TSet<FNiagaraVariableBase> SeenBoundVariables;
-			bool bCircularReferenceDetected = false;
-			while (CurrentBoundVariable != nullptr && bCircularReferenceDetected == false)
+			FGuid EmitterHandleId = OwningEmitterHandle != nullptr ? OwningEmitterHandle->GetId() : FGuid();
+			FNiagaraScriptDataKey ScriptDataKey = FNiagaraScriptDataKey(EmitterHandleId, Script->GetUsage(), Script->GetUsageId());
+			FNiagaraScriptRuntimeCompiledData& ScriptRuntimeCompiledData = OutScriptRuntimeCompiledDataMap.FindOrAdd(ScriptDataKey);
+			ScriptRuntimeCompiledData.ResolvedDataInterfaces.Empty(Script->GetVMExecutableData().DataInterfaceInfo.Num());
+			ScriptRuntimeCompiledData.ResolvedDataInterfaces.AddDefaulted(Script->GetVMExecutableData().DataInterfaceInfo.Num());
+			ScriptRuntimeCompiledData.ResolvedUserDataInterfaceBindings.Empty();
+			ScriptRuntimeCompiledData.ResolvedNpcDataInterfaceBindings.Empty();
+		});
+	}
+
+	/** Converts the system external parameter default map to an emitter specific parameter default map for use by particle scripts in that
+		emitter. This step is necessary since emitter parameter data interfaces are collected from the system script and their parameter references
+		include the emitter guid to disambiguate them.  Emitter parameter references in particle scripts do not include emitter specific information
+		so the emitter guid must be stripped from them so that they can correctly resolve.  We also filter out emitter parameters DIs for other emitters
+		here. */
+	void CreateEmitterSpecificParameterDefaults(
+		FGuid TargetEmitterHandleId,
+		const TMap<FNiagaraCompileParameterReference, UNiagaraDataInterface*>& InExternalParameterDefaultMap,
+		TMap<FNiagaraCompileParameterReference, UNiagaraDataInterface*>& OutEmitterExternalParameterDefaultMap)
+	{
+		for (const TPair<FNiagaraCompileParameterReference, UNiagaraDataInterface*>& ExternalParameterDefaultPair : InExternalParameterDefaultMap)
+		{
+			const FNiagaraCompileParameterReference& ExternalParameterReference = ExternalParameterDefaultPair.Key;
+			UNiagaraDataInterface* ExternalDataInterface = ExternalParameterDefaultPair.Value;
+			if (ExternalParameterReference.HasFlags(ENiagaraCompileParameterReferenceFlags::Emitter))
 			{
-				SeenBoundVariables.Add(*CurrentBoundVariable);
-				BoundVariable = *CurrentBoundVariable;
-				CurrentBoundVariable = InVariableBindingMap.Find(BoundVariable);
-				if (CurrentBoundVariable != nullptr && SeenBoundVariables.Contains(*CurrentBoundVariable))
+				if (ExternalParameterReference.EmitterHandleId == TargetEmitterHandleId)
 				{
-					bCircularReferenceDetected = true;
-					OutErrorMessages.Add(FText::Format(
-						LOCTEXT("CircularDependencyWarningFormat", "A data interface parameter circular dependency was found.  The data interface used in the simulation may be incorrect.  Target Parameter: {0} Resolved Parameter: {1}"),
-						FText::FromName(SourceVariable.GetName()),
-						FText::FromName(CurrentBoundVariable->GetName())));
+					FNiagaraCompileParameterReference EmitterParameterReference(ExternalParameterReference.ParameterFullName, FGuid(), ExternalParameterReference.Flags);
+					OutEmitterExternalParameterDefaultMap.Add(EmitterParameterReference, ExternalDataInterface);
 				}
 			}
-			OutCompactedVariableBindingMap.Add(SourceVariable, BoundVariable);
+			else
+			{
+				OutEmitterExternalParameterDefaultMap.Add(ExternalParameterDefaultPair);
+			}
 		}
 	}
 
-	void ResolveDIsForScript(
-		const UNiagaraSystem* System,
-		UNiagaraScript* TargetScript,
-		const FString& EmitterName,
-		const TMap<FNiagaraVariableBase, FDataInterfaceSourceEmitterNamePair>& VariableReadMap,
-		const TMap<FNiagaraVariableBase, FDataInterfaceSourceEmitterNamePair>& VariableAssignmentMap,
-		const TMap<FNiagaraVariableBase, FNiagaraVariableBase>& CompactedVariableBindingMap,
-		FNiagaraScriptRuntimeCompiledData& OutScriptRuntimeCompiledData,
-		TArray<FText>& OutErrorMessages)
+	/** Collects resolved data interfaces for parameters that can be referenced external to the script and maps them by parameter reference for use by
+		dependent scripts. */
+	void CollectExternalParameterDefaults(
+		const TArray<FNiagaraScriptResolvedDataInterfaceInfo>& ResolvedDataInterfaces,
+		TMap<FNiagaraCompileParameterReference, UNiagaraDataInterface*>& ExternalParameterDefaultMap)
 	{
-		if (TargetScript == nullptr)
+		for (const FNiagaraScriptResolvedDataInterfaceInfo& ResolvedDataInterfaceData : ResolvedDataInterfaces)
 		{
-			return;
+			if (ResolvedDataInterfaceData.TargetParameterReference.IsValid() &&
+				ResolvedDataInterfaceData.TargetParameterReference.IsScriptInternalReference() == false &&
+				ResolvedDataInterfaceData.ResolvedDataInterface != nullptr &&
+				ExternalParameterDefaultMap.Contains(ResolvedDataInterfaceData.TargetParameterReference) == false)
+			{
+				ExternalParameterDefaultMap.Add(ResolvedDataInterfaceData.TargetParameterReference, ResolvedDataInterfaceData.ResolvedDataInterface);
+			}
 		}
-		TArray<FNiagaraScriptResolvedDataInterfaceInfo> ResolvedDataInterfaces;
-		ResolvedDataInterfaces.Reserve(TargetScript->GetCachedDefaultDataInterfaces().Num());
-		TArray<FNiagaraResolvedUserDataInterfaceBinding> UserDataInterfaceBindings;
-		int32 ResolvedDataInterfaceIndex = 0;
-		TSet<FNiagaraVariableBase> ParameterStoreVariables;
-		for (const FNiagaraScriptDataInterfaceInfo& CachedDefaultDataInterface : TargetScript->GetCachedDefaultDataInterfaces())
-		{
-			FNiagaraScriptResolvedDataInterfaceInfo& ResolvedDataInterface = ResolvedDataInterfaces.AddDefaulted_GetRef();
-			ResolvedDataInterface.Name = CachedDefaultDataInterface.Name;
-			ResolvedDataInterface.CompileName = CachedDefaultDataInterface.CompileName;
-			ResolvedDataInterface.UserPtrIdx = CachedDefaultDataInterface.UserPtrIdx;
+	}
 
-			if (CachedDefaultDataInterface.RegisteredParameterMapRead != NAME_None)
+	/** Collects default DI instances for unresolved system parameters which were created while resolving particle scripts within an emitter.  The instances are added to a map
+		for the system so that they will be shared among particle scripts in different emitters referencing the same system parameter. */
+	void CollectSystemUnresolvedExternalParameterDefaults(
+		const TMap<FNiagaraCompileParameterReference, UNiagaraDataInterface*>& EmitterUnresolvedExternalParameterDefaultMap,
+		TMap<FNiagaraCompileParameterReference, UNiagaraDataInterface*>& SystemUnresolvedExternalParameterDefaultMap)
+	{
+		for (const TPair<FNiagaraCompileParameterReference, UNiagaraDataInterface*>& EmitterUnresolvedExternalParameterDefaultPair : EmitterUnresolvedExternalParameterDefaultMap)
+		{
+			const FNiagaraCompileParameterReference& EmitterUnresolvedExternalParameterReference = EmitterUnresolvedExternalParameterDefaultPair.Key;
+			if (EmitterUnresolvedExternalParameterReference.HasFlags(ENiagaraCompileParameterReferenceFlags::System) && 
+				SystemUnresolvedExternalParameterDefaultMap.Contains(EmitterUnresolvedExternalParameterReference) == false)
 			{
-				// If this DI is read from a parameter try to resolve it through the binding and assignment or read maps.
-				FNiagaraVariable ReadVariable(CachedDefaultDataInterface.Type, CachedDefaultDataInterface.RegisteredParameterMapRead);
-				if (EmitterName.IsEmpty() == false)
-				{
-					ReadVariable = FNiagaraUtilities::ResolveAliases(ReadVariable, FNiagaraAliasContext().ChangeEmitterToEmitterName(EmitterName));
-				}
-
-				const FNiagaraVariableBase* BoundVariable = CompactedVariableBindingMap.Find(ReadVariable);
-				if (BoundVariable == nullptr)
-				{
-					BoundVariable = &ReadVariable;
-				}
-
-				const FDataInterfaceSourceEmitterNamePair* BoundDataInterface = VariableAssignmentMap.Find(*BoundVariable);
-				if (BoundDataInterface == nullptr)
-				{
-					// If the data interface was not written to explicitly, then it needs to use the default which was cached in the first read.
-					BoundDataInterface = VariableReadMap.Find(*BoundVariable);
-				}
+				UNiagaraDataInterface* UnresolvedExternalDataInterface = EmitterUnresolvedExternalParameterDefaultPair.Value;
+				SystemUnresolvedExternalParameterDefaultMap.Add(EmitterUnresolvedExternalParameterReference, UnresolvedExternalDataInterface);
+			}
+		}
+	}
 
-				/*
-				if (BoundDataInterface == nullptr)
-				{
-					OutErrorMessages.Add(FText::Format(
-						LOCTEXT("MissingReadDataInterface", "A data interface was read from the parameter map, but its default could not be found.  The data interface used in the simulation may be incorrect.  Read Parameter: {0} Bound Parameter: {1}"),
-						FText::FromName(ReadVariable.GetName()),
-						FText::FromName(BoundVariable->GetName())));
-				}
-				*/
+	/** Resolves data interfaces for a set of data interface compile infos using previously discovered or created instances. */
+	void ResolveDIParametersAndDefaults(
+		const TArray<FNiagaraScriptDataInterfaceCompileInfo>& ScriptDataInterfaceCompileInfo,
+		const TMap<FNiagaraScriptDataInterfaceCompileId, UNiagaraDataInterface*>& InternalDefaultMap,
+		const TMap<FNiagaraCompileParameterReference, UNiagaraDataInterface*>& ParameterDefaultMap,
+		const TMap<FNiagaraCompileParameterReference, UNiagaraDataInterface*>& ExternalParameterDefaultMap,
+		const TMap<FNiagaraCompileParameterReference, UNiagaraDataInterface*>& UnresolvedExternalParameterDefaultMap,
+		const TOptional<ENiagaraScriptUsage>& ScriptUsageFilter,
+		TArray<FNiagaraScriptResolvedDataInterfaceInfo>& OutResolvedDataInterfaces)
+	{
+		for (int32 ScriptDIIndex = 0; ScriptDIIndex < ScriptDataInterfaceCompileInfo.Num(); ScriptDIIndex++)
+		{
+			const FNiagaraScriptDataInterfaceCompileInfo& DataInterfaceCompileInfo = ScriptDataInterfaceCompileInfo[ScriptDIIndex];
+			if (ScriptUsageFilter.IsSet() && DataInterfaceCompileInfo.Id.GetScriptUsage() != ScriptUsageFilter.GetValue())
+			{
+				continue;
+			}
 
-				if (BoundDataInterface != nullptr && ParameterStoreVariables.Contains(ReadVariable) == false)
+			UNiagaraDataInterface* ResolvedDataInterfaceInstance = nullptr;
+			FNiagaraCompileParameterReference ResolvedParameterReference;
+			if (DataInterfaceCompileInfo.SourceExternalParameter.IsValid())
+			{
+				UNiagaraDataInterface*const* ExternalDefaultPtr = ExternalParameterDefaultMap.Find(DataInterfaceCompileInfo.SourceExternalParameter);
+				if (ExternalDefaultPtr != nullptr && *ExternalDefaultPtr != nullptr)
 				{
-					ResolvedDataInterface.ResolvedVariable = *BoundVariable;
-					ResolvedDataInterface.ParameterStoreVariable = ReadVariable;
-					ResolvedDataInterface.bIsInternal = false;
-					ResolvedDataInterface.ResolvedDataInterface = BoundDataInterface->DataInterface;
-					ResolvedDataInterface.ResolvedSourceEmitterName = BoundDataInterface->SourceEmitterName;
-					ParameterStoreVariables.Add(ResolvedDataInterface.ParameterStoreVariable);
+					ResolvedDataInterfaceInstance = *ExternalDefaultPtr;
+					ResolvedParameterReference = DataInterfaceCompileInfo.SourceExternalParameter;
 				}
-
-				if (FNiagaraUserRedirectionParameterStore::IsUserParameter(*BoundVariable))
+				else
 				{
-					const FNiagaraParameterStore& InstanceParameters = System->GetSystemCompiledData().InstanceParamStore;
-					const int32* UserParameterIndex = InstanceParameters.FindParameterOffset(*BoundVariable);
-					if (UserParameterIndex != nullptr && *UserParameterIndex != INDEX_NONE)
+					// NOTE: This case handles legacy content where an update script defines a data interface that's used in spawn.  This
+					// should be handled by fail if not set for data interfaces if that's ever implemented.
+					UNiagaraDataInterface*const* ParameterDefaultPtr = ParameterDefaultMap.Find(DataInterfaceCompileInfo.SourceExternalParameter);
+					if (ParameterDefaultPtr != nullptr && *ParameterDefaultPtr != nullptr)
 					{
-						UserDataInterfaceBindings.Add(FNiagaraResolvedUserDataInterfaceBinding(*UserParameterIndex, ResolvedDataInterfaceIndex));
+						ResolvedDataInterfaceInstance = *ParameterDefaultPtr;
+						ResolvedParameterReference = DataInterfaceCompileInfo.SourceExternalParameter;
+					}
+					else
+					{
+						// NOTE: This case handles legacy content where there are multiple reads for a data interface that was never defined.  This
+						// should be handled by fail if not set for data interfaces if that's ever implemented.
+						UNiagaraDataInterface*const* UnresolvedExternalDefault = UnresolvedExternalParameterDefaultMap.Find(DataInterfaceCompileInfo.SourceExternalParameter);
+						if (UnresolvedExternalDefault != nullptr && *UnresolvedExternalDefault != nullptr)
+						{
+							// We don't set the resolved parameter here since it wasn't actually resolved, we only use the data interface so that
+							// we don't create multiple independent data interfaces for the same unresolved read.
+							ResolvedDataInterfaceInstance = *UnresolvedExternalDefault;
+						}
 					}
 				}
 			}
-			else if (CachedDefaultDataInterface.RegisteredParameterMapWrite != NAME_None)
+			if (ResolvedDataInterfaceInstance == nullptr && DataInterfaceCompileInfo.SourceScriptDIIndex != INDEX_NONE)
 			{
-				FNiagaraVariable WriteVariable(CachedDefaultDataInterface.Type, CachedDefaultDataInterface.RegisteredParameterMapWrite);
-				const FDataInterfaceSourceEmitterNamePair* DataInterfaceAssignment = VariableAssignmentMap.Find(WriteVariable);
-				if (DataInterfaceAssignment != nullptr && ParameterStoreVariables.Contains(WriteVariable) == false)
+				if (ensureMsgf(DataInterfaceCompileInfo.SourceScriptDIIndex < ScriptDIIndex, TEXT("Data interface ordering is incorrect.")))
 				{
-					ResolvedDataInterface.ResolvedVariable = WriteVariable;
-					ResolvedDataInterface.ParameterStoreVariable = WriteVariable;
-					ResolvedDataInterface.bIsInternal = false;
-					ResolvedDataInterface.ResolvedDataInterface = DataInterfaceAssignment->DataInterface;
-					ResolvedDataInterface.ResolvedSourceEmitterName = DataInterfaceAssignment->SourceEmitterName;
-					ParameterStoreVariables.Add(ResolvedDataInterface.ParameterStoreVariable);
-				}
+					ResolvedDataInterfaceInstance = OutResolvedDataInterfaces[DataInterfaceCompileInfo.SourceScriptDIIndex].ResolvedDataInterface;
+					ResolvedParameterReference = OutResolvedDataInterfaces[DataInterfaceCompileInfo.SourceScriptDIIndex].ResolvedParameterReference;
+				}	
 			}
-
-			if (ResolvedDataInterface.ResolvedDataInterface == nullptr)
+			if (ResolvedDataInterfaceInstance == nullptr && DataInterfaceCompileInfo.SourceGraphInputDataInterfaceKey.IsValid())
 			{
-				// NPC variables will fail to resolve but we want to main that it's both external & not append the internal namespace
-				if ( FNiagaraVariable::IsInNameSpace(FNiagaraConstants::ParameterCollectionNamespaceString, ResolvedDataInterface.Name) )
+				UNiagaraDataInterface*const* InternalSourceDataInterfacePtr = InternalDefaultMap.Find(DataInterfaceCompileInfo.Id);
+				if (InternalSourceDataInterfacePtr != nullptr && *InternalSourceDataInterfacePtr != nullptr)
 				{
-					FNiagaraVariable NPCVariable(CachedDefaultDataInterface.Type, ResolvedDataInterface.Name);
-					ResolvedDataInterface.ResolvedVariable = NPCVariable;
-					ResolvedDataInterface.ParameterStoreVariable = NPCVariable;
-					ResolvedDataInterface.bIsInternal = false;
-					ResolvedDataInterface.ResolvedDataInterface = CachedDefaultDataInterface.DataInterface;
-					ResolvedDataInterface.ResolvedSourceEmitterName = CachedDefaultDataInterface.SourceEmitterName;
-				}
-				else
-				{
-					// If the DI was not read from a parameter or couldn't be found, use the one cached during compilation, and give it an internal
-					// name to prevent it from being bound incorrectly.
-					FNameBuilder NameBuilder;
-					NameBuilder.Append(FNiagaraConstants::InternalNamespaceString);
-					NameBuilder.AppendChar(TEXT('.'));
-					ResolvedDataInterface.Name.AppendString(NameBuilder);
-
-					FNiagaraVariable InternalVariable(CachedDefaultDataInterface.Type, FName(NameBuilder.ToString()));
-					ResolvedDataInterface.ResolvedVariable = InternalVariable;
-					ResolvedDataInterface.ParameterStoreVariable = InternalVariable;
-					ResolvedDataInterface.bIsInternal = true;
-					ResolvedDataInterface.ResolvedDataInterface = CachedDefaultDataInterface.DataInterface;
-					ResolvedDataInterface.ResolvedSourceEmitterName = CachedDefaultDataInterface.SourceEmitterName;
+					ResolvedDataInterfaceInstance = *InternalSourceDataInterfacePtr;
 				}
 			}
 
-			ResolvedDataInterfaceIndex++;
+			FNiagaraScriptResolvedDataInterfaceInfo& ResolvedDataInterfaceData = OutResolvedDataInterfaces[ScriptDIIndex];
+			ResolvedDataInterfaceData.OwningEmitterHandleId = DataInterfaceCompileInfo.Id.GetEmitterHandleId();
+			ResolvedDataInterfaceData.CompileName = DataInterfaceCompileInfo.Name;
+			ResolvedDataInterfaceData.Type = DataInterfaceCompileInfo.Type;
+			ResolvedDataInterfaceData.ResolvedParameterReference = ResolvedParameterReference;
+			ResolvedDataInterfaceData.TargetParameterReference = DataInterfaceCompileInfo.TargetParameter;
+			ResolvedDataInterfaceData.ResolvedDataInterface = ResolvedDataInterfaceInstance;
+			ResolvedDataInterfaceData.UserPtrIdx = DataInterfaceCompileInfo.UserPtrIdx;
 		}
-		OutScriptRuntimeCompiledData.ResolvedDataInterfaces = ResolvedDataInterfaces;
-		OutScriptRuntimeCompiledData.ResolvedUserDataInterfaceBindings = UserDataInterfaceBindings;
 	}
 
-	void SynchronizeMatchingInternalResolvedDataInterfaces(TArrayView<const FNiagaraScriptResolvedDataInterfaceInfo> SourceResolvedDataInterfaces, TArray<FNiagaraScriptResolvedDataInterfaceInfo>& TargetResolvedDataInterfaces)
+	/** Creates default DI instances for any data interface compile infos which couldn't be resolved to a previously defined
+		or created DI instance. */
+	void GenerateUnresolvedDIs(
+		const TArray<FNiagaraScriptDataInterfaceCompileInfo>& ScriptDataInterfaceCompileInfo,
+		const TOptional<ENiagaraScriptUsage>& ScriptUsageFilter,
+		UObject& OwnerForCreatedDIs,
+		TArray<FNiagaraScriptResolvedDataInterfaceInfo>& ResolvedDataInterfaces,
+		TMap<FNiagaraCompileParameterReference, UNiagaraDataInterface*>& UnresolvedExternalParameterDefaultMap)
 	{
-		for (const FNiagaraScriptResolvedDataInterfaceInfo& SourceResolvedDataInterface : SourceResolvedDataInterfaces)
+		for (int32 ScriptDIIndex = 0; ScriptDIIndex < ScriptDataInterfaceCompileInfo.Num(); ScriptDIIndex++)
 		{
-			for (FNiagaraScriptResolvedDataInterfaceInfo& TargetResolvedDataInterface : TargetResolvedDataInterfaces)
+			const FNiagaraScriptDataInterfaceCompileInfo& DataInterfaceCompileInfo = ScriptDataInterfaceCompileInfo[ScriptDIIndex];
+			FNiagaraScriptResolvedDataInterfaceInfo& ResolvedDataInterfaceData = ResolvedDataInterfaces[ScriptDIIndex];
+			if (ScriptUsageFilter.IsSet() && DataInterfaceCompileInfo.Id.GetScriptUsage() != ScriptUsageFilter.GetValue())
 			{
-				if (TargetResolvedDataInterface.bIsInternal && TargetResolvedDataInterface.Name == SourceResolvedDataInterface.Name)
+				continue;
+			}
+
+			if (ResolvedDataInterfaceData.ResolvedDataInterface == nullptr)
+			{
+				ResolvedDataInterfaceData.ResolvedDataInterface = NewObject<UNiagaraDataInterface>(&OwnerForCreatedDIs, DataInterfaceCompileInfo.Type.GetClass());
+				if (DataInterfaceCompileInfo.SourceExternalParameter.IsValid())
 				{
-					TargetResolvedDataInterface.ResolvedDataInterface = SourceResolvedDataInterface.ResolvedDataInterface;
-					TargetResolvedDataInterface.ResolvedSourceEmitterName = SourceResolvedDataInterface.ResolvedSourceEmitterName;
-					break;
+					// NOTE: This missing default is an unresolved external read and we don't want to create multiple duplicates for multiple
+					// reads, so it needs to be collected here. This should be handled by fail if not set for data interfaces if that's implemented.
+					UnresolvedExternalParameterDefaultMap.Add(DataInterfaceCompileInfo.SourceExternalParameter, ResolvedDataInterfaceData.ResolvedDataInterface);
 				}
 			}
 		}
 	}
 
-	/** Handles the special case where internal data interfaces defined in certain scripts need to be synchronized with other scripts that contain them. */
-	void ResolveInternalDataInterfaces(TArray<FNiagaraScriptRuntimeCompiledData*> SourceScriptsRuntimeCompiledData, FNiagaraScriptRuntimeCompiledData& TargetScriptRuntimeCompiledData)
+	/** Generates resolved DI infos for a specific script using previously collected data from defaults or from higher level script resolves. */
+	void GenerateResolvedDIsForScript(
+		const FNiagaraEmitterHandle* EmitterHandle,
+		const UNiagaraScript* Script,
+		const FScriptDefaultData& ScriptDefaultData,
+		const TOptional<ENiagaraScriptUsage>& ScriptUsageFilter,
+		UObject& OwnerForCreatedDIs,
+		TMap<FNiagaraScriptDataKey, FNiagaraScriptRuntimeCompiledData>& ScriptRuntimeCompiledDataMap,
+		TMap<FNiagaraCompileParameterReference, UNiagaraDataInterface*>& ExternalParameterDefaultMap,
+		TMap<FNiagaraCompileParameterReference, UNiagaraDataInterface*>& UnresolvedExternalParameterDefaultMap)
 	{
-		TArray<FNiagaraScriptResolvedDataInterfaceInfo> ResolvedTargetDataInterfaces;
-		ResolvedTargetDataInterfaces.Append(TargetScriptRuntimeCompiledData.ResolvedDataInterfaces);
-		for (FNiagaraScriptRuntimeCompiledData* SourceScriptRuntimeCompiledData : SourceScriptsRuntimeCompiledData)
-		{
-			TArrayView<const FNiagaraScriptResolvedDataInterfaceInfo> ResolvedSourceParticleDataInterfaces = MakeConstArrayView(SourceScriptRuntimeCompiledData->ResolvedDataInterfaces);
-			SynchronizeMatchingInternalResolvedDataInterfaces(ResolvedSourceParticleDataInterfaces, ResolvedTargetDataInterfaces);
-		}
-		TargetScriptRuntimeCompiledData.ResolvedDataInterfaces = ResolvedTargetDataInterfaces;
-	}
+		FGuid EmitterHandleId = EmitterHandle != nullptr ? EmitterHandle->GetId() : FGuid();
+		FNiagaraScriptDataKey ScriptDataKey = FNiagaraScriptDataKey(EmitterHandleId, Script->GetUsage(), Script->GetUsageId());
+		FNiagaraScriptRuntimeCompiledData& ScriptRuntimeCompiledData = ScriptRuntimeCompiledDataMap[ScriptDataKey];
+		const TArray<FNiagaraScriptDataInterfaceCompileInfo>& ScriptDataInterfaceCompileInfo = Script->GetVMExecutableData().DataInterfaceInfo;
+		
+		ResolveDIParametersAndDefaults(ScriptDataInterfaceCompileInfo, ScriptDefaultData.InternalDefaultMap, ScriptDefaultData.ParameterDefaultMap,
+			ExternalParameterDefaultMap, UnresolvedExternalParameterDefaultMap, ScriptUsageFilter, ScriptRuntimeCompiledData.ResolvedDataInterfaces);
+		GenerateUnresolvedDIs(ScriptDataInterfaceCompileInfo, ScriptUsageFilter, OwnerForCreatedDIs, ScriptRuntimeCompiledData.ResolvedDataInterfaces,
+			UnresolvedExternalParameterDefaultMap);
+		CollectExternalParameterDefaults(ScriptRuntimeCompiledData.ResolvedDataInterfaces, ExternalParameterDefaultMap);
+	};
 
-	void ResolveDIsInternal(
-		UNiagaraSystem* System,
-		const TMap<FGuid, TMap<FNiagaraVariableBase, FDataInterfaceSourceEmitterNamePair>>& EmitterIdToVariableReadsMap,
-		const TMap<FGuid, TMap<FNiagaraVariableBase, FDataInterfaceSourceEmitterNamePair>>& EmitterIdToVariableAssignmentsMap,
-		const TMap<FGuid, TMap<FNiagaraVariableBase, FNiagaraVariableBase>>& EmitterIdToVariableBindingsMap,
+	/** Generates resolved DI infos for all scripts in a given system. */
+	void GenerateResolvedDIsForSystem(
+		const UNiagaraSystem& System,
+		UObject& OwnerForCreatedDIs,
+		const TMap<FGuid, FScriptDefaultData>& ScriptDefaultDataMap,
+		const TMap<FNiagaraCompileParameterReference, UNiagaraDataInterface*>& ExternalParameterDefaultMap,
 		TMap<FNiagaraScriptDataKey, FNiagaraScriptRuntimeCompiledData>& OutScriptRuntimeCompiledDataMap,
 		TArray<FText>& OutErrorMessages)
 	{
-		TMap<FNiagaraVariableBase, FDataInterfaceSourceEmitterNamePair> VariableReadMap = EmitterIdToVariableReadsMap[FGuid()];
-		TMap<FNiagaraVariableBase, FDataInterfaceSourceEmitterNamePair> VariableAssignmentMap = EmitterIdToVariableAssignmentsMap[FGuid()];
-		TMap<FNiagaraVariableBase, FNiagaraVariableBase> VariableBindingMap = EmitterIdToVariableBindingsMap[FGuid()];
-		TMap<FNiagaraVariableBase, FNiagaraVariableBase> CompactedVariableBindingMap;
-		CompactBindings(VariableBindingMap, CompactedVariableBindingMap, OutErrorMessages);
-
-		UNiagaraScript* SystemSpawnScript = System->GetSystemSpawnScript();
-		UNiagaraScript* SystemUpdateScript = System->GetSystemUpdateScript();
-
-		FNiagaraScriptRuntimeCompiledData SystemSpawnRuntimeCompiledData;
-		FNiagaraScriptRuntimeCompiledData SystemUpdateRuntimeCompiledData;
-
-		ResolveDIsForScript(System, SystemSpawnScript, FString(), VariableReadMap, VariableAssignmentMap, CompactedVariableBindingMap, SystemSpawnRuntimeCompiledData, OutErrorMessages);
-		ResolveDIsForScript(System, SystemUpdateScript, FString(), VariableReadMap, VariableAssignmentMap, CompactedVariableBindingMap, SystemUpdateRuntimeCompiledData, OutErrorMessages);
-
-		OutScriptRuntimeCompiledDataMap.Add(FNiagaraScriptDataKey(FGuid(), SystemSpawnScript->GetUsage(), SystemSpawnScript->GetUsageId()), SystemSpawnRuntimeCompiledData);
-		OutScriptRuntimeCompiledDataMap.Add(FNiagaraScriptDataKey(FGuid(), SystemUpdateScript->GetUsage(), SystemUpdateScript->GetUsageId()), SystemUpdateRuntimeCompiledData);
-
-		for (const FNiagaraEmitterHandle& EmitterHandle : System->GetEmitterHandles())
+		const FScriptDefaultData& SystemScriptDefaultData = ScriptDefaultDataMap[FGuid()];
+		TMap<FNiagaraCompileParameterReference, UNiagaraDataInterface*> SystemExternalParameterDefaultMap = ExternalParameterDefaultMap;
+		TMap<FNiagaraCompileParameterReference, UNiagaraDataInterface*> SystemUnresolvedExternalParameterDefaultMap;
+		
+		auto HandleSystemScript = [&OwnerForCreatedDIs,&OutScriptRuntimeCompiledDataMap,&SystemScriptDefaultData,
+			&SystemExternalParameterDefaultMap,&SystemUnresolvedExternalParameterDefaultMap]
+			(const UNiagaraScript* Script, ENiagaraScriptUsage TargetUsage)
 		{
-			if (EmitterHandle.GetIsEnabled() == false)
-			{
-				continue;
-			}
+			GenerateResolvedDIsForScript(nullptr, Script, SystemScriptDefaultData, TargetUsage, OwnerForCreatedDIs, OutScriptRuntimeCompiledDataMap,
+				SystemExternalParameterDefaultMap, SystemUnresolvedExternalParameterDefaultMap);
+		};
+		
+		// Run the resolve for the system spawn and system update scripts, but with filters for only system scripts so that bindings
+		// from system spawn to emitter update can be handled correctly.
+		HandleSystemScript(System.GetSystemSpawnScript(), ENiagaraScriptUsage::SystemSpawnScript);
+		HandleSystemScript(System.GetSystemUpdateScript(), ENiagaraScriptUsage::SystemUpdateScript);
+		
+		// Once all system data interfaces are resolved, we resolve the systems scripts again, but filtered to just emitter scripts data interfaces.
+		HandleSystemScript(System.GetSystemSpawnScript(), ENiagaraScriptUsage::EmitterSpawnScript);
+		HandleSystemScript(System.GetSystemUpdateScript(), ENiagaraScriptUsage::EmitterUpdateScript);
 
-			FVersionedNiagaraEmitterData* VersionedNiagaraEmitterData = EmitterHandle.GetEmitterData();
+		for (const FNiagaraEmitterHandle& EmitterHandle : System.GetEmitterHandles())
+		{
+			FVersionedNiagaraEmitterData* VersionedNiagaraEmitterData = EmitterHandle.GetIsEnabled() ? EmitterHandle.GetEmitterData() : nullptr;
 			if (VersionedNiagaraEmitterData == nullptr)
 			{
 				continue;
 			}
 
-			TMap<FNiagaraVariableBase, FDataInterfaceSourceEmitterNamePair> EmitterVariableReadMap = EmitterIdToVariableReadsMap[EmitterHandle.GetId()];
-			TMap<FNiagaraVariableBase, FDataInterfaceSourceEmitterNamePair> EmitterVariableAssignmentMap = EmitterIdToVariableAssignmentsMap[EmitterHandle.GetId()];
-			TMap<FNiagaraVariableBase, FNiagaraVariableBase> EmitterVariableBindingMap = EmitterIdToVariableBindingsMap[EmitterHandle.GetId()];
-			TMap<FNiagaraVariableBase, FNiagaraVariableBase> EmitterCompactedVariableBindingMap;
-			CompactBindings(EmitterVariableBindingMap, EmitterCompactedVariableBindingMap, OutErrorMessages);
-
-			UNiagaraScript* SpawnScript = VersionedNiagaraEmitterData->SpawnScriptProps.Script;
-			UNiagaraScript* UpdateScript = VersionedNiagaraEmitterData->UpdateScriptProps.Script;
-			UNiagaraScript* GpuScript = VersionedNiagaraEmitterData->GetGPUComputeScript();
-
-			FNiagaraScriptRuntimeCompiledData SpawnRuntimeCompiledData;
-			FNiagaraScriptRuntimeCompiledData UpdateRuntimeCompiledData;
+			const FScriptDefaultData& EmitterScriptDefaultData = ScriptDefaultDataMap[EmitterHandle.GetId()];
+			TMap<FNiagaraCompileParameterReference, UNiagaraDataInterface*> EmitterExternalParameterDefaultMap;
+			TMap<FNiagaraCompileParameterReference, UNiagaraDataInterface*> EmitterUnresolvedExternalParameterDefaultMap;
+			CreateEmitterSpecificParameterDefaults(EmitterHandle.GetId(), SystemExternalParameterDefaultMap, EmitterExternalParameterDefaultMap);
+			CreateEmitterSpecificParameterDefaults(EmitterHandle.GetId(), SystemUnresolvedExternalParameterDefaultMap, EmitterUnresolvedExternalParameterDefaultMap);
 
-			ResolveDIsForScript(System, SpawnScript, EmitterHandle.GetUniqueInstanceName(), EmitterVariableReadMap, EmitterVariableAssignmentMap, EmitterCompactedVariableBindingMap, SpawnRuntimeCompiledData, OutErrorMessages);
-			ResolveDIsForScript(System, UpdateScript, EmitterHandle.GetUniqueInstanceName(), EmitterVariableReadMap, EmitterVariableAssignmentMap, EmitterCompactedVariableBindingMap, UpdateRuntimeCompiledData, OutErrorMessages);
-
-			if (VersionedNiagaraEmitterData->UsesInterpolatedSpawning())
+			auto HandleParticleScript = [&OwnerForCreatedDIs,&OutScriptRuntimeCompiledDataMap,&EmitterHandle,&EmitterScriptDefaultData,
+				&EmitterExternalParameterDefaultMap,&EmitterUnresolvedExternalParameterDefaultMap]
+				(const UNiagaraScript* Script)
 			{
-				TArray<FNiagaraScriptRuntimeCompiledData*> SourceScriptRuntimeCompiledData = { &UpdateRuntimeCompiledData };
-				ResolveInternalDataInterfaces(SourceScriptRuntimeCompiledData, SpawnRuntimeCompiledData);
-			}
-
-			OutScriptRuntimeCompiledDataMap.Add(FNiagaraScriptDataKey(EmitterHandle.GetId(), SpawnScript->GetUsage(), SpawnScript->GetUsageId()), SpawnRuntimeCompiledData);
-			OutScriptRuntimeCompiledDataMap.Add(FNiagaraScriptDataKey(EmitterHandle.GetId(), UpdateScript->GetUsage(), UpdateScript->GetUsageId()), UpdateRuntimeCompiledData);
+				GenerateResolvedDIsForScript(&EmitterHandle, Script, EmitterScriptDefaultData, TOptional<ENiagaraScriptUsage>(), OwnerForCreatedDIs,
+					OutScriptRuntimeCompiledDataMap, EmitterExternalParameterDefaultMap, EmitterUnresolvedExternalParameterDefaultMap);
+			};
 
+			HandleParticleScript(VersionedNiagaraEmitterData->SpawnScriptProps.Script);
+			HandleParticleScript(VersionedNiagaraEmitterData->UpdateScriptProps.Script);
 			for (const FNiagaraEventScriptProperties& EventHandler : VersionedNiagaraEmitterData->GetEventHandlers())
 			{
-				UNiagaraScript* EventScript = EventHandler.Script;
-				FNiagaraScriptRuntimeCompiledData EventRuntimeCompiledData;
-
-				ResolveDIsForScript(System, EventScript, EmitterHandle.GetUniqueInstanceName(), EmitterVariableReadMap, EmitterVariableAssignmentMap, EmitterCompactedVariableBindingMap, EventRuntimeCompiledData, OutErrorMessages);
-				OutScriptRuntimeCompiledDataMap.Add(FNiagaraScriptDataKey(EmitterHandle.GetId(), EventScript->GetUsage(), EventScript->GetUsageId()), EventRuntimeCompiledData);
+				HandleParticleScript(EventHandler.Script);
 			}
-
-			if (VersionedNiagaraEmitterData->SimTarget == ENiagaraSimTarget::GPUComputeSim && GpuScript != nullptr)
+			if (VersionedNiagaraEmitterData->SimTarget == ENiagaraSimTarget::GPUComputeSim)
 			{
-				FNiagaraScriptRuntimeCompiledData GpuRuntimeCompiledData;
-				ResolveDIsForScript(System, GpuScript, EmitterHandle.GetUniqueInstanceName(), EmitterVariableReadMap, EmitterVariableAssignmentMap, EmitterCompactedVariableBindingMap, GpuRuntimeCompiledData, OutErrorMessages);
+				HandleParticleScript(VersionedNiagaraEmitterData->GetGPUComputeScript());
+			}
 
-				TArray<FNiagaraScriptRuntimeCompiledData> SimulationStageScriptRuntimeCompiledDataItems;
-				for (const UNiagaraSimulationStageBase* SimulationStage : VersionedNiagaraEmitterData->GetSimulationStages())
+			// Collect any unresolved data interfaces created for the system here so that they are shared across particle scripts.
+			CollectSystemUnresolvedExternalParameterDefaults(EmitterUnresolvedExternalParameterDefaultMap, SystemUnresolvedExternalParameterDefaultMap);
+		}
+	}
+
+	/** Generates the aliased parameter variables on the resolved DI infos using the source parameter references and the emitter id to name map. */
+	void GenerateDIParameterVariablesForScript(TArray<FNiagaraScriptResolvedDataInterfaceInfo>& ResolvedDataInterfaces,
+		const FNiagaraEmitterHandle* OwningEmitterHandle,
+		const TMap<FGuid, FString>& EmitterHandleIdToNameMap)
+	{
+		for (FNiagaraScriptResolvedDataInterfaceInfo& ResolvedDataInterfaceData : ResolvedDataInterfaces)
+		{
+			// This is a bit inconsistent since emitter read variables are expected to be in the form Emitter.[VariableName], but emitter write variables are required to be in the
+			// form [EmitterName].[VariableName] to disambiguate them in the system parameter store.
+			if (ResolvedDataInterfaceData.ResolvedParameterReference.IsValid())
+			{
+				ResolvedDataInterfaceData.ResolvedVariable = FNiagaraVariableBase(ResolvedDataInterfaceData.Type, ResolvedDataInterfaceData.ResolvedParameterReference.ParameterFullName);
+			}
+			if (ResolvedDataInterfaceData.TargetParameterReference.IsValid())
+			{ 
+				if (ResolvedDataInterfaceData.TargetParameterReference.HasFlags(ENiagaraCompileParameterReferenceFlags::Emitter))
 				{
-					if (SimulationStage == nullptr || SimulationStage->bEnabled == false)
+					FString EmitterName;
+					if (OwningEmitterHandle != nullptr)
 					{
-						continue;
+						EmitterName = OwningEmitterHandle->GetUniqueInstanceName();
+					}
+					else
+					{
+						const FString* FoundEmitterName = EmitterHandleIdToNameMap.Find(ResolvedDataInterfaceData.TargetParameterReference.EmitterHandleId);
+						EmitterName = FoundEmitterName != nullptr ? *FoundEmitterName : FString();
+					}
+					if (EmitterName.IsEmpty() == false)
+					{
+						ResolvedDataInterfaceData.ParameterStoreVariable = FNiagaraUtilities::ResolveAliases(
+							FNiagaraVariable(ResolvedDataInterfaceData.Type, ResolvedDataInterfaceData.TargetParameterReference.ParameterFullName),
+							FNiagaraAliasContext().ChangeEmitterToEmitterName(EmitterName));
 					}
-					UNiagaraScript* SimulationStageScript = SimulationStage->Script;
-					FNiagaraScriptRuntimeCompiledData& SimulationStageRuntimeCompiledData = SimulationStageScriptRuntimeCompiledDataItems.AddDefaulted_GetRef();
-
-					ResolveDIsForScript(System, SimulationStageScript, EmitterHandle.GetUniqueInstanceName(), EmitterVariableReadMap, EmitterVariableAssignmentMap, EmitterCompactedVariableBindingMap, SimulationStageRuntimeCompiledData, OutErrorMessages);
 				}
+				else
+				{
+					ResolvedDataInterfaceData.ParameterStoreVariable = FNiagaraVariableBase(ResolvedDataInterfaceData.Type, ResolvedDataInterfaceData.TargetParameterReference.ParameterFullName);
+				}
+			}
+			else
+			{
+				ResolvedDataInterfaceData.ParameterStoreVariable = FNiagaraVariableBase(ResolvedDataInterfaceData.Type, ResolvedDataInterfaceData.CompileName);
+			}
+		}
+	}
 
-				TArray<FNiagaraScriptRuntimeCompiledData*> SourceScriptRuntimeCompiledData;
-				SourceScriptRuntimeCompiledData.Add(&SpawnRuntimeCompiledData);
-				SourceScriptRuntimeCompiledData.Add(&UpdateRuntimeCompiledData);
-				for (FNiagaraScriptRuntimeCompiledData& SimulationStageScriptRuntimeCompiledData : SimulationStageScriptRuntimeCompiledDataItems)
+	void GenerateUserDIBindingsForScript(const TArray<FNiagaraScriptResolvedDataInterfaceInfo>& ResolvedDataInterfaces,
+		const FNiagaraParameterStore& SystemInstanceParameters,
+		TArray<FNiagaraResolvedUserDataInterfaceBinding>& OutResolvedUserDataInterfaceBindings)
+	{
+		int32 ResolvedDataInterfaceIndex = 0;
+		for (const FNiagaraScriptResolvedDataInterfaceInfo& ResolvedDataInterfaceData : ResolvedDataInterfaces)
+		{
+			if (ResolvedDataInterfaceData.ResolvedParameterReference.HasFlags(ENiagaraCompileParameterReferenceFlags::User))
+			{
+				const int32* UserParameterIndex = SystemInstanceParameters.FindParameterOffset(ResolvedDataInterfaceData.ResolvedVariable);
+				if (UserParameterIndex != nullptr && *UserParameterIndex != INDEX_NONE)
 				{
-					SourceScriptRuntimeCompiledData.Add(&SimulationStageScriptRuntimeCompiledData);
+					OutResolvedUserDataInterfaceBindings.Add(FNiagaraResolvedUserDataInterfaceBinding(*UserParameterIndex, ResolvedDataInterfaceIndex));
 				}
-				ResolveInternalDataInterfaces(SourceScriptRuntimeCompiledData, GpuRuntimeCompiledData);
+			}
+			ResolvedDataInterfaceIndex++;
+		}
+	}
 
-				GpuRuntimeCompiledData.SimTarget = ENiagaraSimTarget::GPUComputeSim;
-				OutScriptRuntimeCompiledDataMap.Add(FNiagaraScriptDataKey(EmitterHandle.GetId(), GpuScript->GetUsage(), GpuScript->GetUsageId()), GpuRuntimeCompiledData);
+	void GenerateNPCDIBindingsForScript(const TArray<FNiagaraScriptResolvedDataInterfaceInfo>& ResolvedDataInterfaces,
+		TArray<FNiagaraResolvedNpcDataInterfaceBinding>& OutResolvedNpcDataInterfaceBindings)
+	{
+		int32 ResolvedDataInterfaceIndex = 0;
+		for (const FNiagaraScriptResolvedDataInterfaceInfo& ResolvedDataInterfaceData : ResolvedDataInterfaces)
+		{
+			if (ResolvedDataInterfaceData.ResolvedParameterReference.IsValid() &&
+				ResolvedDataInterfaceData.ResolvedParameterReference.HasFlags(ENiagaraCompileParameterReferenceFlags::NPC))
+			{
+				FNiagaraResolvedNpcDataInterfaceBinding* NpcBinding = OutResolvedNpcDataInterfaceBindings.FindByPredicate([&ResolvedDataInterfaceData]
+					(const FNiagaraResolvedNpcDataInterfaceBinding& NpcBinding)
+				{
+					return NpcBinding.SourceNpcParameter == ResolvedDataInterfaceData.ResolvedVariable;
+				});
+				if (NpcBinding == nullptr)
+				{
+					NpcBinding = &OutResolvedNpcDataInterfaceBindings.AddDefaulted_GetRef();
+					NpcBinding->SourceNpcParameter = ResolvedDataInterfaceData.ResolvedVariable;
+				}
+				NpcBinding->TargetScriptParameterStoreDataInterfaceIndices.AddUnique(ResolvedDataInterfaceIndex);
 			}
+			ResolvedDataInterfaceIndex++;
+		}
+	}
+
+	void UpdatePostResolveData(
+		const UNiagaraSystem& System,
+		TMap<FNiagaraScriptDataKey, FNiagaraScriptRuntimeCompiledData>& OutScriptRuntimeCompiledDataMap)
+	{
+		TMap<FGuid, FString> EmitterHandleIdToNameMap;
+		for (const FNiagaraEmitterHandle& EmitterHandle : System.GetEmitterHandles())
+		{
+			EmitterHandleIdToNameMap.Add(EmitterHandle.GetId(), EmitterHandle.GetUniqueInstanceName());
 		}
+
+		const FNiagaraParameterStore& SystemInstanceParameters = System.GetSystemCompiledData().InstanceParamStore;
+
+		System.ForEachScriptWithOwningContext([&EmitterHandleIdToNameMap,&SystemInstanceParameters,&OutScriptRuntimeCompiledDataMap]
+			(const UNiagaraSystem* OwningSystem, const FNiagaraEmitterHandle* OwningEmitterHandle, UNiagaraScript* Script)
+		{
+			FGuid EmitterHandleId = OwningEmitterHandle != nullptr ? OwningEmitterHandle->GetId() : FGuid();
+			FNiagaraScriptDataKey ScriptDataKey = FNiagaraScriptDataKey(EmitterHandleId, Script->GetUsage(), Script->GetUsageId());
+			FNiagaraScriptRuntimeCompiledData& ScriptRuntimeCompiledData = OutScriptRuntimeCompiledDataMap[ScriptDataKey];
+
+			GenerateDIParameterVariablesForScript(ScriptRuntimeCompiledData.ResolvedDataInterfaces, OwningEmitterHandle, EmitterHandleIdToNameMap);
+			GenerateUserDIBindingsForScript(ScriptRuntimeCompiledData.ResolvedDataInterfaces, SystemInstanceParameters, ScriptRuntimeCompiledData.ResolvedUserDataInterfaceBindings);
+			GenerateNPCDIBindingsForScript(ScriptRuntimeCompiledData.ResolvedDataInterfaces, ScriptRuntimeCompiledData.ResolvedNpcDataInterfaceBindings);
+		});
 	}
 
-	void ResolveDIs(UNiagaraSystem* System, TMap<FNiagaraScriptDataKey, FNiagaraScriptRuntimeCompiledData>& OutScriptRuntimeCompiledDataMap, TArray<FText>& OutErrorMessages)
+	void ResolveDIs(UNiagaraSystem& System, TMap<FNiagaraScriptDataKey, FNiagaraScriptRuntimeCompiledData>& OutScriptRuntimeCompiledDataMap, TArray<FText>& OutErrorMessages)
 	{
-		TMap<FGuid, TMap<FNiagaraVariableBase, FDataInterfaceSourceEmitterNamePair>> EmitterIdToVariableReadsMap;
-		TMap<FGuid, TMap<FNiagaraVariableBase, FDataInterfaceSourceEmitterNamePair>> EmitterIdToVariableAssignmentsMap;
-		TMap<FGuid, TMap<FNiagaraVariableBase, FNiagaraVariableBase>> EmitterIdToVariableBindingsMap;
-		FNiagaraResolveDIHelpers::CollectDIReadsAssignmentsAndBindings(System, EmitterIdToVariableReadsMap, EmitterIdToVariableAssignmentsMap, EmitterIdToVariableBindingsMap, OutErrorMessages);
-		FNiagaraResolveDIHelpers::ResolveDIsInternal(System, EmitterIdToVariableReadsMap, EmitterIdToVariableAssignmentsMap, EmitterIdToVariableBindingsMap, OutScriptRuntimeCompiledDataMap, OutErrorMessages);
+		TMap<FGuid, FScriptDefaultData> ScriptDefaultDataMap;
+		TMap<FNiagaraCompileParameterReference, UNiagaraDataInterface*> SystemExternalParameterDefaultMap;
+		UObject& OwnerForCreateDIs = System;
+		CollectScriptDIResolveDataForSystem(System, OwnerForCreateDIs, ScriptDefaultDataMap, SystemExternalParameterDefaultMap, OutErrorMessages);
+		ResetResolvedDIsForSystem(System, OutScriptRuntimeCompiledDataMap);
+		GenerateResolvedDIsForSystem(System, OwnerForCreateDIs, ScriptDefaultDataMap, SystemExternalParameterDefaultMap, OutScriptRuntimeCompiledDataMap, OutErrorMessages);
+		UpdatePostResolveData(System, OutScriptRuntimeCompiledDataMap);
 	}
 }
 
diff --git a/Engine/Plugins/FX/Niagara/Source/Niagara/Private/NiagaraResolveDIHelpers.h b/Engine/Plugins/FX/Niagara/Source/Niagara/Private/NiagaraResolveDIHelpers.h
--- a/Engine/Plugins/FX/Niagara/Source/Niagara/Private/NiagaraResolveDIHelpers.h
+++ b/Engine/Plugins/FX/Niagara/Source/Niagara/Private/NiagaraResolveDIHelpers.h
@@ -5,18 +5,17 @@
 #include "HAL/Platform.h"
 #include "Misc/CoreMiscDefines.h"
 
-#include "NiagaraCommon.h"
-#include "NiagaraScriptRuntimeCompiledData.h"
 #include "NiagaraTypes.h"
 
-class UNiagaraDataInterface;
-class UNiagaraSystem;
+struct FNiagaraScriptRuntimeCompiledData;
+struct FNiagaraScriptDataKey;
+class  UNiagaraSystem;
 
 #if WITH_EDITORONLY_DATA
 
 namespace FNiagaraResolveDIHelpers
 {
-	void ResolveDIs(UNiagaraSystem* System, TMap<FNiagaraScriptDataKey, FNiagaraScriptRuntimeCompiledData>& OutScriptRuntimeCompiledDataMap, TArray<FText>& OutErrorMessages);
+	void ResolveDIs(UNiagaraSystem& System, TMap<FNiagaraScriptDataKey, FNiagaraScriptRuntimeCompiledData>& OutScriptRuntimeCompiledDataMap, TArray<FText>& OutErrorMessages);
 }
 
 #endif
\ No newline at end of file
diff --git a/Engine/Plugins/FX/Niagara/Source/Niagara/Private/NiagaraScript.cpp b/Engine/Plugins/FX/Niagara/Source/Niagara/Private/NiagaraScript.cpp
--- a/Engine/Plugins/FX/Niagara/Source/Niagara/Private/NiagaraScript.cpp
+++ b/Engine/Plugins/FX/Niagara/Source/Niagara/Private/NiagaraScript.cpp
@@ -1483,9 +1483,9 @@ void UNiagaraScript::ComputeVMCompilationId_EmitterShared(FNiagaraVMExecutableDa
 
 void UNiagaraScript::InvalidateAndEmptyDefaultDataInterfaces(int32 SlackCount, EDataInterfaceInvalidationMode InvalidationMode)
 {
-	DataInterfacesPendingInvalidation.Reserve(DataInterfacesPendingInvalidation.Num() + CachedDefaultDataInterfaces.Num());
+	DataInterfacesPendingInvalidation.Reserve(DataInterfacesPendingInvalidation.Num() + CachedGraphDefaultDataInterfaces.Num());
 
-	for (FNiagaraScriptDataInterfaceInfo& DefaultDataInterface : CachedDefaultDataInterfaces)
+	for (FNiagaraScriptDataInterfaceDefaultInfo& DefaultDataInterface : CachedGraphDefaultDataInterfaces)
 	{
 		if (DefaultDataInterface.DataInterface)
 		{
@@ -1496,7 +1496,7 @@ void UNiagaraScript::InvalidateAndEmptyDefaultDataInterfaces(int32 SlackCount, E
 		}
 	}
 
-	CachedDefaultDataInterfaces.Empty(SlackCount);
+	CachedGraphDefaultDataInterfaces.Empty(SlackCount);
 
 	if (InvalidationMode == EDataInterfaceInvalidationMode::ProcessInvaliation)
 	{
@@ -1945,7 +1945,7 @@ void UNiagaraScript::PostLoad()
 
 #if WITH_EDITORONLY_DATA
 	// Because we might be using these cached data interfaces, we need to make sure that they are properly postloaded.
-	for (FNiagaraScriptDataInterfaceInfo& Info : CachedDefaultDataInterfaces)
+	for (FNiagaraScriptDataInterfaceDefaultInfo& Info : CachedGraphDefaultDataInterfaces)
 	{
 		if (Info.DataInterface)
 		{
@@ -2650,24 +2650,6 @@ FVersionedNiagaraScriptData* FVersionedNiagaraScript::GetScriptData() const
 }
 #endif
 
-UNiagaraDataInterface* ResolveDataInterface(const TMap<FName, UNiagaraDataInterface*>& ObjectNameMap, FName VariableName)
-{
-	UNiagaraDataInterface* const* FoundDI = ObjectNameMap.Find(VariableName);
-	if (FoundDI && *(FoundDI))
-	{
-		return *FoundDI;
-	}
-	return nullptr;
-}
-
-void DumpNameMap(const TMap<FName, UNiagaraDataInterface*>& ObjectNameMap)
-{
-	for (const TPair<FName, UNiagaraDataInterface*>& Pair : ObjectNameMap)
-	{
-		UE_LOG(LogNiagara, Log, TEXT("%s -> %s"), *Pair.Key.ToString(), *GetPathNameSafe(Pair.Value));
-	}
-}
-
 FName ResolveEmitterAlias(const FName& InName, const FString& InAlias)
 {
 	// If the alias is empty than the name can't be resolved.
@@ -2683,7 +2665,7 @@ FName ResolveEmitterAlias(const FName& InName, const FString& InAlias)
 	return Var.GetName();
 }
 
-void UNiagaraScript::SetVMCompilationResults(const FNiagaraVMExecutableDataId& InCompileId, FNiagaraVMExecutableData& InScriptVM, const FString& EmitterUniqueName, const TMap<FName, UNiagaraDataInterface*>& ObjectNameMap, bool bApplyRapidIterationParameters)
+void UNiagaraScript::SetVMCompilationResults(const FNiagaraVMExecutableDataId& InCompileId, FNiagaraVMExecutableData& InScriptVM, const TMap<FNiagaraGraphInputDataInterfaceKey, UNiagaraDataInterface*>& GraphInputDataInterfaces, bool bApplyRapidIterationParameters)
 {
 	CachedScriptVMId = InCompileId;
 	CachedScriptVM = InScriptVM;
@@ -2706,37 +2688,24 @@ void UNiagaraScript::SetVMCompilationResults(const FNiagaraVMExecutableDataId& I
 	InvalidateAndEmptyDefaultDataInterfaces(CachedScriptVM.DataInterfaceInfo.Num(), EDataInterfaceInvalidationMode::ProcessInvaliation);
 	for (const FNiagaraScriptDataInterfaceCompileInfo& Info : CachedScriptVM.DataInterfaceInfo)
 	{
-		int32 Idx = CachedDefaultDataInterfaces.AddDefaulted();
-		CachedDefaultDataInterfaces[Idx].UserPtrIdx = Info.UserPtrIdx;
-		CachedDefaultDataInterfaces[Idx].Name = ResolveEmitterAlias(Info.Name, EmitterUniqueName);
-		CachedDefaultDataInterfaces[Idx].CompileName = Info.Name;
-		CachedDefaultDataInterfaces[Idx].Type = Info.Type;
-		CachedDefaultDataInterfaces[Idx].SourceEmitterName = Info.SourceEmitterName;
-		CachedDefaultDataInterfaces[Idx].RegisteredParameterMapRead = ResolveEmitterAlias(Info.RegisteredParameterMapRead, EmitterUniqueName);
-		CachedDefaultDataInterfaces[Idx].RegisteredParameterMapWrite = ResolveEmitterAlias(Info.RegisteredParameterMapWrite, EmitterUniqueName);
-
-		// We compiled it just a bit ago, so we should be able to resolve it from the table that we passed in.
-		UNiagaraDataInterface* FindDIById = ResolveDataInterface(ObjectNameMap, CachedDefaultDataInterfaces[Idx].Name);
-		if (FindDIById != nullptr )
+		if (Info.SourceGraphInputDataInterfaceKey.IsValid())
 		{
-			CachedDefaultDataInterfaces[Idx].DataInterface = CopyDataInterface(FindDIById, this);
-			check(CachedDefaultDataInterfaces[Idx].DataInterface != nullptr);
-		}
-
-		if (CachedDefaultDataInterfaces[Idx].DataInterface == nullptr)
-		{
-			// Use the CDO since we didn't have a default..
-			UObject* Obj = Info.Type.GetClass()->GetDefaultObject(true);
-			CachedDefaultDataInterfaces[Idx].DataInterface = Cast<UNiagaraDataInterface>(CopyDataInterface(CastChecked<UNiagaraDataInterface>(Obj), this));
-
-			if (Info.bIsPlaceholder == false)
+			int32 Idx = CachedGraphDefaultDataInterfaces.AddDefaulted();
+			CachedGraphDefaultDataInterfaces[Idx].Id = Info.Id;
+			UNiagaraDataInterface*const* GraphInputDataInterfacePtr = GraphInputDataInterfaces.Find(Info.SourceGraphInputDataInterfaceKey);
+			if (GraphInputDataInterfacePtr != nullptr)
 			{
-				UE_LOG(LogNiagara, Display, TEXT("We somehow ended up with a data interface that we couldn't match post compile. This shouldn't happen. Creating a dummy to prevent crashes. DataInterfaceInfoName:%s Object:%s"), *Info.Name.ToString(), *GetPathNameSafe(this));
-				UE_LOG(LogNiagara, Log, TEXT("Object to Name map contents:"));
-				DumpNameMap(ObjectNameMap);
+				CachedGraphDefaultDataInterfaces[Idx].DataInterface = CopyDataInterface(*GraphInputDataInterfacePtr, this);
+			}
+			else
+			{
+				// Use the CDO since we didn't have a default..
+				UObject* Obj = Info.Type.GetClass()->GetDefaultObject(true);
+				CachedGraphDefaultDataInterfaces[Idx].DataInterface = Cast<UNiagaraDataInterface>(CopyDataInterface(CastChecked<UNiagaraDataInterface>(Obj), this));
+
+				UE_LOG(LogNiagara, Display, TEXT("We somehow ended up with an input data interface that we couldn't match post compile. This shouldn't happen. Creating a dummy to prevent crashes. DataInterfaceInfoName:%s Object:%s"), *Info.Name.ToString(), *GetPathNameSafe(this));
 			}
 		}
-		check(CachedDefaultDataInterfaces[Idx].DataInterface != nullptr);
 	}
 
 	CachedDefaultUObjects = CachedScriptVM.UObjectInfos;
@@ -3040,7 +3009,7 @@ void UNiagaraScript::RequestCompile(const FGuid& ScriptVersion, bool bForceCompi
 			if (BinaryToExecData(this, OutData, ExeData))
 			{
 				COOK_STAT(Timer.AddHit(OutData.Num()));
-				SetVMCompilationResults(LastGeneratedVMId, ExeData, FString(), RequestDuplicateData->GetObjectNameMap(), false);
+				SetVMCompilationResults(LastGeneratedVMId, ExeData, RequestDuplicateData->GetGraphInputDataInterfaces(), false);
 				return;
 			}
 		}
@@ -3054,7 +3023,7 @@ void UNiagaraScript::RequestCompile(const FGuid& ScriptVersion, bool bForceCompi
 		TSharedPtr<FNiagaraVMExecutableData> ExeData = NiagaraModule.GetCompileJobResult(JobHandle, true, ScriptMetrics);
 		if (ExeData)
 		{
-			SetVMCompilationResults(LastGeneratedVMId, *ExeData, FString(), RequestDuplicateData->GetObjectNameMap(), false);
+			SetVMCompilationResults(LastGeneratedVMId, *ExeData, RequestDuplicateData->GetGraphInputDataInterfaces(), false);
 			// save result to the ddc
 			if (ExecToBinaryData(this, OutData, *ExeData))
 			{
@@ -3294,11 +3263,11 @@ bool UNiagaraScript::ShouldCompile(EShaderPlatform Platform) const
 
 	// Add all data interfaces
 	TSet<UClass*> DIUniqueClasses;
-	for (const FNiagaraScriptDataInterfaceInfo& DataInterfaceInfo : CachedDefaultDataInterfaces)
+	for (const FNiagaraScriptDataInterfaceCompileInfo& DataInterfaceInfo : CachedScriptVM.DataInterfaceInfo)
 	{
-		if (UNiagaraDataInterface* DataInterface = DataInterfaceInfo.DataInterface)
+		if (UClass* DataInterfaceClass = DataInterfaceInfo.Type.GetClass())
 		{
-			DIUniqueClasses.Add(DataInterface->GetClass());
+			DIUniqueClasses.Add(DataInterfaceClass);
 		}
 	}
 
@@ -3324,11 +3293,11 @@ void UNiagaraScript::ModifyCompilationEnvironment(EShaderPlatform Platform, stru
 #if WITH_EDITOR
 	// Add all data interfaces
 	TSet<UClass*> DIUniqueClasses;
-	for ( const FNiagaraScriptDataInterfaceInfo& DataInterfaceInfo : CachedDefaultDataInterfaces )
+	for (const FNiagaraScriptDataInterfaceCompileInfo& DataInterfaceInfo : CachedScriptVM.DataInterfaceInfo)
 	{
-		if ( UNiagaraDataInterface* DataInterface = DataInterfaceInfo.DataInterface )
+		if (UClass* DataInterfaceClass = DataInterfaceInfo.Type.GetClass())
 		{
-			DIUniqueClasses.Add(DataInterface->GetClass());
+			DIUniqueClasses.Add(DataInterfaceClass);
 		}
 	}
 
@@ -3684,12 +3653,6 @@ void UNiagaraScript::SyncAliases(const FNiagaraAliasContext& ResolveAliasesConte
 	// Sync up any simulation stage name references.
 	for (FSimulationStageMetaData& SimStageMetaData : GetVMExecutableData().SimulationStageMetaData)
 	{
-		if ( !SimStageMetaData.IterationDataInterface.IsNone() )
-		{
-			FNiagaraVariable Var(FNiagaraTypeDefinition(UNiagaraDataInterface::StaticClass()), SimStageMetaData.IterationDataInterface);
-			SimStageMetaData.IterationDataInterface = FNiagaraUtilities::ResolveAliases(Var, ResolveAliasesContext).GetName();
-		}
-
 		if (!SimStageMetaData.EnabledBinding.IsNone())
 		{
 			FNiagaraVariable Var(FNiagaraTypeDefinition::GetIntDef(), SimStageMetaData.EnabledBinding);
@@ -3809,13 +3772,13 @@ bool UNiagaraScript::SynchronizeExecutablesWithCompilation(const UNiagaraScript*
 		CachedScriptVMId = Script->CachedScriptVMId;
 		CachedParameterCollectionReferences = Script->CachedParameterCollectionReferences;
 		CachedDefaultUObjects = Script->CachedDefaultUObjects;
-		InvalidateAndEmptyDefaultDataInterfaces(Script->CachedDefaultDataInterfaces.Num(), EDataInterfaceInvalidationMode::ProcessInvaliation);
-		for (const FNiagaraScriptDataInterfaceInfo& Info : Script->CachedDefaultDataInterfaces)
+		InvalidateAndEmptyDefaultDataInterfaces(Script->CachedGraphDefaultDataInterfaces.Num(), EDataInterfaceInvalidationMode::ProcessInvaliation);
+		for (const FNiagaraScriptDataInterfaceDefaultInfo& Info : Script->CachedGraphDefaultDataInterfaces)
 		{
-			FNiagaraScriptDataInterfaceInfo AddInfo;
+			FNiagaraScriptDataInterfaceDefaultInfo AddInfo;
 			AddInfo = Info;
 			AddInfo.DataInterface = CopyDataInterface(Info.DataInterface, this);
-			CachedDefaultDataInterfaces.Add(AddInfo);
+			CachedGraphDefaultDataInterfaces.Add(AddInfo);
 		}
 
 		GenerateStatIDs();
diff --git a/Engine/Plugins/FX/Niagara/Source/Niagara/Private/NiagaraScriptRuntimeCompiledData.cpp b/Engine/Plugins/FX/Niagara/Source/Niagara/Private/NiagaraScriptRuntimeCompiledData.cpp
--- a/Engine/Plugins/FX/Niagara/Source/Niagara/Private/NiagaraScriptRuntimeCompiledData.cpp
+++ b/Engine/Plugins/FX/Niagara/Source/Niagara/Private/NiagaraScriptRuntimeCompiledData.cpp
@@ -22,7 +22,7 @@ bool FNiagaraScriptRuntimeCompiledData::ValidateWithScript(UNiagaraScript& InScr
 		for (const FNiagaraScriptResolvedDataInterfaceInfo& ResolvedDataInterfaceData : ResolvedDataInterfaces)
 		{
 			UE_LOG(LogNiagara, Log, TEXT("Name:%s, Type: %s, Path:%s"),
-				*ResolvedDataInterfaceData.CompileName.ToString(), *ResolvedDataInterfaceData.ParameterStoreVariable.GetType().GetName(),
+				*ResolvedDataInterfaceData.CompileName.ToString(), *ResolvedDataInterfaceData.Type.GetName(),
 				ResolvedDataInterfaceData.ResolvedDataInterface != nullptr
 				? *ResolvedDataInterfaceData.ResolvedDataInterface->GetPathName()
 				: TEXT("None"));
diff --git a/Engine/Plugins/FX/Niagara/Source/Niagara/Private/NiagaraSystem.cpp b/Engine/Plugins/FX/Niagara/Source/Niagara/Private/NiagaraSystem.cpp
--- a/Engine/Plugins/FX/Niagara/Source/Niagara/Private/NiagaraSystem.cpp
+++ b/Engine/Plugins/FX/Niagara/Source/Niagara/Private/NiagaraSystem.cpp
@@ -1956,7 +1956,7 @@ void UNiagaraSystem::GetAssetRegistryTags(FAssetRegistryTagsContext Context) con
 			{
 				for (FNiagaraScriptDataInterfaceCompileInfo& Info : Script->GetVMExecutableData().DataInterfaceInfo)
 				{
-					UNiagaraDataInterface* DefaultDataInterface = Info.GetDefaultDataInterface();
+					const UNiagaraDataInterface* DefaultDataInterface = Info.GetDefaultDataInterface();
 					DataInterfaces.AddUnique(DefaultDataInterface);
 				}
 			}
@@ -2596,7 +2596,7 @@ void UNiagaraSystem::CacheFromCompiledData()
 				FStaticBuffersCachingFlags& CachingFlagsForDataInterface = DataInterfaceToStaticCachingFlagsMap.FindOrAdd(DataInterface);
 				CachingFlagsForDataInterface.bIsUsedByCpu |= bUsedByCPU;
 				CachingFlagsForDataInterface.bIsUsedByGpu |= bUsedByGPU;
-				CachingFlagsForDataInterface.bIsExternalDataInterface |= DataInterfaceData.ResolvedVariable.IsInExternalNamespace();
+				CachingFlagsForDataInterface.bIsExternalDataInterface |= DataInterfaceData.ResolvedParameterReference.HasAnyFlags(ENiagaraCompileParameterReferenceFlags::User | ENiagaraCompileParameterReferenceFlags::NPC);
 			}
 		}
 	});
@@ -2648,7 +2648,7 @@ void UNiagaraSystem::InitScriptCompiledData()
 	ScriptRuntimeCompiledDataForEditor = NewObject<UNiagaraScriptRuntimeCompiledDataEditorReference>(this);
 
 	TArray<FText> ErrorMessages;
-	FNiagaraResolveDIHelpers::ResolveDIs(this, ScriptRuntimeCompiledDataForEditor->ScriptRuntimeCompiledDataMap, ErrorMessages);
+	FNiagaraResolveDIHelpers::ResolveDIs(*this, ScriptRuntimeCompiledDataForEditor->ScriptRuntimeCompiledDataMap, ErrorMessages);
 	if (ErrorMessages.Num() > 0)
 	{
 		INiagaraModule& NiagaraModule = FModuleManager::GetModuleChecked<INiagaraModule>("Niagara");
diff --git a/Engine/Plugins/FX/Niagara/Source/Niagara/Private/NiagaraSystemInstance.cpp b/Engine/Plugins/FX/Niagara/Source/Niagara/Private/NiagaraSystemInstance.cpp
--- a/Engine/Plugins/FX/Niagara/Source/Niagara/Private/NiagaraSystemInstance.cpp
+++ b/Engine/Plugins/FX/Niagara/Source/Niagara/Private/NiagaraSystemInstance.cpp
@@ -977,6 +977,16 @@ UNiagaraParameterCollectionInstance* FNiagaraSystemInstance::GetParameterCollect
 	return SystemSimulation->GetParameterCollectionInstance(Collection);
 }
 
+void FNiagaraSystemInstance::BindScriptParameterCollections(UNiagaraScript* Script, const FNiagaraScriptRuntimeData* ScriptRuntimeData, FNiagaraParameterStore& DestStore)
+{
+	SystemSimulation->BindScriptParameterCollections(Script, ScriptRuntimeData, DestStore);
+}
+
+void FNiagaraSystemInstance::UnbindScriptParameterCollections(UNiagaraScript* Script, FNiagaraParameterStore& DestStore)
+{
+	SystemSimulation->UnbindScriptParameterCollections(Script, DestStore);
+}
+
 void FNiagaraSystemInstance::AdvanceSimulation(int32 TickCountToSimulate, float TickDeltaSeconds)
 {
 	if (TickCountToSimulate > 0 && !IsPaused())
diff --git a/Engine/Plugins/FX/Niagara/Source/Niagara/Private/NiagaraSystemSimulation.cpp b/Engine/Plugins/FX/Niagara/Source/Niagara/Private/NiagaraSystemSimulation.cpp
--- a/Engine/Plugins/FX/Niagara/Source/Niagara/Private/NiagaraSystemSimulation.cpp
+++ b/Engine/Plugins/FX/Niagara/Source/Niagara/Private/NiagaraSystemSimulation.cpp
@@ -6,6 +6,7 @@
 #include "NiagaraModule.h"
 #include "NiagaraTypes.h"
 #include "NiagaraParameterCollection.h"
+#include "NiagaraScriptRuntimeData.h"
 #include "NiagaraSimCache.h"
 #include "NiagaraSimulationTaskPriority.h"
 #include "NiagaraSystemEmitterState.h"
@@ -652,35 +653,8 @@ bool FNiagaraSystemSimulation::Init(UNiagaraSystem* InSystem, UWorld* InWorld, b
 		{
 			//SCOPE_CYCLE_COUNTER(STAT_NiagaraSystemSim_Init_BindParams);
 
-			//Bind parameter collections.
-			auto BindParameterColleciton = [&](UNiagaraParameterCollection* Collection, FNiagaraParameterStore& DestStore)
-			{
-				if (Collection)
-				{
-					if (UNiagaraParameterCollectionInstance* CollectionInst = GetParameterCollectionInstance(Collection))
-					{
-						CollectionInst->GetParameterStore().Bind(&DestStore);
-					}
-					else
-					{
-						UE_LOG(LogNiagara, Error, TEXT("Attempting to bind system simulation to a null parameter collection instance | Collection: %s | System: %s |"), *Collection->GetPathName(), *System->GetPathName());
-					}
-				}
-				else
-				{
-					UE_LOG(LogNiagara, Error, TEXT("Attempting to bind system simulation to a null parameter collection | System: %s |"), *System->GetPathName());
-				}
-			};
-
-			//Bind parameter collections.
-			for (UNiagaraParameterCollection* Collection : SpawnScript->GetCachedParameterCollectionReferences())
-			{
-				BindParameterColleciton(Collection, SpawnExecContext->Parameters);
-			}
-			for (UNiagaraParameterCollection* Collection : UpdateScript->GetCachedParameterCollectionReferences())
-			{
-				BindParameterColleciton(Collection, UpdateExecContext->Parameters);
-			}
+			BindScriptParameterCollections(SpawnScript, SpawnExecContext->GetScriptRuntimeData(), SpawnExecContext->Parameters);
+			BindScriptParameterCollections(UpdateScript, UpdateExecContext->GetScriptRuntimeData(), UpdateExecContext->Parameters);
 
 			SpawnScript->RapidIterationParameters.Bind(&SpawnExecContext->Parameters);
 			UpdateScript->RapidIterationParameters.Bind(&UpdateExecContext->Parameters);
@@ -2607,3 +2581,85 @@ ENiagaraGPUTickHandlingMode FNiagaraSystemSimulation::GetGPUTickHandlingMode()co
 
 	return ENiagaraGPUTickHandlingMode::None;
 }
+
+void FNiagaraSystemSimulation::BindScriptParameterCollections(UNiagaraScript* Script, const FNiagaraScriptRuntimeData* ScriptRuntimeData, FNiagaraParameterStore& DestStore)
+{
+	for (UNiagaraParameterCollection* Collection : Script->GetCachedParameterCollectionReferences())
+	{
+		if (Collection == nullptr)
+		{
+			UE_LOG(LogNiagara, Warning, TEXT("Null cached parameter collection reference.\nScript: %s"), *Script->GetPathName());
+			continue;
+		}
+
+		UNiagaraParameterCollectionInstance* Inst = GetParameterCollectionInstance(Collection);
+		if (Inst == nullptr)
+		{
+			UE_LOG(LogNiagara, Warning, TEXT("Null parameter collection instance.\nScript: %s\nCollection: %s"),
+				*Script->GetPathName(), *Collection->GetPathName());
+			continue;
+		}
+
+		Inst->GetParameterStore().Bind(&DestStore);
+		ResolveNpcDataInterfaceBindingsForScript(Script, ScriptRuntimeData, Inst, DestStore);
+	}
+}
+
+void FNiagaraSystemSimulation::UnbindScriptParameterCollections(UNiagaraScript* Script, FNiagaraParameterStore& DestStore)
+{
+	for (UNiagaraParameterCollection* Collection : Script->GetCachedParameterCollectionReferences())
+	{
+		if (Collection == nullptr)
+		{
+			continue;
+		}
+
+		UNiagaraParameterCollectionInstance* Inst = GetParameterCollectionInstance(Collection);
+		if (Inst == nullptr)
+		{
+			continue;
+		}
+
+		Inst->GetParameterStore().Unbind(&DestStore);
+	}
+}
+
+void FNiagaraSystemSimulation::ResolveNpcDataInterfaceBindingsForScript(
+	const UNiagaraScript* Script,
+	const FNiagaraScriptRuntimeData* ScriptRuntimeData,
+	UNiagaraParameterCollectionInstance* CollectionInstance,
+	FNiagaraParameterStore& DestStore) const
+{
+	if (DestStore.GetDataInterfaces().Num() == 0 || ScriptRuntimeData == nullptr)
+	{
+		return;
+	}
+
+	for (const FNiagaraResolvedNpcDataInterfaceBinding& ResolvedNpcDataInterfaceBinding : ScriptRuntimeData->GetResolvedNpcDataInterfaceBindings())
+	{
+		UNiagaraDataInterface* CollectionDataInterface = CollectionInstance->GetParameterStore().GetDataInterface(ResolvedNpcDataInterfaceBinding.SourceNpcParameter);
+		if (CollectionDataInterface == nullptr)
+		{
+			UE_LOG(LogNiagara, Warning, TEXT("Missing parameter collection data interface.\nScript:%s\nCollection:%s\nParameter:%s"),
+				*Script->GetPathName(),
+				*CollectionInstance->GetPathName(),
+				*ResolvedNpcDataInterfaceBinding.SourceNpcParameter.GetName().ToString());
+			continue;
+		}
+
+		for (int32 TargetDataInterfaceIndex : ResolvedNpcDataInterfaceBinding.TargetScriptParameterStoreDataInterfaceIndices)
+		{
+			if (TargetDataInterfaceIndex >= 0 && TargetDataInterfaceIndex < DestStore.GetDataInterfaces().Num())
+			{
+				DestStore.SetDataInterface(CollectionDataInterface, TargetDataInterfaceIndex);
+			}
+			else
+			{
+				UE_LOG(LogNiagara, Warning, TEXT("Failed to bind parameter collection data interface.\nScript:%s\nCollection:%s\nParameter:%s"),
+					*Script->GetPathName(),
+					*CollectionInstance->GetPathName(),
+					*ResolvedNpcDataInterfaceBinding.SourceNpcParameter.GetName().ToString());
+			}
+		}
+	}
+}
diff --git a/Engine/Plugins/FX/Niagara/Source/Niagara/Public/NiagaraCommon.h b/Engine/Plugins/FX/Niagara/Source/Niagara/Public/NiagaraCommon.h
--- a/Engine/Plugins/FX/Niagara/Source/Niagara/Public/NiagaraCommon.h
+++ b/Engine/Plugins/FX/Niagara/Source/Niagara/Public/NiagaraCommon.h
@@ -17,6 +17,9 @@
 #include "NiagaraTypeRegistry.h"
 #include "NiagaraCore.h"
 #include "Templates/RefCounting.h"
+#include "NiagaraScriptCompileParameterReference.h"
+#include "NiagaraScriptDataInterfaceCompileId.h"
+#include "NiagaraScriptUsage.h"
 
 #include "NiagaraCommon.generated.h"
 
@@ -793,44 +796,20 @@ struct FNiagaraExternalUObjectInfo
 };
 
 USTRUCT()
-struct FNiagaraScriptDataInterfaceInfo
+struct FNiagaraScriptDataInterfaceDefaultInfo
 {
 	GENERATED_USTRUCT_BODY()
 public:
-	FNiagaraScriptDataInterfaceInfo()
+	FNiagaraScriptDataInterfaceDefaultInfo()
 		: DataInterface(nullptr)
-		, Name(NAME_None)
-		, UserPtrIdx(INDEX_NONE)
 	{
-
 	}
 
 	UPROPERTY()
-	TObjectPtr<class UNiagaraDataInterface> DataInterface;
-	
-	UPROPERTY()
-	FName Name;
-
-	UPROPERTY()
-	FName CompileName;
-
-	/** Index of the user pointer for this data interface. */
-	UPROPERTY()
-	int32 UserPtrIdx;
-
-	UPROPERTY()
-	FNiagaraTypeDefinition Type;
-
-	UPROPERTY()
-	FName RegisteredParameterMapRead;
-
-	UPROPERTY()
-	FName RegisteredParameterMapWrite;
+	FNiagaraScriptDataInterfaceCompileId Id;
 
 	UPROPERTY()
-	FString SourceEmitterName;
-
-	NIAGARA_API bool IsUserDataInterface() const;
+	TObjectPtr<class UNiagaraDataInterface> DataInterface;
 };
 
 USTRUCT()
@@ -857,35 +836,58 @@ struct FNiagaraResolvedUserDataInterfaceBinding
 	}
 };
 
+USTRUCT()
+struct FNiagaraResolvedNpcDataInterfaceBinding
+{
+	GENERATED_BODY()
+
+	UPROPERTY()
+	FNiagaraVariableBase SourceNpcParameter;
+
+	UPROPERTY()
+	TArray<int32> TargetScriptParameterStoreDataInterfaceIndices;
+
+	FNiagaraResolvedNpcDataInterfaceBinding()
+	{
+	}
+
+	FNiagaraResolvedNpcDataInterfaceBinding(const FNiagaraVariableBase& InSourceNpcParameter)
+		: SourceNpcParameter(InSourceNpcParameter)
+	{
+	}
+};
+
 USTRUCT()
 struct FNiagaraScriptResolvedDataInterfaceInfo
 {
 	GENERATED_USTRUCT_BODY();
 
 	FNiagaraScriptResolvedDataInterfaceInfo()
-		: bIsInternal(false)
-		, ResolvedDataInterface(nullptr)
+		: ResolvedDataInterface(nullptr)
 		, UserPtrIdx(-1)
 	{
 	}
 
 	UPROPERTY()
-	FName Name;
+	FGuid OwningEmitterHandleId;
 
 	UPROPERTY()
 	FName CompileName;
 
 	UPROPERTY()
-	FString ResolvedSourceEmitterName;
+	FNiagaraTypeDefinition Type;
 
 	UPROPERTY()
-	FNiagaraVariableBase ResolvedVariable;
+	FNiagaraCompileParameterReference ResolvedParameterReference;
 
 	UPROPERTY()
-	FNiagaraVariableBase ParameterStoreVariable;
+	FNiagaraCompileParameterReference TargetParameterReference;
 
 	UPROPERTY()
-	bool bIsInternal;
+	FNiagaraVariableBase ResolvedVariable;
+
+	UPROPERTY()
+	FNiagaraVariableBase ParameterStoreVariable;
 
 	UPROPERTY()
 	TObjectPtr<UNiagaraDataInterface> ResolvedDataInterface;
@@ -896,21 +898,55 @@ struct FNiagaraScriptResolvedDataInterfaceInfo
 	NIAGARA_API bool NeedsPerInstanceBinding() const;
 };
 
+USTRUCT()
+struct FNiagaraGraphInputDataInterfaceKey
+{
+	GENERATED_BODY()
+
+	FNiagaraGraphInputDataInterfaceKey()
+	{
+	}
+
+	FNiagaraGraphInputDataInterfaceKey(const FName& InGraphPath, const FNiagaraVariableBase& InInputVariable)
+		: GraphPath(InGraphPath)
+		, InputVariable(InInputVariable)
+	{
+	}
+
+	bool operator==(const FNiagaraGraphInputDataInterfaceKey& Other) const
+	{
+		return GraphPath == Other.GraphPath && InputVariable == Other.InputVariable;
+	}
+
+	friend uint32 GetTypeHash(const FNiagaraGraphInputDataInterfaceKey& Ref)
+	{
+		return HashCombine(GetTypeHash(Ref.GraphPath), GetTypeHash(Ref.InputVariable));
+	}
+
+	bool IsValid() const { return GraphPath != NAME_None && InputVariable.IsValid(); }
+
+	UPROPERTY()
+	FName GraphPath;
+
+	UPROPERTY()
+	FNiagaraVariableBase InputVariable;
+};
+
 USTRUCT()
 struct FNiagaraScriptDataInterfaceCompileInfo
 {
 	GENERATED_USTRUCT_BODY()
 public:
 	FNiagaraScriptDataInterfaceCompileInfo()
-		: Name(NAME_None)
-		, UserPtrIdx(INDEX_NONE)
-		, bIsPlaceholder(false)
+		: UserPtrIdx(INDEX_NONE)
+		, SourceScriptDIIndex(INDEX_NONE)
 	{
-
 	}
 
+	// A unique id for this data interface based on source information when it was first encountered.  This will be unique across the entire owning
+	// object, either the system for system assets, or the script for script assets.
 	UPROPERTY()
-	FName Name;
+	FNiagaraScriptDataInterfaceCompileId Id;
 
 	/** Index of the user pointer for this data interface. */
 	UPROPERTY()
@@ -919,32 +955,29 @@ public:
 	UPROPERTY()
 	FNiagaraTypeDefinition Type;
 
-	// Removed from cooked builds, if we need to add this back the TMap<FName, FName> FunctionSpecifiers should be replaced with an array
-#if WITH_EDITORONLY_DATA
 	UPROPERTY()
 	TArray<FNiagaraFunctionSignature> RegisteredFunctions;
-#endif
 
+	// A generated name for this data interface for compilation purposes.  This name is guaranteed to be unique within the compiled script.
 	UPROPERTY()
-	FName RegisteredParameterMapRead;
+	FName Name;
 
 	UPROPERTY()
-	FName RegisteredParameterMapWrite;
+	FNiagaraGraphInputDataInterfaceKey SourceGraphInputDataInterfaceKey;
 
 	UPROPERTY()
-	bool bIsPlaceholder;
+	int32 SourceScriptDIIndex;
 
 	UPROPERTY()
-	FString SourceEmitterName;
+	FNiagaraCompileParameterReference SourceExternalParameter;
 
-	/** Would this data interface work on the target execution type? Only call this on the game thread.*/
-	NIAGARA_API bool CanExecuteOnTarget(ENiagaraSimTarget SimTarget) const;
+	// The parameter that this data interface was written to if any.
+	UPROPERTY()
+	FNiagaraCompileParameterReference TargetParameter;
 
 	/** Note that this is the CDO for this type of data interface, as we often cannot guarantee that the same instance of the data interface we compiled with is the one the user ultimately executes.  Only call this on the game thread.*/
 	NIAGARA_API UNiagaraDataInterface* GetDefaultDataInterface() const;
 
-	NIAGARA_API bool NeedsPerInstanceBinding() const;
-
 	NIAGARA_API bool MatchesClass(const UClass* InClass) const;
 };
 
@@ -1108,38 +1141,6 @@ private:
 	FCustomWorkDelegate PostWork;
 };
 
-/** Defines different usages for a niagara script. */
-UENUM()
-enum class ENiagaraScriptUsage : uint8
-{
-	/** The script defines a function for use in modules. */
-	Function,
-	/** The script defines a module for use in particle, emitter, or system scripts. */
-	Module,
-	/** The script defines a dynamic input for use in particle, emitter, or system scripts. */
-	DynamicInput,
-	/** The script is called when spawning particles. */
-	ParticleSpawnScript,
-	/** Particle spawn script that handles intra-frame spawning and also pulls in the update script. */
-	ParticleSpawnScriptInterpolated UMETA(Hidden),
-	/** The script is called to update particles every frame. */
-	ParticleUpdateScript,
-	/** The script is called to update particles in response to an event. */
-	ParticleEventScript ,
-	/** The script is called as a particle simulation stage. */
-	ParticleSimulationStageScript,
-	/** The script is called to update particles on the GPU. */
-	ParticleGPUComputeScript UMETA(Hidden),
-	/** The script is called once when the emitter spawns. */
-	EmitterSpawnScript,
-	/** The script is called every frame to tick the emitter. */
-	EmitterUpdateScript ,
-	/** The script is called once when the system spawns. */
-	SystemSpawnScript ,
-	/** The script is called every frame to tick the system. */
-	SystemUpdateScript,
-};
-
 /** Defines common bit masks for script usage */
 namespace ENiagaraScriptUsageMask
 {
diff --git a/Engine/Plugins/FX/Niagara/Source/Niagara/Public/NiagaraCompilationTypes.h b/Engine/Plugins/FX/Niagara/Source/Niagara/Public/NiagaraCompilationTypes.h
--- a/Engine/Plugins/FX/Niagara/Source/Niagara/Public/NiagaraCompilationTypes.h
+++ b/Engine/Plugins/FX/Niagara/Source/Niagara/Public/NiagaraCompilationTypes.h
@@ -58,9 +58,6 @@ struct FNiagaraScriptAsyncCompileData
 
 	UPROPERTY()
 	TArray<FNiagaraVariable> RapidIterationParameters;
-
-	UPROPERTY()
-	TMap<FName, TObjectPtr<UNiagaraDataInterface>> NamedDataInterfaces;
 };
 
 USTRUCT()
diff --git a/Engine/Plugins/FX/Niagara/Source/Niagara/Public/NiagaraDataInterfaceEmitterBinding.h b/Engine/Plugins/FX/Niagara/Source/Niagara/Public/NiagaraDataInterfaceEmitterBinding.h
--- a/Engine/Plugins/FX/Niagara/Source/Niagara/Public/NiagaraDataInterfaceEmitterBinding.h
+++ b/Engine/Plugins/FX/Niagara/Source/Niagara/Public/NiagaraDataInterfaceEmitterBinding.h
@@ -49,4 +49,7 @@ struct FNiagaraDataInterfaceEmitterBinding
 	{
 		return BindingMode == Other.BindingMode && EmitterName == Other.EmitterName;
 	}
+
+private:
+	const FNiagaraEmitterHandle* ResolveHandleOtherBinding(const UNiagaraSystem* OwnerSystem) const;
 };
diff --git a/Engine/Plugins/FX/Niagara/Source/Niagara/Public/NiagaraScriptCompileParameterReference.h b/Engine/Plugins/FX/Niagara/Source/Niagara/Public/NiagaraScriptCompileParameterReference.h
new file mode 100644
--- /dev/null
+++ b/Engine/Plugins/FX/Niagara/Source/Niagara/Public/NiagaraScriptCompileParameterReference.h
@@ -0,0 +1,89 @@
+// Copyright Epic Games, Inc. All Rights Reserved.
+
+#pragma once
+
+#include "NiagaraScriptCompileParameterReference.generated.h"
+
+UENUM()
+enum class ENiagaraCompileParameterReferenceFlags : uint8
+{
+	None = 0,
+	User = 1,
+	System = 2,
+	Emitter = 4,
+	Particle = 8,
+	Transient = 16,
+	Local = 32,
+	NPC = 64
+};
+
+FORCEINLINE ENiagaraCompileParameterReferenceFlags operator|(
+	ENiagaraCompileParameterReferenceFlags FlagsA,
+	ENiagaraCompileParameterReferenceFlags FlagsB)
+{
+	return (ENiagaraCompileParameterReferenceFlags)((uint8)FlagsA | (uint8)FlagsB);
+}
+
+FORCEINLINE ENiagaraCompileParameterReferenceFlags operator&(
+	ENiagaraCompileParameterReferenceFlags FlagsA,
+	ENiagaraCompileParameterReferenceFlags FlagsB)
+{
+	return (ENiagaraCompileParameterReferenceFlags)((uint8)FlagsA & (uint8)FlagsB);
+}
+
+USTRUCT()
+struct FNiagaraCompileParameterReference
+{
+	GENERATED_BODY()
+
+	FNiagaraCompileParameterReference()
+		: ParameterFullName(NAME_None)
+		, EmitterHandleId(FGuid())
+		, Flags(ENiagaraCompileParameterReferenceFlags::None)
+	{
+	}
+
+	FNiagaraCompileParameterReference(const FName& InParameterFullName, const FGuid& InEmitterHandleId, ENiagaraCompileParameterReferenceFlags InFlags)
+		: ParameterFullName(InParameterFullName)
+		, EmitterHandleId(InEmitterHandleId)
+		, Flags(InFlags)
+	{
+	}
+
+	bool operator==(const FNiagaraCompileParameterReference& Other) const
+	{
+		return ParameterFullName == Other.ParameterFullName && EmitterHandleId == Other.EmitterHandleId && Flags == Other.Flags;
+	}
+
+	friend uint32 GetTypeHash(const FNiagaraCompileParameterReference& Ref)
+	{
+		int32 Hash = GetTypeHash(Ref.ParameterFullName);
+		Hash = HashCombine(Hash, GetTypeHash(Ref.EmitterHandleId));
+		Hash = HashCombine(Hash, GetTypeHash(Ref.Flags));
+		return Hash;
+	}
+
+	bool IsValid() const { return ParameterFullName != NAME_None; }
+
+	bool HasFlags(ENiagaraCompileParameterReferenceFlags FlagsToCheck) const { return (Flags & FlagsToCheck) == FlagsToCheck; }
+
+	bool HasAnyFlags(ENiagaraCompileParameterReferenceFlags FlagsToCheck) const { return (Flags & FlagsToCheck) != (ENiagaraCompileParameterReferenceFlags)0; }
+
+	void SetFlags(ENiagaraCompileParameterReferenceFlags FlagsToSet)
+	{
+		Flags = (Flags | FlagsToSet);
+	}
+
+	bool IsScriptInternalReference() const { return HasAnyFlags(ENiagaraCompileParameterReferenceFlags::Transient | ENiagaraCompileParameterReferenceFlags::Local); }
+
+	bool IsNonScriptExternalReference() const { return HasAnyFlags(ENiagaraCompileParameterReferenceFlags::User | ENiagaraCompileParameterReferenceFlags::NPC); }
+
+	UPROPERTY()
+	FName ParameterFullName;
+
+	UPROPERTY()
+	FGuid EmitterHandleId;
+
+	UPROPERTY()
+	ENiagaraCompileParameterReferenceFlags Flags;
+};
\ No newline at end of file
diff --git a/Engine/Plugins/FX/Niagara/Source/Niagara/Public/NiagaraScriptDataInterfaceCompileId.h b/Engine/Plugins/FX/Niagara/Source/Niagara/Public/NiagaraScriptDataInterfaceCompileId.h
new file mode 100644
--- /dev/null
+++ b/Engine/Plugins/FX/Niagara/Source/Niagara/Public/NiagaraScriptDataInterfaceCompileId.h
@@ -0,0 +1,99 @@
+// Copyright Epic Games, Inc. All Rights Reserved.
+
+#pragma once
+
+#include "NiagaraScriptUsage.h"
+
+#include "NiagaraScriptDataInterfaceCompileId.generated.h"
+
+/** A unique id for a data interface instance within a script compile.
+	NOTE:  In a single system it's possible to have ids from particle data interfaces in different emitters	to have matching ids
+	since the particle data interface ids do not include the emitter handle id. */
+USTRUCT()
+struct FNiagaraScriptDataInterfaceCompileId
+{
+	GENERATED_BODY();
+
+	FNiagaraScriptDataInterfaceCompileId()
+		: EmitterHandleId(FGuid())
+		, ScriptUsage(ENiagaraScriptUsage::Function)
+		, ScriptUsageId(FGuid())
+		, OwnerName(NAME_None)
+		, BaseName(NAME_None)
+		, Index(INDEX_NONE)
+	{
+	}
+
+	FNiagaraScriptDataInterfaceCompileId(const FGuid& InEmitterHandleId, ENiagaraScriptUsage InScriptUsage, const FGuid& InScriptUsageId, FName InOwnerName, FName InBaseName, int32 InIndex)
+		: EmitterHandleId(InEmitterHandleId)
+		, ScriptUsage(InScriptUsage)
+		, ScriptUsageId(InScriptUsageId)
+		, OwnerName(InOwnerName)
+		, BaseName(InBaseName)
+		, Index(InIndex)
+	{
+	}
+
+	FNiagaraScriptDataInterfaceCompileId(const FGuid& InEmitterHandleId, const FNiagaraScriptDataInterfaceCompileId& Other)
+		: EmitterHandleId(InEmitterHandleId)
+		, ScriptUsage(Other.ScriptUsage)
+		, ScriptUsageId(Other.ScriptUsageId)
+		, OwnerName(Other.OwnerName)
+		, BaseName(Other.BaseName)
+		, Index(Other.Index)
+	{
+	}
+
+	bool operator==(const FNiagaraScriptDataInterfaceCompileId& Other) const
+	{
+		return
+			EmitterHandleId == Other.EmitterHandleId &&
+			ScriptUsage == Other.ScriptUsage &&
+			ScriptUsageId == Other.ScriptUsageId &&
+			OwnerName == Other.OwnerName &&
+			BaseName == Other.BaseName &&
+			Index == Other.Index;
+	}
+
+	friend uint32 GetTypeHash(const FNiagaraScriptDataInterfaceCompileId& Ref)
+	{
+		int32 Hash = GetTypeHash(Ref.EmitterHandleId);
+		Hash = HashCombine(Hash, GetTypeHash(Ref.ScriptUsage));
+		Hash = HashCombine(Hash, GetTypeHash(Ref.ScriptUsageId));
+		Hash = HashCombine(Hash, GetTypeHash(Ref.OwnerName));
+		Hash = HashCombine(Hash, GetTypeHash(Ref.BaseName));
+		Hash = HashCombine(Hash, GetTypeHash(Ref.Index));
+		return Hash;
+	}
+
+	const FGuid& GetEmitterHandleId() const { return EmitterHandleId; }
+
+	ENiagaraScriptUsage GetScriptUsage() const { return ScriptUsage; }
+
+	const FGuid& GetScriptUsageId() const { return ScriptUsageId; }
+
+	FName GetOwnerName() const { return OwnerName; }
+
+	FName GetBaseName() const { return BaseName; }
+
+	int32 GetIndex() const { return Index; }
+
+private:
+	UPROPERTY()
+	FGuid EmitterHandleId;
+
+	UPROPERTY()
+	ENiagaraScriptUsage ScriptUsage;
+
+	UPROPERTY()
+	FGuid ScriptUsageId;
+
+	UPROPERTY()
+	FName OwnerName;
+
+	UPROPERTY()
+	FName BaseName;
+
+	UPROPERTY()
+	int32 Index;
+};
\ No newline at end of file
diff --git a/Engine/Plugins/FX/Niagara/Source/Niagara/Public/NiagaraScriptRuntimeCompiledData.h b/Engine/Plugins/FX/Niagara/Source/Niagara/Public/NiagaraScriptRuntimeCompiledData.h
--- a/Engine/Plugins/FX/Niagara/Source/Niagara/Public/NiagaraScriptRuntimeCompiledData.h
+++ b/Engine/Plugins/FX/Niagara/Source/Niagara/Public/NiagaraScriptRuntimeCompiledData.h
@@ -22,6 +22,9 @@ struct FNiagaraScriptRuntimeCompiledData
 	UPROPERTY()
 	TArray<FNiagaraResolvedUserDataInterfaceBinding> ResolvedUserDataInterfaceBindings;
 
+	UPROPERTY()
+	TArray<FNiagaraResolvedNpcDataInterfaceBinding> ResolvedNpcDataInterfaceBindings;
+
 	UPROPERTY()
 	TArray<FNiagaraResolvedUObjectInfo> ResolvedUObjects;
 
diff --git a/Engine/Plugins/FX/Niagara/Source/Niagara/Public/NiagaraScriptRuntimeData.h b/Engine/Plugins/FX/Niagara/Source/Niagara/Public/NiagaraScriptRuntimeData.h
--- a/Engine/Plugins/FX/Niagara/Source/Niagara/Public/NiagaraScriptRuntimeData.h
+++ b/Engine/Plugins/FX/Niagara/Source/Niagara/Public/NiagaraScriptRuntimeData.h
@@ -22,6 +22,8 @@ public:
 
 	TConstArrayView<FNiagaraResolvedUserDataInterfaceBinding> GetResolvedUserDataInterfaceBindings() const { return MakeConstArrayView(CompiledDataHandle->Instance->ResolvedUserDataInterfaceBindings); }
 
+	TConstArrayView<FNiagaraResolvedNpcDataInterfaceBinding> GetResolvedNpcDataInterfaceBindings() const { return MakeConstArrayView(CompiledDataHandle->Instance->ResolvedNpcDataInterfaceBindings); }
+
 	TConstArrayView<FNiagaraResolvedUObjectInfo> GetResolvedUObjects() const { return MakeConstArrayView(CompiledDataHandle->Instance->ResolvedUObjects); }
 
 	const FNiagaraScriptExecutionParameterStore* GetExecutionReadyParameterStore() const { return TransientData->GetExecutionReadyParameterStore(); } 
@@ -31,14 +33,4 @@ public:
 private:
 	TSharedPtr<const FNiagaraScriptRuntimeCompiledDataHandle> CompiledDataHandle;
 	TSharedPtr<const FNiagaraScriptRuntimeTransientData> TransientData;
-};
-
-class FNiagaraScriptComputeRuntimeData : public FNiagaraScriptRuntimeData
-{
-
-};
-
-class FNiagaraScriptEventRuntimeData : public FNiagaraScriptRuntimeData
-{
-
 };
\ No newline at end of file
diff --git a/Engine/Plugins/FX/Niagara/Source/Niagara/Public/NiagaraScriptUsage.h b/Engine/Plugins/FX/Niagara/Source/Niagara/Public/NiagaraScriptUsage.h
new file mode 100644
--- /dev/null
+++ b/Engine/Plugins/FX/Niagara/Source/Niagara/Public/NiagaraScriptUsage.h
@@ -0,0 +1,37 @@
+// Copyright Epic Games, Inc. All Rights Reserved.
+
+#pragma once
+
+#include "NiagaraScriptUsage.generated.h"
+
+/** Defines different usages for a niagara script. */
+UENUM()
+enum class ENiagaraScriptUsage : uint8
+{
+	/** The script defines a function for use in modules. */
+	Function,
+	/** The script defines a module for use in particle, emitter, or system scripts. */
+	Module,
+	/** The script defines a dynamic input for use in particle, emitter, or system scripts. */
+	DynamicInput,
+	/** The script is called when spawning particles. */
+	ParticleSpawnScript,
+	/** Particle spawn script that handles intra-frame spawning and also pulls in the update script. */
+	ParticleSpawnScriptInterpolated UMETA(Hidden),
+	/** The script is called to update particles every frame. */
+	ParticleUpdateScript,
+	/** The script is called to update particles in response to an event. */
+	ParticleEventScript ,
+	/** The script is called as a particle simulation stage. */
+	ParticleSimulationStageScript,
+	/** The script is called to update particles on the GPU. */
+	ParticleGPUComputeScript UMETA(Hidden),
+	/** The script is called once when the emitter spawns. */
+	EmitterSpawnScript,
+	/** The script is called every frame to tick the emitter. */
+	EmitterUpdateScript ,
+	/** The script is called once when the system spawns. */
+	SystemSpawnScript ,
+	/** The script is called every frame to tick the system. */
+	SystemUpdateScript,
+};
\ No newline at end of file
diff --git a/Engine/Plugins/FX/Niagara/Source/Niagara/Public/NiagaraSystemInstance.h b/Engine/Plugins/FX/Niagara/Source/Niagara/Public/NiagaraSystemInstance.h
--- a/Engine/Plugins/FX/Niagara/Source/Niagara/Public/NiagaraSystemInstance.h
+++ b/Engine/Plugins/FX/Niagara/Source/Niagara/Public/NiagaraSystemInstance.h
@@ -386,7 +386,11 @@ public:
 
 	NIAGARA_API bool IsReadyToRun() const;
 
-	NIAGARA_API UNiagaraParameterCollectionInstance* GetParameterCollectionInstance(UNiagaraParameterCollection* Collection);
+	UNiagaraParameterCollectionInstance* GetParameterCollectionInstance(UNiagaraParameterCollection* Collection);
+
+	void BindScriptParameterCollections(UNiagaraScript* Script, const FNiagaraScriptRuntimeData* ScriptRuntimeData, FNiagaraParameterStore& DestStore);
+
+	void UnbindScriptParameterCollections(UNiagaraScript* Script, FNiagaraParameterStore& DestStore);
 
 	/** 
 	Manually advances this system's simulation by the specified number of ticks and tick delta. 
diff --git a/Engine/Plugins/FX/Niagara/Source/Niagara/Public/NiagaraSystemSimulation.h b/Engine/Plugins/FX/Niagara/Source/Niagara/Public/NiagaraSystemSimulation.h
--- a/Engine/Plugins/FX/Niagara/Source/Niagara/Public/NiagaraSystemSimulation.h
+++ b/Engine/Plugins/FX/Niagara/Source/Niagara/Public/NiagaraSystemSimulation.h
@@ -325,6 +325,10 @@ public:
 
 	const FNiagaraTickInfo& GetTickInfo() { return TickInfo; }
 
+	void BindScriptParameterCollections(UNiagaraScript* Script, const FNiagaraScriptRuntimeData* ScriptRuntimeData, FNiagaraParameterStore& DestStore);
+
+	void UnbindScriptParameterCollections(UNiagaraScript* Script, FNiagaraParameterStore& DestStore);
+
 protected:
 	void DumpStalledInfo();
 
@@ -352,6 +356,12 @@ protected:
 
 	TArray<FNiagaraSystemInstance*>& GetSystemInstances(ENiagaraSystemInstanceState State) { check(State != ENiagaraSystemInstanceState::None); return SystemInstancesPerState[int32(State)]; }
 
+	void ResolveNpcDataInterfaceBindingsForScript(
+		const UNiagaraScript* Script,
+		const FNiagaraScriptRuntimeData* ScriptRuntimeData,
+		UNiagaraParameterCollectionInstance* CollectionInstance,
+		FNiagaraParameterStore& DestStore) const;
+
 	/** System of instances being simulated.  We use a weak object ptr here because once the last referencing object goes away this system may be come invalid at runtime. */
 	UNiagaraSystem* System;
 
diff --git a/Engine/Plugins/FX/Niagara/Source/NiagaraEditor/Private/NiagaraCompilationDigestBridgeImpl.cpp b/Engine/Plugins/FX/Niagara/Source/NiagaraEditor/Private/NiagaraCompilationDigestBridgeImpl.cpp
--- a/Engine/Plugins/FX/Niagara/Source/NiagaraEditor/Private/NiagaraCompilationDigestBridgeImpl.cpp
+++ b/Engine/Plugins/FX/Niagara/Source/NiagaraEditor/Private/NiagaraCompilationDigestBridgeImpl.cpp
@@ -156,6 +156,11 @@ ENiagaraScriptUsage FNiagaraCompilationDigestBridge::GetEmitterUsage(const FEmit
 	return EmitterNode->Usage;
 }
 
+FGuid FNiagaraCompilationDigestBridge::GetEmitterHandleId(const FEmitterNode* EmitterNode)
+{
+	return EmitterNode->EmitterHandleID;
+}
+
 FString FNiagaraCompilationDigestBridge::GetEmitterHandleIdString(const FEmitterNode* EmitterNode)
 {
 	return EmitterNode->EmitterHandleIdString;
@@ -293,6 +298,11 @@ FNiagaraVariable FNiagaraCompilationDigestBridge::GetInputVariable(const FInputN
 	return InputNode->InputVariable;
 }
 
+FNiagaraGraphInputDataInterfaceKey FNiagaraCompilationDigestBridge::GetGraphInputDataInterfaceKey(const FNiagaraCompilationNodeInput* InputNode, const FCompilationCopy* CompilationCopy)
+{
+	return FNiagaraGraphInputDataInterfaceKey(InputNode->GraphPathFName, InputNode->InputVariable);
+}
+
 const TArray<FNiagaraVariable>& FNiagaraCompilationDigestBridge::GetOutputVariables(const FOutputNode* OutputNode)
 {
 	return OutputNode->Outputs;
diff --git a/Engine/Plugins/FX/Niagara/Source/NiagaraEditor/Private/NiagaraCompilationGraphBridgeImpl.cpp b/Engine/Plugins/FX/Niagara/Source/NiagaraEditor/Private/NiagaraCompilationGraphBridgeImpl.cpp
--- a/Engine/Plugins/FX/Niagara/Source/NiagaraEditor/Private/NiagaraCompilationGraphBridgeImpl.cpp
+++ b/Engine/Plugins/FX/Niagara/Source/NiagaraEditor/Private/NiagaraCompilationGraphBridgeImpl.cpp
@@ -276,6 +276,11 @@ FString FNiagaraCompilationGraphBridge::GetEmitterPathName(const FEmitterNode* E
 	return EmitterNode->GetFullName();
 }
 
+FGuid FNiagaraCompilationGraphBridge::GetEmitterHandleId(const FEmitterNode* EmitterNode)
+{
+	return EmitterNode->GetEmitterHandleId();
+}
+
 FString FNiagaraCompilationGraphBridge::GetEmitterHandleIdString(const FEmitterNode* EmitterNode)
 {
 	return EmitterNode->GetEmitterHandleId().ToString(EGuidFormats::Digits);
@@ -383,6 +388,12 @@ FNiagaraVariable FNiagaraCompilationGraphBridge::GetInputVariable(const FInputNo
 	return InputNode->Input;
 }
 
+FNiagaraGraphInputDataInterfaceKey FNiagaraCompilationGraphBridge::GetGraphInputDataInterfaceKey(const UNiagaraNodeInput* InputNode, const FCompilationCopy* CompilationCopy)
+{
+	const UNiagaraDataInterface* InputDataInterface = InputNode->GetDataInterface();
+	return InputDataInterface != nullptr ? CompilationCopy->GetGraphInputDataInterfaceKey(*InputDataInterface) : FNiagaraGraphInputDataInterfaceKey();
+}
+
 const TArray<FNiagaraVariable>& FNiagaraCompilationGraphBridge::GetOutputVariables(const FOutputNode* OutputNode)
 {
 	return OutputNode->GetOutputs();
diff --git a/Engine/Plugins/FX/Niagara/Source/NiagaraEditor/Private/NiagaraCompilationPrivate.h b/Engine/Plugins/FX/Niagara/Source/NiagaraEditor/Private/NiagaraCompilationPrivate.h
--- a/Engine/Plugins/FX/Niagara/Source/NiagaraEditor/Private/NiagaraCompilationPrivate.h
+++ b/Engine/Plugins/FX/Niagara/Source/NiagaraEditor/Private/NiagaraCompilationPrivate.h
@@ -291,9 +291,10 @@ public:
 
 	virtual bool IsDuplicateDataFor(UNiagaraSystem* InSystem, UNiagaraEmitter* InEmitter, UNiagaraScript* InScript) const override;
 	virtual void GetDuplicatedObjects(TArray<UObject*>& Objects) override;
-	virtual const TMap<FName, UNiagaraDataInterface*>& GetObjectNameMap() override;
+	virtual const TMap<FNiagaraGraphInputDataInterfaceKey, UNiagaraDataInterface*>& GetGraphInputDataInterfaces() override;
 	virtual const UNiagaraScriptSourceBase* GetScriptSource() const override;
 	UNiagaraDataInterface* GetDuplicatedDataInterfaceCDOForClass(UClass* Class) const;
+	FNiagaraGraphInputDataInterfaceKey GetGraphInputDataInterfaceKey(const UNiagaraDataInterface& DataInterface) const;
 
 	TArray<FNiagaraParameterMapHistory>& GetPrecomputedHistories() { return PrecompiledHistories; }
 	const TArray<FNiagaraParameterMapHistory>& GetPrecomputedHistories() const { return PrecompiledHistories; }
@@ -335,10 +336,11 @@ public:
 	};
 
 	TSharedPtr<TMap<const UNiagaraGraph*, TArray<FDuplicatedGraphData>>> SharedSourceGraphToDuplicatedGraphsMap;
-	TSharedPtr<TMap<FName, UNiagaraDataInterface*>> SharedNameToDuplicatedDataInterfaceMap;
+	TSharedPtr<TMap<FNiagaraGraphInputDataInterfaceKey, UNiagaraDataInterface*>> SharedDuplicatedGraphInputDataInterfaces;
 	TSharedPtr<TMap<UClass*, UNiagaraDataInterface*>> SharedDataInterfaceClassToDuplicatedCDOMap;
 
 protected:
+	void CollectDuplicatedDataInterfaces(UNiagaraGraph* DuplicateGraph, FName SourceGraphPath);
 	void DuplicateReferencedGraphsRecursive(UNiagaraGraph* InGraph, const FCompileConstantResolver& ConstantResolver, TMap<UNiagaraNodeFunctionCall*, ENiagaraScriptUsage> FunctionsWithUsage);
 private:
 	TArray<TWeakObjectPtr<UNiagaraScriptSource>> TrackedScriptSourceCopies;
diff --git a/Engine/Plugins/FX/Niagara/Source/NiagaraEditor/Private/NiagaraCompilationTasks.cpp b/Engine/Plugins/FX/Niagara/Source/NiagaraEditor/Private/NiagaraCompilationTasks.cpp
--- a/Engine/Plugins/FX/Niagara/Source/NiagaraEditor/Private/NiagaraCompilationTasks.cpp
+++ b/Engine/Plugins/FX/Niagara/Source/NiagaraEditor/Private/NiagaraCompilationTasks.cpp
@@ -33,6 +33,14 @@
 #include "NiagaraNodeOutput.h"
 #include "NiagaraNodeFunctionCall.h"
 
+static int GNiagaraCompilationForceOneThread = 0;
+static FAutoConsoleVariableRef CVarNiagaraCompilationForceOneThread(
+	TEXT("fx.Niagara.Compilation.ForceOneThread"),
+	GNiagaraCompilationForceOneThread,
+	TEXT("Forces the compile tasks to run inline on the main thread."),
+	ECVF_Default
+);
+
 #if ENABLE_COOK_STATS
 namespace NiagaraSystemCookStats
 {
@@ -828,84 +836,6 @@ const FNiagaraSystemCompilationTask::FScriptInfo* FNiagaraSystemCompilationTask:
 	return DigestedScriptInfo.Find(ScriptKey);
 }
 
-void FNiagaraSystemCompilationTask::FCompileTaskInfo::CollectNamedDataInterfaces(FNiagaraSystemCompilationTask* SystemCompileTask, const FCompileGroupInfo& GroupInfo)
-{
-	FString UniqueEmitterName;
-	const FNiagaraCompilationCopyData* CompilationCopyData = nullptr;
-
-	if (const FEmitterInfo* EmitterInfo = SystemCompileTask->SystemInfo.EmitterInfoBySourceEmitter(GroupInfo.SourceEmitterIndex))
-	{
-		CompilationCopyData = static_cast<const FNiagaraCompilationCopyData*>(GroupInfo.CompilationCopy->GetDependentRequest(EmitterInfo->DigestedEmitterIndex).Get());
-		UniqueEmitterName = EmitterInfo->UniqueEmitterName;
-	}
-	else
-	{
-		CompilationCopyData = GroupInfo.CompilationCopy.Get();
-	}
-
-	const FScriptInfo& ScriptInfo = SystemCompileTask->DigestedScriptInfo.FindChecked(ScriptKey);
-
-	auto AccumulateInputDataInterfaces = [this, &UniqueEmitterName](const FNiagaraCompilationNode& Node) -> bool
-	{
-		if (const FNiagaraCompilationNodeInput* InputNode = Node.AsType<FNiagaraCompilationNodeInput>())
-		{
-			if (InputNode->DataInterfaceName != NAME_None && InputNode->DuplicatedDataInterface)
-			{
-				constexpr bool bIsParameterMapDataInterface = false;
-				FName DIName = FNiagaraHlslTranslator::GetDataInterfaceName(InputNode->DataInterfaceName, UniqueEmitterName, bIsParameterMapDataInterface);
-
-				NamedDataInterfaces.Add(DIName, InputNode->DuplicatedDataInterface);
-			}
-		}
-
-		return true;
-	};
-
-	auto DefaultNodeFinder = [&ScriptInfo](const FNiagaraCompilationNodeOutput& OutputNode) -> bool
-	{
-		if (UNiagaraScript::IsEquivalentUsage(OutputNode.Usage, ScriptInfo.Usage) && OutputNode.UsageId == ScriptInfo.UsageId)
-		{
-			return true;
-		}
-
-		if (ScriptInfo.Usage == ENiagaraScriptUsage::ParticleGPUComputeScript)
-		{
-			return UNiagaraScript::IsParticleScript(OutputNode.Usage);
-		}
-		else if (ScriptInfo.Usage == ENiagaraScriptUsage::ParticleSpawnScriptInterpolated)
-		{
-			return UNiagaraScript::IsEquivalentUsage(OutputNode.Usage, ENiagaraScriptUsage::ParticleUpdateScript);
-		}
-
-		return OutputNode.Usage == ScriptInfo.Usage && OutputNode.UsageId == ScriptInfo.UsageId;
-	};
-
-	CompilationCopyData->InstantiatedGraph->NodeTraversal(true, false, DefaultNodeFinder, AccumulateInputDataInterfaces);
-
-	if (UNiagaraScript::IsSystemScript(ScriptInfo.Usage))
-	{
-		check(GroupInfo.SourceEmitterIndex == INDEX_NONE);
-
-		const ENiagaraScriptUsage EmitterUsage = ScriptInfo.Usage == ENiagaraScriptUsage::SystemSpawnScript
-			? ENiagaraScriptUsage::EmitterSpawnScript
-			: ENiagaraScriptUsage::EmitterUpdateScript;
-		const FGuid EmitterUsageId = ScriptInfo.UsageId;
-		auto EmitterNodeFinder = [EmitterUsage, EmitterUsageId](const FNiagaraCompilationNodeOutput& OutputNode) -> bool
-		{
-			return UNiagaraScript::IsEquivalentUsage(OutputNode.Usage, EmitterUsage) && OutputNode.UsageId == EmitterUsageId;
-		};
-
-		for (const FNiagaraCompilationCopyData::FSharedCompilationCopy& EmitterCompilationCopy : CompilationCopyData->EmitterData)
-		{
-			if (const FNiagaraCompilationGraph* EmitterGraph = EmitterCompilationCopy->InstantiatedGraph.Get())
-			{
-				UniqueEmitterName = EmitterCompilationCopy->EmitterUniqueName;
-				EmitterGraph->NodeTraversal(true, false, EmitterNodeFinder, AccumulateInputDataInterfaces);
-			}
-		}
-	}
-}
-
 void FNiagaraSystemCompilationTask::FCompileTaskInfo::Translate(FNiagaraSystemCompilationTask* SystemCompileTask, const FCompileGroupInfo& GroupInfo)
 {
 	TRACE_CPUPROFILER_EVENT_SCOPE(NiagaraAsyncTask_Translate);
@@ -1881,12 +1811,6 @@ bool FNiagaraSystemCompilationTask::Poll(FNiagaraSystemAsyncCompileResults& Resu
 						ResultsInfo.CompilationErrors = ShaderTaskInfo.CompilationErrors;
 					}
 
-					CompileData.NamedDataInterfaces.Reserve(TaskInfo.NamedDataInterfaces.Num());
-					Algo::Transform(TaskInfo.NamedDataInterfaces, CompileData.NamedDataInterfaces, [](const TTuple<FName, UNiagaraDataInterface*>& InElement)
-					{
-						return MakeTuple(InElement.Key, InElement.Value);
-					});
-
 					CompileData.CompileMetrics.TaskWallTime = (float) (FPlatformTime::Seconds() - TaskInfo.TaskStartTime);
 					CompileData.CompileMetrics.DDCFetchTime = 0.0f;
 					CompileData.CompileMetrics.CompilerWallTime = TaskInfo.CompilerWallTime;
@@ -2646,6 +2570,11 @@ void FNiagaraSystemCompilationTask::IssueCompilationTasks()
 	TArray<FTask> IssuedCompilationTasks;
 	TArray<FTaskEvent> CompilationTasks;
 
+	ETaskPriority TaskPriority = ETaskPriority::Normal;
+	EExtendedTaskPriority TaskExtendedPriority = GNiagaraCompilationForceOneThread > 0 
+		? EExtendedTaskPriority::Inline 
+		: EExtendedTaskPriority::None;
+
 	if (HasOutstandingCompileTasks())
 	{
 		FTask PrecompileTask = Launch(UE_SOURCE_LOCATION, [Self = AsShared()]
@@ -2653,7 +2582,7 @@ void FNiagaraSystemCompilationTask::IssueCompilationTasks()
 			COOK_STAT(auto Timer = NiagaraSystemCookStats::UsageStats.TimeSyncWork(); Timer.TrackCyclesOnly(););
 
 			Self->Precompile();
-		});
+		}, TaskPriority, TaskExtendedPriority);
 
 		// async compilation copy tasks
 		for (FCompileGroupInfo& GroupInfo : CompileGroups)
@@ -2665,7 +2594,7 @@ void FNiagaraSystemCompilationTask::IssueCompilationTasks()
 					COOK_STAT(auto Timer = NiagaraSystemCookStats::UsageStats.TimeSyncWork(); Timer.TrackCyclesOnly(););
 
 					GroupInfo.InstantiateCompileGraph(*Self);
-				}, PrecompileTask);
+				}, PrecompileTask, TaskPriority, TaskExtendedPriority);
 
 				for (int32 CompileTaskIndex : GroupInfo.CompileTaskIndices)
 				{
@@ -2681,7 +2610,6 @@ void FNiagaraSystemCompilationTask::IssueCompilationTasks()
 
 							FNiagaraSystemCompilationTask* TaskPtr = &Self.Get();
 
-							CompileTask.CollectNamedDataInterfaces(TaskPtr, GroupInfo);
 							switch (CompileTask.ScriptCompileType)
 							{
 								case EScriptCompileType::CompileForVm:
@@ -2709,14 +2637,14 @@ void FNiagaraSystemCompilationTask::IssueCompilationTasks()
 									CompileTask.Translate(TaskPtr, GroupInfo);
 								break;
 							}
-						}, InstantiateGraphTask));
+						}, InstantiateGraphTask, TaskPriority, TaskExtendedPriority));
 
 						Launch(UE_SOURCE_LOCATION, [Self = AsShared(), &GroupInfo, &CompileTask]
 						{
 							COOK_STAT(auto Timer = NiagaraSystemCookStats::UsageStats.TimeSyncWork(); Timer.TrackCyclesOnly(););
 
 							CompileTask.ProcessCompilation(&Self.Get(), GroupInfo);
-						}, *CompileTask.CompileResultsReadyEvent);
+						}, *CompileTask.CompileResultsReadyEvent, TaskPriority, TaskExtendedPriority);
 
 						CompilationTasks.Add(*CompileTask.CompileResultsProcessedEvent);
 					}
@@ -2728,12 +2656,12 @@ void FNiagaraSystemCompilationTask::IssueCompilationTasks()
 	Launch(UE_SOURCE_LOCATION, [Self = AsShared()]
 	{
 		Self->CompilationState = EState::WaitingForProcessing;
-	}, IssuedCompilationTasks);
+	}, IssuedCompilationTasks, TaskPriority, TaskExtendedPriority);
 
 	Launch(UE_SOURCE_LOCATION, [Self = AsShared()]
 	{
 		Self->IssuePostResultsProcessedTasks();
-	}, CompilationTasks);
+	}, CompilationTasks, TaskPriority, TaskExtendedPriority);
 }
 
 bool FNiagaraSystemCompilationTask::HasOutstandingCompileTasks() const
diff --git a/Engine/Plugins/FX/Niagara/Source/NiagaraEditor/Private/NiagaraCompilationTasks.h b/Engine/Plugins/FX/Niagara/Source/NiagaraEditor/Private/NiagaraCompilationTasks.h
--- a/Engine/Plugins/FX/Niagara/Source/NiagaraEditor/Private/NiagaraCompilationTasks.h
+++ b/Engine/Plugins/FX/Niagara/Source/NiagaraEditor/Private/NiagaraCompilationTasks.h
@@ -204,7 +204,6 @@ struct FNiagaraSystemCompilationTask : public TSharedFromThis<FNiagaraSystemComp
 
 		TUniquePtr<UE::Tasks::FTaskEvent> CompileResultsReadyEvent;
 		TUniquePtr<UE::Tasks::FTaskEvent> CompileResultsProcessedEvent;
-		TMap<FName, UNiagaraDataInterface*> NamedDataInterfaces;
 		TArray<FNiagaraVariable> BakedRapidIterationParameters; // todo - remove and have the data pulled directly from the precompile data
 		double TaskStartTime = 0;
 		float TranslationTime = 0.0f;
@@ -215,7 +214,6 @@ struct FNiagaraSystemCompilationTask : public TSharedFromThis<FNiagaraSystemComp
 		EScriptCompileType ScriptCompileType = EScriptCompileType::Invalid;
 
 		bool IsOutstanding(const FNiagaraSystemCompilationTask& ParentTask) const;
-		void CollectNamedDataInterfaces(FNiagaraSystemCompilationTask* SystemCompileTask, const FCompileGroupInfo& GroupInfo);
 		void Translate(FNiagaraSystemCompilationTask* SystemCompileTask, const FCompileGroupInfo& GroupInfo);
 		void IssueCompileVm(FNiagaraSystemCompilationTask* SystemCompileTask, const FCompileGroupInfo& GroupInfo);
 		void IssueTranslateGpu(FNiagaraSystemCompilationTask* SystemCompileTask, const FCompileGroupInfo& GroupInfo);
diff --git a/Engine/Plugins/FX/Niagara/Source/NiagaraEditor/Private/NiagaraCompiler.cpp b/Engine/Plugins/FX/Niagara/Source/NiagaraEditor/Private/NiagaraCompiler.cpp
--- a/Engine/Plugins/FX/Niagara/Source/NiagaraEditor/Private/NiagaraCompiler.cpp
+++ b/Engine/Plugins/FX/Niagara/Source/NiagaraEditor/Private/NiagaraCompiler.cpp
@@ -154,6 +154,29 @@ void FNiagaraCompileRequestDuplicateData::DuplicateReferencedGraphs(UNiagaraGrap
 	DuplicateReferencedGraphsRecursive(InDupeGraph, ConstantResolver, FunctionsWithUsage);
 }
 
+void FNiagaraCompileRequestDuplicateData::CollectDuplicatedDataInterfaces(UNiagaraGraph* DuplicateGraph, FName SourceGraphPath)
+{
+	TArray<UNiagaraNodeInput*> InputNodes;
+	DuplicateGraph->GetNodesOfClass(InputNodes);
+	const UEdGraphSchema_Niagara* Schema = GetDefault<UEdGraphSchema_Niagara>();
+
+	for (UNiagaraNodeInput* InputNode : InputNodes)
+	{
+		if (InputNode)
+		{
+			if (InputNode->Input.IsDataInterface())
+			{
+				UNiagaraDataInterface* DataInterface = InputNode->GetDataInterface();
+				if (DataInterface != nullptr)
+				{
+					FNiagaraGraphInputDataInterfaceKey GraphInputDataInterfaceKey(SourceGraphPath, InputNode->Input);
+					SharedDuplicatedGraphInputDataInterfaces->Add(GraphInputDataInterfaceKey, DataInterface);
+				}
+			}
+		}
+	}
+}
+
 void FNiagaraCompileRequestDuplicateData::DuplicateReferencedGraphsRecursive(UNiagaraGraph* InGraph, const FCompileConstantResolver& ConstantResolver, TMap<UNiagaraNodeFunctionCall*, ENiagaraScriptUsage> FunctionsWithUsage)
 {
 	if (!InGraph)
@@ -169,20 +192,6 @@ void FNiagaraCompileRequestDuplicateData::DuplicateReferencedGraphsRecursive(UNi
 	{
 		if (UNiagaraNode* InNode = Cast<UNiagaraNode>(Node))
 		{
-			UNiagaraNodeInput* InputNode = Cast<UNiagaraNodeInput>(InNode);
-			if (InputNode)
-			{
-				if (InputNode->Input.IsDataInterface())
-				{
-					UNiagaraDataInterface* DataInterface = InputNode->GetDataInterface();
-					bool bIsParameterMapDataInterface = false;
-					FName DIName = FNiagaraHlslTranslator::GetDataInterfaceName(InputNode->Input.GetName(), EmitterUniqueName, bIsParameterMapDataInterface);
-					UNiagaraDataInterface* Dupe = PrecompileDuplicateObject<UNiagaraDataInterface>(DataInterface, GetTransientPackage());
-					SharedNameToDuplicatedDataInterfaceMap->Add(DIName, Dupe);
-				}
-				continue;
-			}
-
 			UNiagaraNodeFunctionCall* FunctionCallNode = Cast<UNiagaraNodeFunctionCall>(InNode);
 			if (FunctionCallNode)
 			{
@@ -232,6 +241,7 @@ void FNiagaraCompileRequestDuplicateData::DuplicateReferencedGraphsRecursive(UNi
 
 						// Assign the copied script and process any child scripts.
 						FunctionCallNode->FunctionScript = DupeScript;
+						CollectDuplicatedDataInterfaces(DupeGraph, *FunctionGraph->GetPathName());
 						DuplicateReferencedGraphsRecursive(DupeGraph, FunctionConstantResolver, FunctionsWithUsage);
 					}
 					else
@@ -261,9 +271,9 @@ void FNiagaraCompileRequestDuplicateData::DuplicateReferencedGraphsRecursive(UNi
 	}
 }
 
-const TMap<FName, UNiagaraDataInterface*>& FNiagaraCompileRequestDuplicateData::GetObjectNameMap()
+const TMap<FNiagaraGraphInputDataInterfaceKey, UNiagaraDataInterface*>& FNiagaraCompileRequestDuplicateData::GetGraphInputDataInterfaces()
 {
-	return *SharedNameToDuplicatedDataInterfaceMap.Get();
+	return *SharedDuplicatedGraphInputDataInterfaces.Get();
 }
 
 const UNiagaraScriptSourceBase* FNiagaraCompileRequestDuplicateData::GetScriptSource() const
@@ -284,6 +294,18 @@ UNiagaraDataInterface* FNiagaraCompileRequestDuplicateData::GetDuplicatedDataInt
 	return nullptr;
 }
 
+FNiagaraGraphInputDataInterfaceKey FNiagaraCompileRequestDuplicateData::GetGraphInputDataInterfaceKey(const UNiagaraDataInterface& DataInterface) const
+{
+	for (const TPair<FNiagaraGraphInputDataInterfaceKey, UNiagaraDataInterface*>& GraphInputDataInterfacePair : *SharedDuplicatedGraphInputDataInterfaces)
+	{
+		if (GraphInputDataInterfacePair.Value == &DataInterface)
+		{
+			return GraphInputDataInterfacePair.Key;
+		}
+	}
+	return FNiagaraGraphInputDataInterfaceKey();
+}
+
 void FNiagaraCompileRequestData::SortOutputNodesByDependencies(TArray<class UNiagaraNodeOutput*>& NodesToSort, const TArray<class UNiagaraSimulationStageBase*>* SimStages)
 {
 	if (!SimStages)
@@ -355,10 +377,10 @@ void FNiagaraCompileRequestDuplicateData::GetDuplicatedObjects(TArray<UObject*>&
 	Objects.Add(SourceDeepCopy.Get());
 	Objects.Add(NodeGraphDeepCopy.Get());
 
-	if (SharedNameToDuplicatedDataInterfaceMap.IsValid())
+	if (SharedDuplicatedGraphInputDataInterfaces.IsValid())
 	{
 		TArray<UNiagaraDataInterface*> DIs;
-		SharedNameToDuplicatedDataInterfaceMap->GenerateValueArray(DIs);
+		SharedDuplicatedGraphInputDataInterfaces->GenerateValueArray(DIs);
 		for (UNiagaraDataInterface* DI : DIs)
 		{
 			Objects.Add(DI);
@@ -419,6 +441,7 @@ void FNiagaraCompileRequestDuplicateData::DeepCopyGraphs(UNiagaraScriptSource* S
 	SourceDeepCopy = ScriptSource->CreateCompilationCopy(ValidUsages);
 	NodeGraphDeepCopy = SourceDeepCopy->NodeGraph;
 	FEdGraphUtilities::MergeChildrenGraphsIn(NodeGraphDeepCopy.Get(), NodeGraphDeepCopy.Get(), /*bRequireSchemaMatch=*/ true);
+	CollectDuplicatedDataInterfaces(NodeGraphDeepCopy.Get(), *ScriptSource->NodeGraph.GetPathName());
 	DuplicateReferencedGraphs(ScriptSource->NodeGraph, NodeGraphDeepCopy.Get(), InUsage, ConstantResolver);
 }
 
@@ -448,6 +471,7 @@ void FNiagaraCompileRequestDuplicateData::DeepCopyGraphs(const FVersionedNiagara
 	}
 
 	FCompileConstantResolver ConstantResolver(Emitter, ENiagaraScriptUsage::EmitterSpawnScript);
+	CollectDuplicatedDataInterfaces(NodeGraphDeepCopy.Get(), *ScriptSource->NodeGraph.GetPathName());
 	DuplicateReferencedGraphs(ScriptSource->NodeGraph, NodeGraphDeepCopy.Get(), ENiagaraScriptUsage::EmitterSpawnScript, ConstantResolver, FunctionsWithUsage);
 }
 
@@ -1310,7 +1334,7 @@ TSharedPtr<FNiagaraCompileRequestDuplicateDataBase, ESPMode::ThreadSafe> FNiagar
 	FCompileConstantResolver EmptyResolver;
 
 	BasePtr->SharedSourceGraphToDuplicatedGraphsMap = MakeShared<TMap<const UNiagaraGraph*, TArray<FNiagaraCompileRequestDuplicateData::FDuplicatedGraphData>>>();
-	BasePtr->SharedNameToDuplicatedDataInterfaceMap = MakeShared<TMap<FName, UNiagaraDataInterface*>>();
+	BasePtr->SharedDuplicatedGraphInputDataInterfaces = MakeShared<TMap<FNiagaraGraphInputDataInterfaceKey, UNiagaraDataInterface*>>();
 	BasePtr->SharedDataInterfaceClassToDuplicatedCDOMap = MakeShared<TMap<UClass*, UNiagaraDataInterface*>>();
 	BasePtr->OwningSystem = OwningSystem;
 	BasePtr->OwningEmitter = OwningEmitter;
@@ -1356,7 +1380,7 @@ TSharedPtr<FNiagaraCompileRequestDuplicateDataBase, ESPMode::ThreadSafe> FNiagar
 			EmitterPtr->EmitterID = FNiagaraEmitterID(i);
 			EmitterPtr->ValidUsages = BasePtr->ValidUsages;
 			EmitterPtr->SharedSourceGraphToDuplicatedGraphsMap = BasePtr->SharedSourceGraphToDuplicatedGraphsMap;
-			EmitterPtr->SharedNameToDuplicatedDataInterfaceMap = BasePtr->SharedNameToDuplicatedDataInterfaceMap;
+			EmitterPtr->SharedDuplicatedGraphInputDataInterfaces = BasePtr->SharedDuplicatedGraphInputDataInterfaces;
 			EmitterPtr->SharedDataInterfaceClassToDuplicatedCDOMap = BasePtr->SharedDataInterfaceClassToDuplicatedCDOMap;
 			//EmitterPtr->bSimulationStagesEnabled = Handle.GetInstance()->bSimulationStagesEnabled;
 			if (Handle.GetIsEnabled() && Handle.GetInstance().Emitter && (OwningEmitter == nullptr || OwningEmitter == Handle.GetInstance().Emitter)) // Don't need to copy the graph if we aren't going to use it.
diff --git a/Engine/Plugins/FX/Niagara/Source/NiagaraEditor/Private/NiagaraEditorUtilities.cpp b/Engine/Plugins/FX/Niagara/Source/NiagaraEditor/Private/NiagaraEditorUtilities.cpp
--- a/Engine/Plugins/FX/Niagara/Source/NiagaraEditor/Private/NiagaraEditorUtilities.cpp
+++ b/Engine/Plugins/FX/Niagara/Source/NiagaraEditor/Private/NiagaraEditorUtilities.cpp
@@ -41,6 +41,7 @@
 #include "ViewModels/NiagaraScriptGraphViewModel.h"
 #include "NiagaraScriptMergeManager.h"
 #include "NiagaraScriptSource.h"
+#include "NiagaraScriptRuntimeData.h"
 #include "NiagaraSettings.h"
 #include "NiagaraSimulationStageBase.h"
 #include "NiagaraStackEditorData.h"
@@ -4914,29 +4915,34 @@ UNiagaraDataInterface* FNiagaraEditorUtilities::GetResolvedRuntimeInstanceForEdi
 		}
 
 		// If the data interface was owned by an input node, then we need to try to update the compiled version.
-		const FNiagaraEmitterHandle* EmitterHandle;
+		const FNiagaraEmitterHandle* EmitterHandle = nullptr;
 		TArray<const UNiagaraScript*> CompiledScripts;
 		
+		FNiagaraGraphInputDataInterfaceKey GraphInputDataInterfaceKey(*OuterInputNode->GetGraph()->GetPathName(), OuterInputNode->Input);
 		FNiagaraStackGraphUtilities::GetEmitterHandleAndCompiledScriptsForStackNode(OwningSystem, *OuterInputNode, EmitterHandle, CompiledScripts);
 		if (ensureMsgf(CompiledScripts.Num() > 0, TEXT("Could not find compiled scripts for data interface input node.")))
 		{
 			for (const UNiagaraScript* CompiledScript : CompiledScripts)
 			{
-				FString EmitterName = EmitterHandle != nullptr ? EmitterHandle->GetUniqueInstanceName() : FString();
-				bool bIsParameterMapDataInterface = false;
-				FName DataInterfaceName = FNiagaraHlslTranslator::GetDataInterfaceName(OuterInputNode->Input.GetName(), EmitterName, bIsParameterMapDataInterface);
-				
-				const FNiagaraEmitterHandle* EmitterHandleForScript = UNiagaraScript::IsSystemScript(CompiledScript->GetUsage()) ? nullptr : EmitterHandle;
-				const FNiagaraScriptRuntimeData* ScriptRuntimeData = OwningSystem.GetScriptRuntimeData(EmitterHandleForScript, *CompiledScript).Get();
-				if (ScriptRuntimeData != nullptr)
+				const FNiagaraScriptDataInterfaceCompileInfo* GraphInputDICompileInfo = CompiledScript->GetVMExecutableData().DataInterfaceInfo.FindByPredicate([&GraphInputDataInterfaceKey]
+					(const FNiagaraScriptDataInterfaceCompileInfo& ScriptDICompileInfo) { return ScriptDICompileInfo.SourceGraphInputDataInterfaceKey == GraphInputDataInterfaceKey; });
+				if (GraphInputDICompileInfo == nullptr)
 				{
-					for (const FNiagaraScriptResolvedDataInterfaceInfo& ResolvedDI : ScriptRuntimeData->GetResolvedDataInterfaces())
-					{
-						if (ResolvedDI.Name == DataInterfaceName)
-						{
-							return ResolvedDI.ResolvedDataInterface;
-						}
-					}
+					continue;
+				}
+
+				FGuid EmitterHandleId = (EmitterHandle != nullptr && UNiagaraScript::IsSystemScript(CompiledScript->GetUsage()) == false) ? EmitterHandle->GetId() : FGuid();
+				const FNiagaraScriptRuntimeData* ScriptRuntimeData = OwningSystem.GetScriptRuntimeData(EmitterHandleId, CompiledScript->GetUsage(), CompiledScript->GetUsageId()).Get();
+				if (ScriptRuntimeData == nullptr)
+				{
+					continue;
+				}
+
+				const FNiagaraScriptResolvedDataInterfaceInfo* ResolvedGraphInputDIData = ScriptRuntimeData->GetResolvedDataInterfaces().FindByPredicate([GraphInputDICompileInfo]
+					(const FNiagaraScriptResolvedDataInterfaceInfo& ResolvedDataInterfaceData) { return ResolvedDataInterfaceData.CompileName == GraphInputDICompileInfo->Name; });
+				if (ResolvedGraphInputDIData != nullptr)
+				{
+					return ResolvedGraphInputDIData->ResolvedDataInterface;
 				}
 			}
 		}
diff --git a/Engine/Plugins/FX/Niagara/Source/NiagaraEditor/Private/NiagaraGraphDigest.cpp b/Engine/Plugins/FX/Niagara/Source/NiagaraEditor/Private/NiagaraGraphDigest.cpp
--- a/Engine/Plugins/FX/Niagara/Source/NiagaraEditor/Private/NiagaraGraphDigest.cpp
+++ b/Engine/Plugins/FX/Niagara/Source/NiagaraEditor/Private/NiagaraGraphDigest.cpp
@@ -3291,6 +3291,7 @@ FNiagaraCompilationNodeInput::FNiagaraCompilationNodeInput(const UNiagaraNodeInp
 	bExposed = InNode->IsExposed();
 	bRequired = InNode->ExposureOptions.bRequired;
 	bCanAutoBind = InNode->ExposureOptions.bCanAutoBind;
+	GraphPathFName = *InNode->GetGraph()->GetPathName();
 }
 
 
@@ -3305,6 +3306,7 @@ FNiagaraCompilationNodeInput::FNiagaraCompilationNodeInput(const FNiagaraCompila
 	, bRequired(InNode.bRequired)
 	, bExposed(InNode.bExposed)
 	, bCanAutoBind(InNode.bCanAutoBind)
+	, GraphPathFName(InNode.GraphPathFName)
 	, DuplicatedDataInterface(InNode.DuplicatedDataInterface)
 	, ObjectAssetPath(InNode.ObjectAssetPath)
 {
@@ -3381,7 +3383,7 @@ void FNiagaraCompilationNodeInput::Compile(FTranslator* Translator, TArray<int32
 					{
 						check(DuplicatedDataInterface);
 						check(InputVariable.IsDataInterface());
-						Outputs.Add(Translator->RegisterDataInterface(InputVariable, DuplicatedDataInterface, false, false));
+						Outputs.Add(Translator->RegisterDataInterfaceForInput(InputVariable, Translator->GetGraphInputDataInterfaceKey(this), DuplicatedDataInterface.Get()));
 						return;
 					}
 					else if (!ObjectAssetName.IsNone())
@@ -3406,7 +3408,7 @@ void FNiagaraCompilationNodeInput::Compile(FTranslator* Translator, TArray<int32
 		{
 			check(DuplicatedDataInterface);
 			check(InputVariable.IsDataInterface());
-			Outputs.Add(Translator->RegisterDataInterface(InputVariable, DuplicatedDataInterface, false, false));
+			Outputs.Add(Translator->RegisterDataInterfaceForInput(InputVariable, Translator->GetGraphInputDataInterfaceKey(this), DuplicatedDataInterface.Get()));
 			break;
 		}
 		else if (!ObjectAssetName.IsNone())
diff --git a/Engine/Plugins/FX/Niagara/Source/NiagaraEditor/Private/NiagaraGraphDigest.h b/Engine/Plugins/FX/Niagara/Source/NiagaraEditor/Private/NiagaraGraphDigest.h
--- a/Engine/Plugins/FX/Niagara/Source/NiagaraEditor/Private/NiagaraGraphDigest.h
+++ b/Engine/Plugins/FX/Niagara/Source/NiagaraEditor/Private/NiagaraGraphDigest.h
@@ -514,6 +514,7 @@ public:
 	bool bRequired;
 	bool bExposed;
 	bool bCanAutoBind;
+	FName GraphPathFName = NAME_None;
 
 	TObjectPtr<UNiagaraDataInterface> DuplicatedDataInterface;
 	FSoftObjectPath ObjectAssetPath;
diff --git a/Engine/Plugins/FX/Niagara/Source/NiagaraEditor/Private/NiagaraHlslTranslator.cpp b/Engine/Plugins/FX/Niagara/Source/NiagaraEditor/Private/NiagaraHlslTranslator.cpp
--- a/Engine/Plugins/FX/Niagara/Source/NiagaraEditor/Private/NiagaraHlslTranslator.cpp
+++ b/Engine/Plugins/FX/Niagara/Source/NiagaraEditor/Private/NiagaraHlslTranslator.cpp
@@ -1613,6 +1613,32 @@ FNiagaraTranslateResults TNiagaraHlslTranslator<GraphBridge>::Translate(const FN
 		return TranslateResults;
 	}
 
+	auto GetDIOwnerNameByUsage = [](ENiagaraScriptUsage Usage)->FName
+	{
+		switch (Usage)
+		{
+			case ENiagaraScriptUsage::SystemSpawnScript:
+				return "SystemSpawn";
+			case ENiagaraScriptUsage::SystemUpdateScript:
+				return "SystemUpdate";
+			case ENiagaraScriptUsage::EmitterSpawnScript:
+				return "EmitterSpawn";
+			case ENiagaraScriptUsage::EmitterUpdateScript:
+				return "EmitterUpdate";
+			case ENiagaraScriptUsage::ParticleSpawnScript:
+			case ENiagaraScriptUsage::ParticleSpawnScriptInterpolated:
+				return "ParticleSpawn";
+			case ENiagaraScriptUsage::ParticleUpdateScript:
+				return "ParticleUpdate";
+			case ENiagaraScriptUsage::ParticleSimulationStageScript:
+				return "ParticleStage";
+			case ENiagaraScriptUsage::ParticleEventScript:
+				return "ParticleEvent";
+			default:
+				return "Unknown";
+		}
+	};
+
 	switch (CompileOptions.TargetUsage)
 	{
 	case ENiagaraScriptUsage::ParticleSpawnScriptInterpolated:
@@ -1624,6 +1650,7 @@ FNiagaraTranslateResults TNiagaraHlslTranslator<GraphBridge>::Translate(const FN
 		SpawnStage.NumIterations = 1;
 		SpawnStage.bWritesParticles = true;
 		SpawnStage.bShouldUpdateInitialAttributeValues = true;
+		SpawnStage.DIOwnerName = GetDIOwnerNameByUsage(SpawnStage.ScriptUsage);
 		OutputNodes.Add(SourceGraph->FindEquivalentOutputNode(ENiagaraScriptUsage::ParticleSpawnScript, SpawnStage.UsageId));
 
 		FNiagaraHlslTranslationStage& UpdateStage = TranslationStages.Emplace_GetRef(ENiagaraScriptUsage::ParticleUpdateScript, FGuid());
@@ -1633,6 +1660,7 @@ FNiagaraTranslateResults TNiagaraHlslTranslator<GraphBridge>::Translate(const FN
 		UpdateStage.SimulationStageIndex = 0;
 		UpdateStage.NumIterations = 1;
 		UpdateStage.bWritesParticles = true;
+		UpdateStage.DIOwnerName = GetDIOwnerNameByUsage(UpdateStage.ScriptUsage);
 		OutputNodes.Add(SourceGraph->FindEquivalentOutputNode(UpdateStage.ScriptUsage, UpdateStage.UsageId));
 
 		ParamMapHistories.AddDefaulted(2);
@@ -1648,6 +1676,7 @@ FNiagaraTranslateResults TNiagaraHlslTranslator<GraphBridge>::Translate(const FN
 		SpawnStage.NumIterations = 1;
 		SpawnStage.bWritesParticles = true;
 		SpawnStage.bShouldUpdateInitialAttributeValues = true;
+		SpawnStage.DIOwnerName = GetDIOwnerNameByUsage(SpawnStage.ScriptUsage);
 		OutputNodes.Add(SourceGraph->FindEquivalentOutputNode(ENiagaraScriptUsage::ParticleSpawnScript, TranslationStages[0].UsageId));
 
 		FNiagaraHlslTranslationStage& UpdateStage = TranslationStages.Emplace_GetRef(ENiagaraScriptUsage::ParticleUpdateScript, FGuid());
@@ -1661,6 +1690,7 @@ FNiagaraTranslateResults TNiagaraHlslTranslator<GraphBridge>::Translate(const FN
 		UpdateStage.SimulationStageIndex = 0;
 		UpdateStage.NumIterations = 1;
 		UpdateStage.bWritesParticles = true;
+		UpdateStage.DIOwnerName = GetDIOwnerNameByUsage(UpdateStage.ScriptUsage);
 		OutputNodes.Add(SourceGraph->FindEquivalentOutputNode(TranslationStages[1].ScriptUsage, TranslationStages[1].UsageId));
 
 		ParamMapHistories.AddDefaulted(2);
@@ -1750,6 +1780,7 @@ FNiagaraTranslateResults TNiagaraHlslTranslator<GraphBridge>::Translate(const FN
 					SimStage.DirectDispatchElementType = CompileSimStageData.DirectDispatchElementType;
 					SimStage.bOverrideGpuDispatchNumThreads = CompileSimStageData.bOverrideGpuDispatchNumThreads;
 					SimStage.OverrideGpuDispatchNumThreads = CompileSimStageData.OverrideGpuDispatchNumThreads;
+					SimStage.DIOwnerName = *FString::Printf(TEXT("%s%d%s"), *GetDIOwnerNameByUsage(SimStage.ScriptUsage).ToString(), DestSimStageIndex, *StageName);
 
 					if (CompileSimStageData.bOverrideGpuDispatchNumThreads)
 					{
@@ -1826,17 +1857,16 @@ FNiagaraTranslateResults TNiagaraHlslTranslator<GraphBridge>::Translate(const FN
 					
 					// Set up the compile output for the shader stages so that we can properly execute at runtime.
 					FSimulationStageMetaData& SimulationStageMetaData = CompilationOutput.ScriptData.SimulationStageMetaData.AddDefaulted_GetRef();
-					SimulationStageMetaData.SimulationStageName = CompileData->CompileSimStageData[SourceSimStageIndex].StageName;
-					SimulationStageMetaData.EnabledBinding = CompileData->CompileSimStageData[SourceSimStageIndex].EnabledBinding;
-					SimulationStageMetaData.ElementCount = CompileData->CompileSimStageData[SourceSimStageIndex].ElementCount;
-					SimulationStageMetaData.ElementCountXBinding = CompileData->CompileSimStageData[SourceSimStageIndex].ElementCountXBinding;
-					SimulationStageMetaData.ElementCountYBinding = CompileData->CompileSimStageData[SourceSimStageIndex].ElementCountYBinding;
-					SimulationStageMetaData.ElementCountZBinding = CompileData->CompileSimStageData[SourceSimStageIndex].ElementCountZBinding;
+					SimulationStageMetaData.SimulationStageName = CompileSimStageData.StageName;
+					SimulationStageMetaData.EnabledBinding = SimStage.EnabledBinding;
+					SimulationStageMetaData.ElementCount = SimStage.ElementCount;
+					SimulationStageMetaData.ElementCountXBinding = SimStage.ElementCountXBinding;
+					SimulationStageMetaData.ElementCountYBinding = SimStage.ElementCountYBinding;
+					SimulationStageMetaData.ElementCountZBinding = SimStage.ElementCountZBinding;
 					SimulationStageMetaData.ExecuteBehavior = SimStage.ExecuteBehavior;
-					SimulationStageMetaData.IterationSourceType = CompileData->CompileSimStageData[SourceSimStageIndex].IterationSourceType;
-					SimulationStageMetaData.IterationDataInterface = CompileData->CompileSimStageData[SourceSimStageIndex].IterationDataInterface;
-					SimulationStageMetaData.IterationDirectBinding = CompileData->CompileSimStageData[SourceSimStageIndex].IterationDirectBinding;
-					SimulationStageMetaData.NumIterationsBinding = CompileData->CompileSimStageData[SourceSimStageIndex].NumIterationsBinding;
+					SimulationStageMetaData.IterationSourceType = SimStage.IterationSourceType;
+					SimulationStageMetaData.IterationDirectBinding = SimStage.IterationDirectBinding;
+					SimulationStageMetaData.NumIterationsBinding = SimStage.NumIterationsBinding;
 					SimulationStageMetaData.NumIterations = SimStage.NumIterations;
 					SimulationStageMetaData.bWritesParticles = SimStage.bWritesParticles;
 					SimulationStageMetaData.bPartialParticleUpdate = SimStage.bPartialParticleUpdate;
@@ -1850,13 +1880,13 @@ FNiagaraTranslateResults TNiagaraHlslTranslator<GraphBridge>::Translate(const FN
 					SimulationStageMetaData.bGpuIndirectDispatch = false;
 					if ( SimulationStageMetaData.IterationSourceType == ENiagaraIterationSource::DataInterface )
 					{
-						if (SimulationStageMetaData.IterationDataInterface.IsNone())
+						if (SimStage.IterationDataInterface.IsNone())
 						{
 							Error(LOCTEXT("NoIterationDataInterface", "An data interface must be bound to a data interface iteration stage"), nullptr, nullptr);
 							return TranslateResults;
 						}
 
-						if (const FNiagaraVariable* IterationSourceVar = CompileData->EncounteredVariables.FindByPredicate([&](const FNiagaraVariable& VarInfo) { return VarInfo.GetName() == SimulationStageMetaData.IterationDataInterface; }))
+						if (const FNiagaraVariable* IterationSourceVar = CompileData->EncounteredVariables.FindByPredicate([&](const FNiagaraVariable& VarInfo) { return VarInfo.GetName() == SimStage.IterationDataInterface; }))
 						{
 							if ( UNiagaraDataInterface* IterationSourceCDO = CompileDuplicateData->GetDuplicatedDataInterfaceCDOForClass(IterationSourceVar->GetType().GetClass()) )
 							{
@@ -1906,6 +1936,7 @@ FNiagaraTranslateResults TNiagaraHlslTranslator<GraphBridge>::Translate(const FN
 		DefaultStage.NumIterations = 1;
 		DefaultStage.bWritesParticles = true;
 		DefaultStage.bShouldUpdateInitialAttributeValues = TranslationStages[0].ShouldDoSpawnOnlyLogic() || (IsEventSpawnScript() && CompileOptions.AdditionalDefines.Contains(FNiagaraCompileOptions::EventSpawnInitialAttribWritesDefine));
+		DefaultStage.DIOwnerName = GetDIOwnerNameByUsage(DefaultStage.ScriptUsage);
 
 		OutputNodes.Add(SourceGraph->FindEquivalentOutputNode(TranslationStages[0].ScriptUsage, TranslationStages[0].UsageId));
 
@@ -2089,6 +2120,9 @@ FNiagaraTranslateResults TNiagaraHlslTranslator<GraphBridge>::Translate(const FN
 			if (TranslationStages[i].ShouldDoSpawnOnlyLogic())
 				bInitializedDefaults = false;
 
+			TGuardValue<FDIOwnerContext> StageDIOwnerContextScope(CurrentDIOwnerContext, FDIOwnerContext(FGuid(),
+				TranslationStages[i].ScriptUsage, TranslationStages[i].UsageId, TranslationStages[i].DIOwnerName));
+
 			if (UNiagaraScript::IsParticleSpawnScript(TranslationStages[i].ScriptUsage))
 			{
 				AddBodyComment(bHasInterpolatedSpawn ? TEXT("//Begin Interpolated Spawn Script!") : TEXT("//Begin Spawn Script!"));
@@ -2147,6 +2181,9 @@ FNiagaraTranslateResults TNiagaraHlslTranslator<GraphBridge>::Translate(const FN
 	}
 	else if (TranslationStages.Num() == 1)
 	{
+		TGuardValue<FDIOwnerContext> StageDIOwnerContextScope(CurrentDIOwnerContext, FDIOwnerContext(FGuid(),
+			TranslationStages[0].ScriptUsage, TranslationStages[0].UsageId, TranslationStages[0].DIOwnerName));
+
 		CurrentBodyChunkMode = TranslationStages[0].ChunkModeIndex;
 		ActiveStageIdx = 0;
 		check(CompileOptions.TargetUsage == TranslationStages[0].ScriptUsage);
@@ -2211,17 +2248,17 @@ FNiagaraTranslateResults TNiagaraHlslTranslator<GraphBridge>::Translate(const FN
 			{
 				int32 OutputIdx = 0;
 				//ensure the interpolated spawn constants are part of the parameter set.
-				ParameterMapRegisterExternalConstantNamespaceVariable(SYS_PARAM_ENGINE_TIME, nullptr, 0, OutputIdx, nullptr);
-				ParameterMapRegisterExternalConstantNamespaceVariable(SYS_PARAM_ENGINE_WORLD_DELTA_TIME, nullptr, 0, OutputIdx, nullptr);
-				ParameterMapRegisterExternalConstantNamespaceVariable(SYS_PARAM_ENGINE_DELTA_TIME, nullptr, 0, OutputIdx, nullptr);
-				ParameterMapRegisterExternalConstantNamespaceVariable(SYS_PARAM_ENGINE_INV_DELTA_TIME, nullptr, 0, OutputIdx, nullptr);
-				ParameterMapRegisterExternalConstantNamespaceVariable(SYS_PARAM_ENGINE_EXEC_COUNT, nullptr, 0, OutputIdx, nullptr);
-				ParameterMapRegisterExternalConstantNamespaceVariable(SYS_PARAM_EMITTER_SPAWNRATE, nullptr, 0, OutputIdx, nullptr);
+				ParameterMapRegisterExternalConstantNamespaceVariable(SYS_PARAM_ENGINE_TIME, FNiagaraVariableBase(), nullptr, 0, OutputIdx, nullptr);
+				ParameterMapRegisterExternalConstantNamespaceVariable(SYS_PARAM_ENGINE_WORLD_DELTA_TIME, FNiagaraVariableBase(), nullptr, 0, OutputIdx, nullptr);
+				ParameterMapRegisterExternalConstantNamespaceVariable(SYS_PARAM_ENGINE_DELTA_TIME, FNiagaraVariableBase(), nullptr, 0, OutputIdx, nullptr);
+				ParameterMapRegisterExternalConstantNamespaceVariable(SYS_PARAM_ENGINE_INV_DELTA_TIME, FNiagaraVariableBase(), nullptr, 0, OutputIdx, nullptr);
+				ParameterMapRegisterExternalConstantNamespaceVariable(SYS_PARAM_ENGINE_EXEC_COUNT, FNiagaraVariableBase(), nullptr, 0, OutputIdx, nullptr);
+				ParameterMapRegisterExternalConstantNamespaceVariable(SYS_PARAM_EMITTER_SPAWNRATE, FNiagaraVariableBase(), nullptr, 0, OutputIdx, nullptr);
 				if (CompilationTarget != ENiagaraSimTarget::GPUComputeSim)
 				{
-					ParameterMapRegisterExternalConstantNamespaceVariable(SYS_PARAM_EMITTER_SPAWN_INTERVAL, nullptr, 0, OutputIdx, nullptr);
-					ParameterMapRegisterExternalConstantNamespaceVariable(SYS_PARAM_EMITTER_INTERP_SPAWN_START_DT, nullptr, 0, OutputIdx, nullptr);
-					ParameterMapRegisterExternalConstantNamespaceVariable(SYS_PARAM_EMITTER_SPAWN_GROUP, nullptr, 0, OutputIdx, nullptr);
+					ParameterMapRegisterExternalConstantNamespaceVariable(SYS_PARAM_EMITTER_SPAWN_INTERVAL, FNiagaraVariableBase(), nullptr, 0, OutputIdx, nullptr);
+					ParameterMapRegisterExternalConstantNamespaceVariable(SYS_PARAM_EMITTER_INTERP_SPAWN_START_DT, FNiagaraVariableBase(), nullptr, 0, OutputIdx, nullptr);
+					ParameterMapRegisterExternalConstantNamespaceVariable(SYS_PARAM_EMITTER_SPAWN_GROUP, FNiagaraVariableBase(), nullptr, 0, OutputIdx, nullptr);
 				}
 			}
 
@@ -2229,10 +2266,10 @@ FNiagaraTranslateResults TNiagaraHlslTranslator<GraphBridge>::Translate(const FN
 			{
 				int32 OutputIdx = 0;
 				// NOTE(mv): This will explicitly expose Engine.Emitter.TotalSpawnedParticles to the HLSL code regardless of whether it is exposed in a script or not. 
-				ParameterMapRegisterExternalConstantNamespaceVariable(SYS_PARAM_ENGINE_EMITTER_TOTAL_SPAWNED_PARTICLES, nullptr, 0, OutputIdx, nullptr);
-				ParameterMapRegisterExternalConstantNamespaceVariable(SYS_PARAM_EMITTER_RANDOM_SEED, nullptr, 0, OutputIdx, nullptr);
-				ParameterMapRegisterExternalConstantNamespaceVariable(SYS_PARAM_ENGINE_EMITTER_INSTANCE_SEED, nullptr, 0, OutputIdx, nullptr);
-				ParameterMapRegisterExternalConstantNamespaceVariable(SYS_PARAM_ENGINE_SYSTEM_RANDOM_SEED, nullptr, 0, OutputIdx, nullptr);
+				ParameterMapRegisterExternalConstantNamespaceVariable(SYS_PARAM_ENGINE_EMITTER_TOTAL_SPAWNED_PARTICLES, FNiagaraVariableBase(), nullptr, 0, OutputIdx, nullptr);
+				ParameterMapRegisterExternalConstantNamespaceVariable(SYS_PARAM_EMITTER_RANDOM_SEED, FNiagaraVariableBase(), nullptr, 0, OutputIdx, nullptr);
+				ParameterMapRegisterExternalConstantNamespaceVariable(SYS_PARAM_ENGINE_EMITTER_INSTANCE_SEED, FNiagaraVariableBase(), nullptr, 0, OutputIdx, nullptr);
+				ParameterMapRegisterExternalConstantNamespaceVariable(SYS_PARAM_ENGINE_SYSTEM_RANDOM_SEED, FNiagaraVariableBase(), nullptr, 0, OutputIdx, nullptr);
 			}
 		}
 
@@ -2757,26 +2794,30 @@ void TNiagaraHlslTranslator<GraphBridge>::HandleSimStageSetupAndTeardown(int32 I
 		}
 	}
 
-	// Find the data interface in the table. Note that this may not be found because we don't actually call any functions on the data interface yet.
-	int32 DataInterfaceOwnerIndex = INDEX_NONE;
-	for (int32 i = 0; i < CompilationOutput.ScriptData.DataInterfaceInfo.Num(); i++)
+	// Try to find the owner by written data interface parameters.
+	FNiagaraCompileParameterReference IterationSourceReference(IterationSourceVar.GetName(), FGuid(), ExtractFlagsFromParameterName(IterationSourceVar.GetName()));
+	int32 IterationSourceDataInterfaceInfoIndex = CompilationOutput.ScriptData.DataInterfaceInfo.IndexOfByPredicate([IterationSourceReference]
+		(const FNiagaraScriptDataInterfaceCompileInfo& CompileInfo) { return CompileInfo.TargetParameter == IterationSourceReference; });
+
+	if (IterationSourceDataInterfaceInfoIndex == INDEX_NONE)
 	{
-		FNiagaraScriptDataInterfaceCompileInfo& Info = CompilationOutput.ScriptData.DataInterfaceInfo[i];
-		if (IterationDataInterface == Info.Name)
-		{
-			DataInterfaceOwnerIndex = i;
-			break;
-		}
+		// If a written data interface wasn't found, try to find a read parameter.
+		IterationSourceDataInterfaceInfoIndex = CompilationOutput.ScriptData.DataInterfaceInfo.IndexOfByPredicate([IterationSourceReference]
+			(const FNiagaraScriptDataInterfaceCompileInfo& CompileInfo) { return CompileInfo.SourceExternalParameter == IterationSourceReference; });
 	}
 
 	// RIght now we need to know if anyone wrote to the IterationSource this stage. That can be one of two ways:
 	// 1) Someone wrote to StackContext.XXXX
-	// 2) Someone called a function that was marked to write 
+	// 2) Someone called a function that was marked to write
+	bool bFunctionReadFromIterationSource = false;
+	bool bFunctionWroteToIterationSource = false;
 	const int32 SourceSimStage = TranslationStage.SimulationStageIndex;
-	ensure(CompilationOutput.ScriptData.SimulationStageMetaData.IsValidIndex(SourceSimStage));
-	bool bWroteToIterationSource = CompilationOutput.ScriptData.SimulationStageMetaData[SourceSimStage].OutputDestinations.Contains(TranslationStage.GetIterationDataInterface());
-	if (WriteVars.Num() > 0)
-		bWroteToIterationSource = true;
+	if (IterationSourceDataInterfaceInfoIndex != INDEX_NONE)
+	{
+		FName IterationSourceDataInterfaceName = CompilationOutput.ScriptData.DataInterfaceInfo[IterationSourceDataInterfaceInfoIndex].Name;
+		bFunctionReadFromIterationSource = CompilationOutput.ScriptData.SimulationStageMetaData[SourceSimStage].InputDataInterfaces.Contains(IterationSourceDataInterfaceName);
+		bFunctionWroteToIterationSource = CompilationOutput.ScriptData.SimulationStageMetaData[SourceSimStage].OutputDestinations.Contains(IterationSourceDataInterfaceName);
+	}
 
 	// Now decide if we need to put in the pre/post
 	if (CDO && CDO->CanExecuteOnTarget(ENiagaraSimTarget::GPUComputeSim))
@@ -2786,8 +2827,14 @@ void TNiagaraHlslTranslator<GraphBridge>::HandleSimStageSetupAndTeardown(int32 I
 		bool bNeedsAttributeWrite = false;
 		bool bNeedsAttributeRead = false;
 
+		// If a function read from or wrote to the iteration source data interface, then it needs an owner.
+		if (bFunctionReadFromIterationSource || bFunctionWroteToIterationSource)
+		{
+			bNeedsDIOwner = true;
+		}
+
 		// Put in the general pre/post if we wrote to the IterationSoruce at all
-		if (CDO->SupportsSetupAndTeardownHLSL() && bWroteToIterationSource)
+		if (CDO->SupportsSetupAndTeardownHLSL() && (bFunctionWroteToIterationSource || WriteVars.Num() > 0))
 		{
 			bNeedsDIOwner = true;
 			bNeedsSetupAndTeardown = true;
@@ -2809,18 +2856,27 @@ void TNiagaraHlslTranslator<GraphBridge>::HandleSimStageSetupAndTeardown(int32 I
 		}
 		
 		// If it wasn't previously added, let's go ahead and do so. Maybe they are solely using the StackContext namespace.
-		if (DataInterfaceOwnerIndex == INDEX_NONE && bNeedsDIOwner)
+		if (IterationSourceDataInterfaceInfoIndex == INDEX_NONE && bNeedsDIOwner)
 		{
-			DataInterfaceOwnerIndex = RegisterDataInterface(IterationSourceVar, CDO, true, true);
+			IterationSourceDataInterfaceInfoIndex = RegisterDataInterfaceForParameterRead(IterationSourceVar, INDEX_NONE);
 		}
 
 		// If we haven't created it by now, bail out. 
-		if (DataInterfaceOwnerIndex == INDEX_NONE && bNeedsDIOwner)
+		if (bNeedsDIOwner && IterationSourceDataInterfaceInfoIndex == INDEX_NONE)
 		{
 			Error(FText::Format(LOCTEXT("CannotRegisterDataInterface", "Variable {0}'s cannot register as a data interface!"), FText::FromName(IterationDataInterface)), nullptr, nullptr);
 			return;
 		}
 
+		// If we need to bind the iteration source data interface, specify the unique owner name here on the stage metadata for binding at runtime.
+		// We use the unique owner name here rather than the iteration source parameter name to avoid issues with emitter aliasing or other ambiguous
+		// resolving issues.
+		if (bNeedsDIOwner)
+		{
+			CompilationOutput.ScriptData.SimulationStageMetaData[TranslationStage.SimulationStageIndex].IterationDataInterface =
+				CompilationOutput.ScriptData.DataInterfaceInfo[IterationSourceDataInterfaceInfoIndex].Name;
+		}
+
 		// It is an invalid state to use the IterationSource and StackContext namespace without implementing SupportsIterationSourceNamespaceAttributesHLSL
 		if (ReadVars.Num() > 0 && !bNeedsAttributeRead)
 		{
@@ -2842,7 +2898,8 @@ void TNiagaraHlslTranslator<GraphBridge>::HandleSimStageSetupAndTeardown(int32 I
 		// Convert to a FNiagaraDataInterfaceGPUParamInfo to keep the API simple and consistent
 		TArray<FNiagaraFunctionSignature> GeneratedFunctionSignatures;
 		FNiagaraDataInterfaceGPUParamInfo DIInstanceInfo;
-		ConvertCompileInfoToParamInfo(CompilationOutput.ScriptData.DataInterfaceInfo[DataInterfaceOwnerIndex], DIInstanceInfo, GeneratedFunctionSignatures);
+		const FNiagaraScriptDataInterfaceCompileInfo& Info = CompilationOutput.ScriptData.DataInterfaceInfo[IterationSourceDataInterfaceInfoIndex];
+		ConvertCompileInfoToParamInfo(Info, DIInstanceInfo, GeneratedFunctionSignatures);
 
 		// This next part might be a big confusing, but because DataInterfaces are in non-editor code, it makes it impossible for them to return graphs or other
 		// structures. We want them to feel free to invoke their own fucntions and not have to do a lot of extra wranging, so we treat them like a custom hlsl node.
@@ -2860,7 +2917,7 @@ void TNiagaraHlslTranslator<GraphBridge>::HandleSimStageSetupAndTeardown(int32 I
 
 		TArray<int32> Inputs;
 		Inputs.Emplace(InWhichStage); // The parameter map input index above
-		Inputs.Emplace(DataInterfaceOwnerIndex); // the data interface index above
+		Inputs.Emplace(IterationSourceDataInterfaceInfoIndex); // the data interface index above
 		FString SetupBody;
 		FString TeardownBody;
 
@@ -3366,10 +3423,7 @@ void FNiagaraHlslTranslator::DefineDataSetWriteFunction(FString &HlslOutputStrin
 
 void FNiagaraHlslTranslator::ConvertCompileInfoToParamInfo(const FNiagaraScriptDataInterfaceCompileInfo& Info, FNiagaraDataInterfaceGPUParamInfo& DIInstanceInfo, TArray<FNiagaraFunctionSignature>& GeneratedFunctionSignatures)
 {
-	FString OwnerIDString = Info.Name.ToString();
-	FString SanitizedOwnerIDString = BuildSanitizedSymbolString(OwnerIDString.Replace(TEXT("."), TEXT("_")));
-
-	DIInstanceInfo.DataInterfaceHLSLSymbol = SanitizedOwnerIDString;
+	DIInstanceInfo.DataInterfaceHLSLSymbol = Info.Name.ToString();
 	DIInstanceInfo.DIClassName = Info.Type.GetClass()->GetName();
 
 	// Build a list of function instances that will be generated for this DI.
@@ -5247,7 +5301,7 @@ int32 TNiagaraHlslTranslator<GraphBridge>::GetRapidIterationParameter(const FNia
 
 		// go ahead and make it into a constant variable..
 		int32 OutputChunkId = INDEX_NONE;
-		if (ParameterMapRegisterExternalConstantNamespaceVariable(Parameter, nullptr, INDEX_NONE, OutputChunkId, nullptr))
+		if (ParameterMapRegisterExternalConstantNamespaceVariable(Parameter, FNiagaraVariableBase(), nullptr, INDEX_NONE, OutputChunkId, nullptr))
 		{
 			return OutputChunkId;
 		}
@@ -5373,7 +5427,7 @@ int32 TNiagaraHlslTranslator<GraphBridge>::GetParameter(const FNiagaraVariable&
 	FNiagaraVariable OutputVariable = Parameter;
 	if (FNiagaraParameterUtilities::IsInNamespace(OutputVariable, PARAM_MAP_ATTRIBUTE_STR) || FNiagaraParameterUtilities::IsExternalConstantNamespace(OutputVariable, CompileOptions.TargetUsage, CompileOptions.GetTargetUsageBitmask()))
 	{
-		if (!ParameterMapRegisterExternalConstantNamespaceVariable(OutputVariable, nullptr, 0, OutputChunkIdx, nullptr))
+		if (!ParameterMapRegisterExternalConstantNamespaceVariable(OutputVariable, FNiagaraVariableBase(), nullptr, 0, OutputChunkIdx, nullptr))
 		{
 			OutputChunkIdx = INDEX_NONE;
 		}
@@ -5381,7 +5435,7 @@ int32 TNiagaraHlslTranslator<GraphBridge>::GetParameter(const FNiagaraVariable&
 	else
 	{
 		OutputVariable = FNiagaraParameterUtilities::MoveToExternalConstantNamespaceVariable(OutputVariable, CompileOptions.TargetUsage);
-		if (!ParameterMapRegisterExternalConstantNamespaceVariable(OutputVariable, nullptr, 0, OutputChunkIdx, nullptr))
+		if (!ParameterMapRegisterExternalConstantNamespaceVariable(OutputVariable, FNiagaraVariableBase(), nullptr, 0, OutputChunkIdx, nullptr))
 		{
 			OutputChunkIdx = INDEX_NONE;
 		}
@@ -6071,6 +6125,7 @@ void TNiagaraHlslTranslator<GraphBridge>::ParameterMapSet(const FParamMapSetNode
 				continue;
 			}
 
+			FNiagaraVariableBase OriginalVar = Var;
 			Var = ActiveHistoryForFunctionCalls.ResolveAliases(Var);
 			
 			const FNiagaraKnownConstantInfo ConstantInfo = FNiagaraConstants::GetKnownConstantInfo(Var.GetName(), false);
@@ -6157,14 +6212,6 @@ void TNiagaraHlslTranslator<GraphBridge>::ParameterMapSet(const FParamMapSetNode
 
 			if (Var.IsDataInterface())
 			{
-				if (CompileOptions.TargetUsage == ENiagaraScriptUsage::ParticleSpawnScriptInterpolated && TranslationStages[ActiveStageIdx].ScriptUsage == ENiagaraScriptUsage::ParticleUpdateScript)
-				{
-					// We don't want to add writes for particle update data interface parameters in both interpolated spawn and update, so skip them when processing the update stage of the 
-					// interpolated spawn script.  We don't skip the writes when compiling the particle update script because it's not recompiled when the interpolated spawn flag is changed
-					// and this would result in missing data interfaces if interpolated spawn was turned off.
-					continue;
-				}
-
 				bool bAllowDataInterfaces = true;
 				if (ParamMapHistoryIdx < ParamMapHistories.Num() && ParamMapHistories[ParamMapHistoryIdx].IsPrimaryDataSetOutput(Var, CompileOptions.TargetUsage, bAllowDataInterfaces))
 				{
@@ -6174,37 +6221,7 @@ void TNiagaraHlslTranslator<GraphBridge>::ParameterMapSet(const FParamMapSetNode
 						continue;
 					}
 
-					FName UsageName;
-					if (FNiagaraParameterUtilities::IsAliasedEmitterParameter(Var))
-					{
-						FNiagaraVariable AliasedVar = ActiveHistoryForFunctionCalls.ResolveAliases(Var);
-						UsageName = AliasedVar.GetName();
-					}
-					else
-					{
-						UsageName = Var.GetName();
-					}
-
-					FNiagaraScriptDataInterfaceCompileInfo& Info = CompilationOutput.ScriptData.DataInterfaceInfo[Input];
-					if (Info.RegisteredParameterMapWrite == NAME_None)
-					{
-						Info.RegisteredParameterMapWrite = UsageName;
-					}
-					else
-					{
-						// This is a special case where a single data interface is written to multiple parameters.  In order to handle this correctly the data interface
-						// needs it's own registration slot so that it ends up in the parameter store under the correct name so that it can be looked up correctly in various places
-						// and can be bound properly to the renderer parameter stores.
-						UNiagaraDataInterface* CDO = CompileDuplicateData->GetDuplicatedDataInterfaceCDOForClass(const_cast<UClass*>(Var.GetType().GetClass()));
-						if (CDO)
-						{
-							FString PlaceholderName = FString::Printf(TEXT("Placeholder.Placeholder%i"), CompilationOutput.ScriptData.DataInterfaceInfo.Num());
-							FNiagaraVariable PlaceholderVar(Var.GetType(), *PlaceholderName);
-							int32 PlaceholderIndex = RegisterDataInterface(PlaceholderVar, CDO, true, false);
-							CompilationOutput.ScriptData.DataInterfaceInfo[PlaceholderIndex].RegisteredParameterMapRead = CompilationOutput.ScriptData.DataInterfaceInfo[Input].RegisteredParameterMapWrite;
-							CompilationOutput.ScriptData.DataInterfaceInfo[PlaceholderIndex].RegisteredParameterMapWrite = UsageName;
-						}
-					}
+					RegisterDataInterfaceForParameterWrite(OriginalVar, Input);
 				}
 			}
 			else if ( Var.IsUObject() )
@@ -6421,14 +6438,14 @@ bool TNiagaraHlslTranslator<GraphBridge>::HandleBoundConstantVariableToDataSetRe
 
 		if (bIsEmitterLocalSpaceCompileOptionSet == false)
 		{
-			return ParameterMapRegisterExternalConstantNamespaceVariable(SYS_PARAM_ENGINE_POSITION, InNode, InParamMapHistoryIdx, Output, InDefaultPin);
+			return ParameterMapRegisterExternalConstantNamespaceVariable(SYS_PARAM_ENGINE_POSITION, FNiagaraVariableBase(), InNode, InParamMapHistoryIdx, Output, InDefaultPin);
 		}
 	}
 	return false;
 }
 
 template<typename GraphBridge>
-bool TNiagaraHlslTranslator<GraphBridge>::ParameterMapRegisterExternalConstantNamespaceVariable(FNiagaraVariable InVariable, const FNode* InNodeForErrorReporting, int32 InParamMapHistoryIdx, int32& Output, const FPin* InDefaultPin)
+bool TNiagaraHlslTranslator<GraphBridge>::ParameterMapRegisterExternalConstantNamespaceVariable(FNiagaraVariable InVariable, const FNiagaraVariableBase& InUnresolvedVariable, const FNode* InNodeForErrorReporting, int32 InParamMapHistoryIdx, int32& Output, const FPin* InDefaultPin)
 {
 	Output = INDEX_NONE;
 
@@ -6481,31 +6498,8 @@ bool TNiagaraHlslTranslator<GraphBridge>::ParameterMapRegisterExternalConstantNa
 			{
 				if (bIsDataInterface)
 				{
-					UNiagaraDataInterface* DataInterface = nullptr;
-					if (GraphBridge::IsCollectionValid(Collection))
-					{
-						DataInterface = GraphBridge::GetCollectionDataInterface(Collection, InVariable);
-						if (DataInterface == nullptr)
-						{
-							Error(FText::Format(LOCTEXT("ParameterCollectionDataInterfaceNotFoundErrorFormat", "Data interface named {0} of type {1} was not found in Parameter Collection {2}"),
-								FText::FromName(InVariable.GetName()), InVariable.GetType().GetNameText(), FText::FromString(GraphBridge::GetCollectionFullName(Collection))), InNodeForErrorReporting, InDefaultPin);
-							return false;
-						}
-					}
-					else
-					{
-						DataInterface = CompileDuplicateData->GetDuplicatedDataInterfaceCDOForClass(const_cast<UClass*>(InVariable.GetType().GetClass()));
-						if (DataInterface == nullptr)
-						{
-							Error(FText::Format(LOCTEXT("GetDuplicatedDataInterfaceCDOForClassFailed", "GetDuplicatedDataInterfaceCDOForClass failed for Variable({0}) Class({1})"), FText::FromName(InVariable.GetName()), InVariable.GetType().GetNameText()), InNodeForErrorReporting, InDefaultPin);
-							return false;
-						}
-					}
-					if (ensure(DataInterface))
-					{
-						Output = RegisterDataInterface(InVariable, DataInterface, true, true);
-						return true;
-					}
+					Output = RegisterDataInterfaceForParameterRead(InUnresolvedVariable.IsValid() ? InUnresolvedVariable : InVariable, INDEX_NONE);
+					return true;
 				}
 				if (bIsUObject)
 				{
@@ -6990,6 +6984,8 @@ void TNiagaraHlslTranslator<GraphBridge>::Emitter(const FEmitterNode* EmitterNod
 
 	{
 		FScopedBuilderEmitter<GraphBridge> EmitterBuidlerScope(ActiveHistoryForFunctionCalls, EmitterNode);
+		TGuardValue<FDIOwnerContext> EmitterDIOwnerContextScope(CurrentDIOwnerContext, FDIOwnerContext(GraphBridge::GetEmitterHandleId(EmitterNode),
+			GraphBridge::GetEmitterUsage(EmitterNode), FGuid(), *GraphBridge::GetEmitterUniqueName(EmitterNode)));
 
 		// Clear out the parameter map writes to emitter module parameters as they should not be shared across emitters.
 		if (ParamMapHistoryIdx != -1 && ParamMapHistoryIdx < ParamMapHistories.Num())
@@ -7216,7 +7212,7 @@ void TNiagaraHlslTranslator<GraphBridge>::HandleParameterRead(int32 ParamMapHist
 			OutputChunkId = MakeStaticVariableDirect(DefaultPin);
 			return;
 		}
-		else if (ParameterMapRegisterExternalConstantNamespaceVariable(Var, ErrorNode, ParamMapHistoryIdx, OutputChunkId, DefaultPin))
+		else if (ParameterMapRegisterExternalConstantNamespaceVariable(Var, InVar, ErrorNode, ParamMapHistoryIdx, OutputChunkId, DefaultPin))
 		{
 			return;
 		}			
@@ -7323,7 +7319,7 @@ void TNiagaraHlslTranslator<GraphBridge>::HandleParameterRead(int32 ParamMapHist
 				}
 
 				// If it isn't found yet.. go ahead and make it into a constant variable..
-				if (LastSetChunkIdx == INDEX_NONE && ParameterMapRegisterExternalConstantNamespaceVariable(Var, ErrorNode, ParamMapHistoryIdx, OutputChunkId, InputPin))
+				if (LastSetChunkIdx == INDEX_NONE && ParameterMapRegisterExternalConstantNamespaceVariable(Var, InVar, ErrorNode, ParamMapHistoryIdx, OutputChunkId, InputPin))
 				{
 					LastSetChunkIdx = OutputChunkId;
 					if (VarIdx != INDEX_NONE && VarIdx < ParamMapSetVariablesToChunks[ParamMapHistoryIdx].Num())
@@ -7583,7 +7579,7 @@ void TNiagaraHlslTranslator<GraphBridge>::HandleParameterRead(int32 ParamMapHist
 				if (Var.IsDataInterface() || Var.IsUObject())
 				{
 					// If this variable is a data interface and it's in the parameter map, but hasn't been set yet, then is is an external data interface so try to register it.
-					if (ParameterMapRegisterExternalConstantNamespaceVariable(Var, ErrorNode, ParamMapHistoryIdx, OutputChunkId, DefaultPin))
+					if (ParameterMapRegisterExternalConstantNamespaceVariable(Var, InVar, ErrorNode, ParamMapHistoryIdx, OutputChunkId, DefaultPin))
 					{
 						return;
 					}
@@ -7600,12 +7596,7 @@ void TNiagaraHlslTranslator<GraphBridge>::HandleParameterRead(int32 ParamMapHist
 		// for variables encountered without ever being set. We do this by creating an instance of the CDO.
 		if (UNiagaraScript::IsStandaloneScript(CompileOptions.TargetUsage) && LastSetChunkIdx == INDEX_NONE)
 		{
-			UNiagaraDataInterface* DataInterface = CompileDuplicateData->GetDuplicatedDataInterfaceCDOForClass(const_cast<UClass*>(Var.GetType().GetClass()));
-			check(DataInterface != nullptr);
-			if (DataInterface)
-			{
-				LastSetChunkIdx = RegisterDataInterface(Var, DataInterface, true, false);
-			}
+			LastSetChunkIdx = RegisterDataInterfaceForParameterRead(Var, INDEX_NONE);
 		}
 
 		OutputChunkId = LastSetChunkIdx;
@@ -7855,82 +7846,6 @@ int32 TNiagaraHlslTranslator<GraphBridge>::RegisterUObjectPath(const FNiagaraVar
 	return ObjectIndex;
 }
 
-template<typename GraphBridge>
-int32 TNiagaraHlslTranslator<GraphBridge>::RegisterDataInterface(const FNiagaraVariable& Var, const UNiagaraDataInterface* DataInterface, bool bPlaceholder, bool bAddParameterMapRead)
-{
-	//FString PathName = DataInterface ? *DataInterface->GetPathName() : TEXT("XX");
-	//UE_LOG(LogNiagaraEditor, Log, TEXT("RegisterDataInterface %s %s %s '%s' %s"), *Var.GetName().ToString(), *Var.GetType().GetName(), *PathName,
-	//	bPlaceholder ? TEXT("bPlaceholder=true") : TEXT("bPlaceholder=false"));
-
-	int32 FuncParam;
-	if (GetFunctionParameter(Var, FuncParam))
-	{
-		if (FuncParam != INDEX_NONE)
-		{
-			//This data interface param has been overridden by the function call so use that index.	
-			return FuncParam;
-		}
-	}
-
-	//If we get here then this is a new data interface.
-	const FString* EmitterAlias = ActiveHistoryForFunctionCalls.GetEmitterAlias();
-	FName DataInterfaceName = GetDataInterfaceName(Var.GetName(), EmitterAlias != nullptr ? *EmitterAlias : FString(), bAddParameterMapRead);
-
-	if (DataInterface && DataInterface->NeedsGPUContextInit() && CompileOptions.IsGpuScript())
-	{
-		CompilationOutput.ScriptData.bNeedsGPUContextInit = true;
-	}
-
-	FNiagaraVariable ReadVariable;
-	if (bAddParameterMapRead)
-	{
-		ReadVariable = FNiagaraParameterUtilities::IsAliasedEmitterParameter(Var.GetName().ToString())
-			? ActiveHistoryForFunctionCalls.ResolveAliases(Var) 
-			: Var;
-	}
-
-	int32 Idx = CompilationOutput.ScriptData.DataInterfaceInfo.IndexOfByPredicate([&](const FNiagaraScriptDataInterfaceCompileInfo& OtherInfo)
-	{
-		return OtherInfo.Name == DataInterfaceName || (ReadVariable.IsValid() && OtherInfo.RegisteredParameterMapWrite == ReadVariable.GetName());
-	});
-
-	if (Idx == INDEX_NONE)
-	{
-		Idx = CompilationOutput.ScriptData.DataInterfaceInfo.AddDefaulted();
-		CompilationOutput.ScriptData.DataInterfaceInfo[Idx].Name = DataInterfaceName;
-		CompilationOutput.ScriptData.DataInterfaceInfo[Idx].Type = Var.GetType();
-		CompilationOutput.ScriptData.DataInterfaceInfo[Idx].bIsPlaceholder = bPlaceholder;
-
-		FString SourceEmitterName = FString();
-		if (CompileData->GetUniqueEmitterName().IsEmpty() == false)
-		{
-			SourceEmitterName = CompileData->GetUniqueEmitterName();
-		}
-		else if (EmitterAlias != nullptr)
-		{
-			SourceEmitterName = *EmitterAlias;
-		}
-		CompilationOutput.ScriptData.DataInterfaceInfo[Idx].SourceEmitterName = SourceEmitterName;
-
-		//Interface requires per instance data so add a user pointer table entry.
-		if (DataInterface != nullptr && DataInterface->PerInstanceDataSize() > 0)
-		{
-			CompilationOutput.ScriptData.DataInterfaceInfo[Idx].UserPtrIdx = CompilationOutput.ScriptData.NumUserPtrs++;
-		}
-
-		if (bAddParameterMapRead)
-		{
-			CompilationOutput.ScriptData.DataInterfaceInfo[Idx].RegisteredParameterMapRead = ReadVariable.GetName();
-		}
-	}
-	else
-	{
-		check(CompilationOutput.ScriptData.DataInterfaceInfo[Idx].Type == Var.GetType());
-	}
-
-	return Idx;
-}
-
 template<typename GraphBridge>
 void TNiagaraHlslTranslator<GraphBridge>::Operation(const FOpNode* Operation, TArray<int32>& Inputs, TArray<int32>& Outputs)
 {
@@ -8809,7 +8724,7 @@ void TNiagaraHlslTranslator<GraphBridge>::ProcessCustomHlsl(const FString& InCus
 
 			if (bPermuteSignatureByDataInterface)
 			{
-				OutSignature.Name = FName(OutSignature.Name.ToString() + BuildSanitizedSymbolName(Info.Name, true));
+				OutSignature.Name = FName(OutSignature.Name.ToString() + Info.Name.ToString());
 			}
 
 			SigInputs.Add(Input);
@@ -8883,7 +8798,7 @@ void TNiagaraHlslTranslator<GraphBridge>::ProcessCustomHlsl(const FString& InCus
 				FNiagaraParameterUtilities::IsExternalConstantNamespace(VarActual, CompileOptions.TargetUsage, CompileOptions.GetTargetUsageBitmask()))
 			{
 				int32 TempOutput;
-				if (ParameterMapRegisterExternalConstantNamespaceVariable(VarActual, InNodeForErrorReporting, ParamMapHistoryIdx, TempOutput, nullptr))
+				if (ParameterMapRegisterExternalConstantNamespaceVariable(VarActual, VarAliased, InNodeForErrorReporting, ParamMapHistoryIdx, TempOutput, nullptr))
 				{
 					continue;
 				}
@@ -9176,7 +9091,7 @@ void TNiagaraHlslTranslator<GraphBridge>::RegisterFunctionCall(ENiagaraScriptUsa
 								// for aliased values we check to see if they're used in the current context by resolving the alias and checking against the current resolved variable
 								// name since aliased values can only be resolved for reading in the correct context.
 								bool bIsValidForCurrentCallingContext = (bIsInTopLevelFunction && bIsAliased == false) || (bIsAliased && ActiveHistoryForFunctionCalls.ResolveAliases(AliasedVar).GetName() == Var.GetName());
-								if (bIsValidForCurrentCallingContext && !Var.GetType().IsStatic() )
+								if (bIsValidForCurrentCallingContext && !Var.GetType().IsStatic() && Var.IsDataInterface() == false)
 								{
 									int32 LastSetChunkIdx = ParamMapSetVariablesToChunks[ActiveStageIdx][VarIdx];
 									if (LastSetChunkIdx == INDEX_NONE)
@@ -9675,24 +9590,6 @@ FString FNiagaraHlslTranslator::GenerateFunctionParametersHlsl(const FNiagaraFun
 	return StringBuilder.ToString();
 }
 
-FName FNiagaraHlslTranslator::GetDataInterfaceName(FName BaseName, const FString& UniqueEmitterName, bool bIsParameterMapDataInterface)
-{
-	if (UniqueEmitterName.IsEmpty() == false)
-	{
-		if (FNiagaraParameterUtilities::IsAliasedEmitterParameter(BaseName.ToString()))
-		{
-			return FNiagaraParameterUtilities::ResolveEmitterAlias(BaseName, UniqueEmitterName);
-		}
-		else if(bIsParameterMapDataInterface == false)
-		{
-			// Don't mangle the parameter map reads for emitter scripts because they are from the system or user parameter stores and
-			// they won't bind correctly.
-			return *(UniqueEmitterName + TEXT(".") + BaseName.ToString());
-		}
-	}
-	return BaseName;
-}
-
 FString FNiagaraHlslTranslator::GetFunctionIncludeStatement(const FNiagaraCustomHlslInclude& Include) const
 {
 	TStringBuilder<128> IncludeStatement;
@@ -9923,6 +9820,247 @@ FString FNiagaraHlslTranslator::GenerateAssignment(const FNiagaraTypeDefinition&
 	return DestinationDefinition + " = " + SourceDefinition;
 }
 
+int32 FNiagaraHlslTranslator::RegisterDataInterfaceForInput(const FNiagaraVariableBase& InputVariable, const FNiagaraGraphInputDataInterfaceKey& GraphInputDataInterfaceKey, const UNiagaraDataInterface* InputDataInterface)
+{
+	int32 FuncParam;
+	if (GetFunctionParameter(InputVariable, FuncParam))
+	{
+		if (FuncParam != INDEX_NONE)
+		{
+			//This data interface param has been overridden by the function call so use that index.	
+			return FuncParam;
+		}
+	}
+
+	if (ensureMsgf(GraphInputDataInterfaceKey.IsValid(), TEXT("GraphInputDataInterfaceKey is valid for input variable %s"), *InputVariable.GetName().ToString()) == false)
+	{
+		return INDEX_NONE;
+	}
+
+	int32 ExistingIndex = CompilationOutput.ScriptData.DataInterfaceInfo.IndexOfByPredicate([&GraphInputDataInterfaceKey]
+	(const FNiagaraScriptDataInterfaceCompileInfo& OtherInfo) { return OtherInfo.SourceGraphInputDataInterfaceKey == GraphInputDataInterfaceKey; });
+	if (ExistingIndex != INDEX_NONE)
+	{
+		return ExistingIndex;
+	}
+
+	const UNiagaraDataInterface* InputDataInterfaceOrCDO = InputDataInterface != nullptr
+		? InputDataInterface
+		: GetDataInterfaceCDO(InputVariable.GetType().GetClass());
+
+	FString InternalCompileName = FString::Printf(TEXT("INTERNAL.%s"), *InputVariable.GetName().ToString());
+	int32 AddedIndex = AddDataInterfaceInfo(InputVariable.GetType(), *InternalCompileName, *InputDataInterfaceOrCDO);
+	FNiagaraScriptDataInterfaceCompileInfo& AddedInfo = CompilationOutput.ScriptData.DataInterfaceInfo[AddedIndex];
+	AddedInfo.SourceGraphInputDataInterfaceKey = GraphInputDataInterfaceKey;
+
+	return AddedIndex;
+}
+
+ENiagaraCompileParameterReferenceFlags FNiagaraHlslTranslator::ExtractFlagsFromParameterName(FName ParameterFullName)
+{
+	FNiagaraParameterHandle ReadHandle(ParameterFullName);
+	ENiagaraCompileParameterReferenceFlags Flags = ENiagaraCompileParameterReferenceFlags::None;
+	if (ReadHandle.IsUserHandle())
+	{
+		Flags = ENiagaraCompileParameterReferenceFlags::User;
+	}
+	else if (ReadHandle.IsSystemHandle())
+	{
+		Flags = ENiagaraCompileParameterReferenceFlags::System;
+	}
+	else if (ReadHandle.IsEmitterHandle())
+	{
+		Flags = ENiagaraCompileParameterReferenceFlags::Emitter;
+	}
+	else if (ReadHandle.IsParticleAttributeHandle())
+	{
+		Flags = ENiagaraCompileParameterReferenceFlags::Particle;
+	}
+	else if (ReadHandle.IsTransientHandle())
+	{
+		Flags = ENiagaraCompileParameterReferenceFlags::Transient;
+	}
+	else if (ReadHandle.IsLocalHandle())
+	{
+		Flags = ENiagaraCompileParameterReferenceFlags::Local;
+	}
+	else if (ReadHandle.IsParameterCollectionHandle())
+	{
+		Flags = ENiagaraCompileParameterReferenceFlags::NPC;
+	}
+	return Flags;
+}
+
+int32 FNiagaraHlslTranslator::RegisterDataInterfaceForParameterRead(const FNiagaraVariableBase& ReadParameter, int32 ConnectedInfoIndex)
+{
+	FNiagaraCompileParameterReference ReadParameterReference(ReadParameter.GetName(), FGuid(), ExtractFlagsFromParameterName(ReadParameter.GetName()));
+	if (ReadParameterReference.HasFlags(ENiagaraCompileParameterReferenceFlags::Emitter))
+	{
+		ReadParameterReference.EmitterHandleId = CurrentDIOwnerContext.EmitterHandleId;
+	}
+
+	// If there is already a data interface that's a registered target for this parameter, we can use that.
+	int32 ExistingWriteIndex = CompilationOutput.ScriptData.DataInterfaceInfo.IndexOfByPredicate([&](const FNiagaraScriptDataInterfaceCompileInfo& OtherInfo)
+		{ return OtherInfo.TargetParameter == ReadParameterReference; });
+	if (ExistingWriteIndex != INDEX_NONE)
+	{
+		return ExistingWriteIndex;
+	}
+
+	// If there was no previous write, then this is an external read, so check if we have a previous external read.
+	int32 ExistingReadIndex = CompilationOutput.ScriptData.DataInterfaceInfo.IndexOfByPredicate([&](const FNiagaraScriptDataInterfaceCompileInfo& OtherInfo)
+		{ return OtherInfo.SourceExternalParameter == ReadParameterReference; });
+	if (ExistingReadIndex != INDEX_NONE)
+	{
+		return ExistingReadIndex;
+	}
+
+	// If we have no writes, and no previous external reads, check if we can use a connected default.
+	FNiagaraScriptDataInterfaceCompileInfo* ConnectedInfo = ConnectedInfoIndex != INDEX_NONE && ConnectedInfoIndex < CompilationOutput.ScriptData.DataInterfaceInfo.Num()
+		? &CompilationOutput.ScriptData.DataInterfaceInfo[ConnectedInfoIndex] : nullptr;
+	if (ConnectedInfo != nullptr)
+	{
+		if (ConnectedInfo->SourceExternalParameter.IsValid() == false)
+		{
+			ConnectedInfo->SourceExternalParameter = ReadParameterReference;
+			return ConnectedInfoIndex;
+		}
+		else
+		{
+			// If the connected default was already used with another external read, we register a duplicate placeholder so we can register this read.
+			int32 PlaceholderIndex = AddPlaceholderDuplicateDataInterface(ConnectedInfoIndex);
+			FNiagaraScriptDataInterfaceCompileInfo& PlaceholderDuplicate = CompilationOutput.ScriptData.DataInterfaceInfo[PlaceholderIndex];
+			PlaceholderDuplicate.SourceExternalParameter = ReadParameterReference;
+			PlaceholderDuplicate.SourceScriptDIIndex = ConnectedInfoIndex;
+		}
+	}
+
+
+	// No registered write, existing external read, or default, so we register a new external read.
+	UNiagaraDataInterface* CDO = GetDataInterfaceCDO(ReadParameter.GetType().GetClass());
+	int32 AddedIndex = AddDataInterfaceInfo(ReadParameter.GetType(), ReadParameter.GetName(), *CDO);
+	FNiagaraScriptDataInterfaceCompileInfo& AddedInfo = CompilationOutput.ScriptData.DataInterfaceInfo[AddedIndex];
+	AddedInfo.SourceExternalParameter = ReadParameterReference;
+	return AddedIndex;
+}
+
+void FNiagaraHlslTranslator::RegisterDataInterfaceForParameterWrite(const FNiagaraVariableBase& WriteParameter, int32 ConnectedInfoIndex)
+{
+	FNiagaraCompileParameterReference WriteParameterReference(WriteParameter.GetName(), FGuid(), ExtractFlagsFromParameterName(WriteParameter.GetName()));
+	if (WriteParameterReference.HasFlags(ENiagaraCompileParameterReferenceFlags::Emitter))
+	{
+		WriteParameterReference.EmitterHandleId = CurrentDIOwnerContext.EmitterHandleId;
+	}
+
+	// First check if this data interfaces was already written to, if it was already written to, validate that it's not being written to
+	// a second time with a different data interface.
+	int32 ExistingWriteIndex = CompilationOutput.ScriptData.DataInterfaceInfo.IndexOfByPredicate([&](const FNiagaraScriptDataInterfaceCompileInfo& OtherInfo)
+		{ return OtherInfo.TargetParameter == WriteParameterReference; });
+	if (ExistingWriteIndex != INDEX_NONE)
+	{
+		if (ConnectedInfoIndex != INDEX_NONE &&
+			ExistingWriteIndex != ConnectedInfoIndex &&
+			CompilationOutput.ScriptData.DataInterfaceInfo[ExistingWriteIndex].SourceScriptDIIndex != ConnectedInfoIndex)
+		{
+			Warning(FText::Format(LOCTEXT("MultipleDIWriteWarningFormat", "Data interface parameter {0} written to multiple times, only using the first write."),
+				FText::FromName(WriteParameterReference.ParameterFullName)));
+		}
+		return;
+	}
+
+	// Get the data interface info currently connected to this parameter if available.
+	FNiagaraScriptDataInterfaceCompileInfo* ConnectedInfo = ConnectedInfoIndex != INDEX_NONE && ConnectedInfoIndex < CompilationOutput.ScriptData.DataInterfaceInfo.Num()
+		? &CompilationOutput.ScriptData.DataInterfaceInfo[ConnectedInfoIndex] : nullptr;
+	if (ConnectedInfo != nullptr)
+	{
+		if (ConnectedInfo->TargetParameter.IsValid() == false)
+		{
+			// The parameter is connected to a registered data interface with no target, so we specify the write here.
+			ConnectedInfo->TargetParameter = WriteParameterReference;
+		}
+		else
+		{
+			// If the connected default was already written to another target, we register a duplicate placeholder so we can register this target write.
+			int32 PlaceholderIndex = AddPlaceholderDuplicateDataInterface(ConnectedInfoIndex);
+			FNiagaraScriptDataInterfaceCompileInfo& PlaceholderDuplicate = CompilationOutput.ScriptData.DataInterfaceInfo[PlaceholderIndex];
+			PlaceholderDuplicate.TargetParameter = WriteParameterReference;
+			PlaceholderDuplicate.SourceScriptDIIndex = ConnectedInfoIndex;
+		}
+		return;
+	}
+
+	// If the write was not connected and it's not already registered then we need to use a placeholder that will get populated
+	// with a default in the resolve step.
+	UNiagaraDataInterface* CDO = GetDataInterfaceCDO(WriteParameter.GetType().GetClass());
+	int32 AddedIndex = AddDataInterfaceInfo(WriteParameter.GetType(), WriteParameter.GetName(), *CDO);
+	FNiagaraScriptDataInterfaceCompileInfo& AddedInfo = CompilationOutput.ScriptData.DataInterfaceInfo[AddedIndex];
+	AddedInfo.TargetParameter = WriteParameterReference;
+}
+
+FNiagaraScriptDataInterfaceCompileId FNiagaraHlslTranslator::MakeUniqueDataInterfaceCompileId(FName BaseName) const
+{
+	int32 IdIndex = 0;
+	FNiagaraScriptDataInterfaceCompileId CompileId = FNiagaraScriptDataInterfaceCompileId(CurrentDIOwnerContext.EmitterHandleId,
+		CurrentDIOwnerContext.ScriptUsage, CurrentDIOwnerContext.ScriptUsageId, CurrentDIOwnerContext.OwnerName, BaseName, IdIndex);
+
+	while (CompilationOutput.ScriptData.DataInterfaceInfo.ContainsByPredicate([&CompileId](const FNiagaraScriptDataInterfaceCompileInfo& CompileInfo)
+		{ return CompileInfo.Id == CompileId; }))
+	{
+		IdIndex++;
+		CompileId = FNiagaraScriptDataInterfaceCompileId(CurrentDIOwnerContext.EmitterHandleId,
+			CurrentDIOwnerContext.ScriptUsage, CurrentDIOwnerContext.ScriptUsageId, CurrentDIOwnerContext.OwnerName, BaseName, IdIndex);
+	}
+
+	return CompileId;
+}
+
+FName FNiagaraHlslTranslator::MakeUniqueDIName(FName OwnerName, FName BaseName) const
+{
+	TSet<FName> CurrentNames;
+	for (const FNiagaraScriptDataInterfaceCompileInfo& DataInterfaceCompileInfo : CompilationOutput.ScriptData.DataInterfaceInfo)
+	{
+		CurrentNames.Add(DataInterfaceCompileInfo.Name);
+	}
+
+	FName CandidateName = *BuildSanitizedSymbolString(BaseName.ToString().Replace(TEXT("."), TEXT("_")).Replace(TEXT(" "), TEXT("_")));
+
+	return FNiagaraUtilities::GetUniqueName(CandidateName, CurrentNames);
+}
+
+int32 FNiagaraHlslTranslator::AddDataInterfaceInfo(const FNiagaraTypeDefinition& Type, FName BaseName, const UNiagaraDataInterface& DataInterface)
+{
+	FNiagaraScriptDataInterfaceCompileId UniqueId = MakeUniqueDataInterfaceCompileId(BaseName);
+	FName UniqueName = MakeUniqueDIName(UniqueId.GetOwnerName(), BaseName);
+
+	int32 Idx = CompilationOutput.ScriptData.DataInterfaceInfo.AddDefaulted();
+	FNiagaraScriptDataInterfaceCompileInfo& Info = CompilationOutput.ScriptData.DataInterfaceInfo[Idx];
+	Info.Id = UniqueId;
+	Info.Name = UniqueName;
+	Info.Type = Type;
+	Info.UserPtrIdx = DataInterface.PerInstanceDataSize() > 0 ? CompilationOutput.ScriptData.NumUserPtrs++ : INDEX_NONE;
+
+	if (DataInterface.NeedsGPUContextInit() && CompileOptions.IsGpuScript())
+	{
+		CompilationOutput.ScriptData.bNeedsGPUContextInit = true;
+	}
+
+	return Idx;
+}
+
+int32 FNiagaraHlslTranslator::AddPlaceholderDuplicateDataInterface(int32 IndexToDuplicate)
+{
+	FName BaseName = CompilationOutput.ScriptData.DataInterfaceInfo[IndexToDuplicate].Id.GetBaseName();
+	FNiagaraScriptDataInterfaceCompileId UniqueId = MakeUniqueDataInterfaceCompileId(BaseName);
+	FName UniqueName = MakeUniqueDIName(UniqueId.GetOwnerName(), BaseName);
+
+	int32 DuplicateIndex = CompilationOutput.ScriptData.DataInterfaceInfo.AddDefaulted();
+	FNiagaraScriptDataInterfaceCompileInfo& DuplicateInfo = CompilationOutput.ScriptData.DataInterfaceInfo[DuplicateIndex];
+	DuplicateInfo = CompilationOutput.ScriptData.DataInterfaceInfo[IndexToDuplicate];
+	DuplicateInfo.Id = UniqueId;
+	DuplicateInfo.Name = UniqueName;
+	return DuplicateIndex;
+}
+
 template<typename GraphBridge>
 void TNiagaraHlslTranslator<GraphBridge>::Convert(const FConvertNode* Convert, TArrayView<const int32> Inputs, TArray<int32>& Outputs)
 {
@@ -10772,6 +10910,12 @@ void TNiagaraHlslTranslator<GraphBridge>::TrimAttributes(const FNiagaraCompileOp
 	}
 }
 
+template<typename GraphBridge>
+FNiagaraGraphInputDataInterfaceKey TNiagaraHlslTranslator<GraphBridge>::GetGraphInputDataInterfaceKey(const FInputNode* InputNode) const
+{
+	return GraphBridge::GetGraphInputDataInterfaceKey(InputNode, CompileDuplicateData);
+}
+
 TUniquePtr<INiagaraHlslTranslator> INiagaraHlslTranslator::CreateTranslator(const FNiagaraCompileRequestDataBase* InCompileData, const FNiagaraCompileRequestDuplicateDataBase* InDuplicateData)
 {
 	// translate everything
diff --git a/Engine/Plugins/FX/Niagara/Source/NiagaraEditor/Private/NiagaraHlslTranslator.h b/Engine/Plugins/FX/Niagara/Source/NiagaraEditor/Private/NiagaraHlslTranslator.h
--- a/Engine/Plugins/FX/Niagara/Source/NiagaraEditor/Private/NiagaraHlslTranslator.h
+++ b/Engine/Plugins/FX/Niagara/Source/NiagaraEditor/Private/NiagaraHlslTranslator.h
@@ -203,6 +203,7 @@ public:
 	FString CustomReadFunction;
 	FString CustomWriteFunction;
 	int32 ParamMapHistoryIndex = INDEX_NONE;
+	FName DIOwnerName = NAME_None;
 
 	bool ShouldDoSpawnOnlyLogic() const;
 	template<typename GraphBridge>
@@ -406,7 +407,6 @@ public:
 
 	static NIAGARAEDITOR_API FString GetFunctionSignatureSymbol(const FNiagaraFunctionSignature& Sig);
 
-	static NIAGARAEDITOR_API FName GetDataInterfaceName(FName BaseName, const FString& UniqueEmitterName, bool bIsParameterMapDataInterface);
 	static NIAGARAEDITOR_API FString GenerateFunctionHlslPrototype(FStringView InVariableName, const FNiagaraFunctionSignature& FunctionSignature);
 	static NIAGARAEDITOR_API FString GenerateFunctionParametersHlsl(const FNiagaraFunctionSignature& FunctionSignature);
 
@@ -415,6 +415,10 @@ public:
 
 	NIAGARAEDITOR_API bool IsEventSpawnScript()const;
 
+	/** Registers an input node reference to a data interface and returns the data interface info index.  The info for this data interfaces may be newly created for this usage, or may be an existing
+		entry if this inputs input variable has been encountered before. */
+	int32 RegisterDataInterfaceForInput(const FNiagaraVariableBase& InputVariable, const FNiagaraGraphInputDataInterfaceKey& GraphInputDataInterfaceKey, const UNiagaraDataInterface* InputDataInterface);
+
 	// collection of functions that rely on the graph implementation and so must be relegated to the TNiagaraHlslTranslator
 	virtual void DefineDataSetVariableWrites(FString& HlslOutput, const FNiagaraDataSetID& Id, int32 DataSetIndex, const TArray<FNiagaraVariable>& WriteVars) = 0;
 	virtual void Message(FNiagaraCompileEventSeverity Severity, FText MessageText, FStringView ShortDescription = FStringView()) = 0;
@@ -481,6 +485,23 @@ protected:
 	void AppendFunctionHlslPrototypeString(FStringView InVariableName, const FNiagaraFunctionSignature& FunctionSignature, FStringBuilderBase& Builder) const;
 	FString BuildFunctionHlslPrototypeString(FStringView InVariableName, const FNiagaraFunctionSignature& FunctionSignature) const;
 
+	/** Extracts parameter reference flags from a parameter's full name. */
+	static ENiagaraCompileParameterReferenceFlags ExtractFlagsFromParameterName(FName ParameterFullName);
+
+	/** Registers a parameter read reference to a data interface and returns the data interface info index.  The info for this data interface may be newly created
+		for this usage, or may be an existing entry if the parameter was read from or written to previously in this translation. */
+	int32 RegisterDataInterfaceForParameterRead(const FNiagaraVariableBase& ReadParameter, int32 BoundIndex);
+
+	/** Registers a target parameter write reference for a data interface. */
+	void RegisterDataInterfaceForParameterWrite(const FNiagaraVariableBase& WriteParameter, int32 BoundIndex);
+
+	/** Creates a unique data interface compile id using the current context and the base name. */
+	FNiagaraScriptDataInterfaceCompileId MakeUniqueDataInterfaceCompileId(FName BaseName) const;
+
+	FName MakeUniqueDIName(FName OwnerName, FName BaseName) const;
+	int32 AddDataInterfaceInfo(const FNiagaraTypeDefinition& Type, FName BaseName, const UNiagaraDataInterface& DataInterface);
+	int32 AddPlaceholderDuplicateDataInterface(int32 IndexToDuplicate);
+
 	template<typename T>
 	void BuildConstantBuffer(ENiagaraCodeChunkMode ChunkMode);
 
@@ -590,6 +611,26 @@ protected:
 
 	// local copy of UNiagaraEditorSettings::GetHLSLKeywordReplacmentsMap() for the duration of the translation process
 	TMap<FString, FString> KeywordReplacementMap;
+
+	struct FDIOwnerContext
+	{
+		FDIOwnerContext() { }
+
+		FDIOwnerContext(const FGuid& InEmitterHandleId, ENiagaraScriptUsage InScriptUsage, const FGuid& InScriptUsageId, FName InOwnerName)
+			: EmitterHandleId(InEmitterHandleId)
+		    , ScriptUsage(InScriptUsage)
+			, ScriptUsageId(InScriptUsageId)
+			, OwnerName(InOwnerName)
+		{
+		}
+
+		FGuid EmitterHandleId;
+		ENiagaraScriptUsage ScriptUsage;
+		FGuid ScriptUsageId;
+		FName OwnerName;
+	};
+
+	FDIOwnerContext CurrentDIOwnerContext;
 };
 
 // all data and function implementations that rely on specifics of the graph implementation
@@ -661,7 +702,6 @@ public:
 
 	int32 RegisterUObject(const FNiagaraVariable& Variable, UObject* Object, bool bAddParameterMapRead);
 	int32 RegisterUObjectPath(const FNiagaraVariable& Variable, const FSoftObjectPath& ObjectPath, bool bAddParmeterMapRead);
-	int32 RegisterDataInterface(const FNiagaraVariable& Var, const UNiagaraDataInterface* DataInterface, bool bPlaceholder, bool bAddParameterMapRead);
 
 	void Operation(const FOpNode* Operation, TArray<int32>& Inputs, TArray<int32>& Outputs);
 	void Output(const FOutputNode* OutputNode, const TArray<int32>& ComputedInputs);
@@ -717,6 +757,8 @@ public:
 
 	void TrimAttributes(const FNiagaraCompileOptions& InCompileOptions, TArray<FNiagaraVariable>& Attributes);
 
+	FNiagaraGraphInputDataInterfaceKey GetGraphInputDataInterfaceKey(const FInputNode* InputNode) const;
+
 protected:
 	void InitializeParameterMapDefaults(int32 ParamMapHistoryIdx);
 	void HandleParameterRead(int32 ParamMapHistoryIdx, const FNiagaraVariable& Var, const FInputPin* DefaultPin, const FNode* ErrorNode, int32& OutputChunkId, TOptional<ENiagaraDefaultMode> DefaultMode = TOptional<ENiagaraDefaultMode>(), TOptional<FNiagaraScriptVariableBinding> DefaultBinding = TOptional<FNiagaraScriptVariableBinding>(), bool bTreatAsUnknownParameterMap = false, bool bIgnoreDefaultSetFirst = false);
@@ -731,7 +773,7 @@ protected:
 	void ProcessCustomHlsl(const FString& InCustomHlsl, ENiagaraScriptUsage InUsage, const FNiagaraFunctionSignature& InSignature, const TArray<int32>& Inputs, const FNode* NodeForErrorReporting, FString& OutCustomHlsl, FNiagaraFunctionSignature& OutSignature);
 	void HandleSimStageSetupAndTeardown(int32 InWhichStage, FString& OutHlsl);
 	// Register a System/Engine/read-only variable in its namespaced form
-	bool ParameterMapRegisterExternalConstantNamespaceVariable(FNiagaraVariable InVariable, const FNode* InNode, int32 InParamMapHistoryIdx, int32& Output, const FPin* InDefaultPin);
+	bool ParameterMapRegisterExternalConstantNamespaceVariable(FNiagaraVariable InVariable, const FNiagaraVariableBase& InUnresolvedVariable, const FNode* InNode, int32 InParamMapHistoryIdx, int32& Output, const FPin* InDefaultPin);
 
 	// Register an attribute in its non namespaced form
 	bool ParameterMapRegisterNamespaceAttributeVariable(const FNiagaraVariable& InVariable, const FNode* InNode, int32 InParamMapHistoryIdx, int32& Output);
diff --git a/Engine/Plugins/FX/Niagara/Source/NiagaraEditor/Private/NiagaraNodeInput.cpp b/Engine/Plugins/FX/Niagara/Source/NiagaraEditor/Private/NiagaraNodeInput.cpp
--- a/Engine/Plugins/FX/Niagara/Source/NiagaraEditor/Private/NiagaraNodeInput.cpp
+++ b/Engine/Plugins/FX/Niagara/Source/NiagaraEditor/Private/NiagaraNodeInput.cpp
@@ -415,11 +415,8 @@ void UNiagaraNodeInput::Compile(FTranslator* Translator, TArray<int32>& Outputs)
 				{
 					if ( IsDataInterface() )
 					{
-						if (ensure(DataInterface != nullptr))
-						{
-							Outputs.Add(Translator->RegisterDataInterface(Input, DataInterface, false, false));
-							return;
-						}
+						Outputs.Add(Translator->RegisterDataInterfaceForInput(Input, Translator->GetGraphInputDataInterfaceKey(this), DataInterface));
+						return;
 					}
 					else if ( IsObjectAsset() )
 					{
@@ -439,7 +436,7 @@ void UNiagaraNodeInput::Compile(FTranslator* Translator, TArray<int32>& Outputs)
 	case ENiagaraInputNodeUsage::Parameter:
 		if (IsDataInterface())
 		{
-			Outputs.Add(Translator->RegisterDataInterface(Input, DataInterface, false, false));
+			Outputs.Add(Translator->RegisterDataInterfaceForInput(Input, Translator->GetGraphInputDataInterfaceKey(this), DataInterface));
 			break;
 		}
 		else if ( IsObjectAsset() )
diff --git a/Engine/Plugins/FX/Niagara/Source/NiagaraEditor/Private/NiagaraScriptSource.cpp b/Engine/Plugins/FX/Niagara/Source/NiagaraEditor/Private/NiagaraScriptSource.cpp
--- a/Engine/Plugins/FX/Niagara/Source/NiagaraEditor/Private/NiagaraScriptSource.cpp
+++ b/Engine/Plugins/FX/Niagara/Source/NiagaraEditor/Private/NiagaraScriptSource.cpp
@@ -259,7 +259,7 @@ void UNiagaraScriptSource::PostLoadFromEmitter(FVersionedNiagaraEmitter OwningEm
 	}
 }
 
-void FindObjectNamesRecursive(UNiagaraGraph* InGraph, ENiagaraScriptUsage Usage, FGuid UsageId, FString EmitterUniqueName, TSet<UNiagaraGraph*>& VisitedGraphs, TMap<FName, UNiagaraDataInterface*>& Result)
+void CollectGraphInputDataInterfacesRecursive(UNiagaraGraph* InGraph, ENiagaraScriptUsage Usage, FGuid UsageId, FString EmitterUniqueName, TSet<UNiagaraGraph*>& VisitedGraphs, TMap<FNiagaraGraphInputDataInterfaceKey, UNiagaraDataInterface*>& Result)
 {
 	if (!InGraph)
 	{
@@ -278,13 +278,12 @@ void FindObjectNamesRecursive(UNiagaraGraph* InGraph, ENiagaraScriptUsage Usage,
 				if (InputNode->Input.IsDataInterface())
 				{
 					UNiagaraDataInterface* DataInterface = InputNode->GetDataInterface();
-					bool bIsParameterMapDataInterface = false;
-					FName DIName = FNiagaraHlslTranslator::GetDataInterfaceName(InputNode->Input.GetName(), EmitterUniqueName, bIsParameterMapDataInterface);
-					if (Result.Contains(DIName) && Result[DIName] != DataInterface)
+					FNiagaraGraphInputDataInterfaceKey GraphInputDataInterfaceKey(*InGraph->GetPathName(), InputNode->Input);
+					if (Result.Contains(GraphInputDataInterfaceKey) && Result[GraphInputDataInterfaceKey] != DataInterface)
 					{
-						UE_LOG(LogNiagaraEditor, Verbose, TEXT("Duplicate data interface name %s in graph %s. One of the data interfaces will override the other when resolving the name."), *DIName.ToString(), *InGraph->GetPathName());
+						UE_LOG(LogNiagaraEditor, Verbose, TEXT("Duplicate data interface name %s in graph %s. One of the data interfaces will override the other when resolving the name."), *InputNode->Input.ToString(), *InGraph->GetPathName());
 					}
-					Result.Add(DIName, DataInterface);
+					Result.Add(GraphInputDataInterfaceKey, DataInterface);
 				}
 				continue;
 			}
@@ -295,16 +294,16 @@ void FindObjectNamesRecursive(UNiagaraGraph* InGraph, ENiagaraScriptUsage Usage,
 				if(VisitedGraphs.Contains(FunctionGraph) == false)
 				{
 					VisitedGraphs.Add(FunctionGraph);
-					FindObjectNamesRecursive(FunctionGraph, FunctionCallNode->GetCalledUsage(), FGuid(), EmitterUniqueName, VisitedGraphs, Result);
+					CollectGraphInputDataInterfacesRecursive(FunctionGraph, FunctionCallNode->GetCalledUsage(), FGuid(), EmitterUniqueName, VisitedGraphs, Result);
 				}
 			}
 		}
 	}
 }
 
-TMap<FName, UNiagaraDataInterface*> UNiagaraScriptSource::ComputeObjectNameMap(UNiagaraSystem& System, ENiagaraScriptUsage Usage, FGuid UsageId, FString EmitterUniqueName) const
+TMap<FNiagaraGraphInputDataInterfaceKey, UNiagaraDataInterface*> UNiagaraScriptSource::CollectGraphInputDataInterfaces(UNiagaraSystem& System, ENiagaraScriptUsage Usage, FGuid UsageId, FString EmitterUniqueName) const
 {
-	TMap<FName, UNiagaraDataInterface*> Result;
+	TMap<FNiagaraGraphInputDataInterfaceKey, UNiagaraDataInterface*> Result;
 	if (!NodeGraph)
 	{
 		return Result;
@@ -323,19 +322,19 @@ TMap<FName, UNiagaraDataInterface*> UNiagaraScriptSource::ComputeObjectNameMap(U
 				OutputNode->GetUsage() == ENiagaraScriptUsage::ParticleUpdateScript ||
 				OutputNode->GetUsage() == ENiagaraScriptUsage::ParticleSimulationStageScript)
 			{
-				FindObjectNamesRecursive(NodeGraph, OutputNode->GetUsage(), OutputNode->GetUsageId(), EmitterUniqueName, VisitedGraphs, Result);
+				CollectGraphInputDataInterfacesRecursive(NodeGraph, OutputNode->GetUsage(), OutputNode->GetUsageId(), EmitterUniqueName, VisitedGraphs, Result);
 			}
 		}
 	}
 	else
 	{
 		// Collect the data interfaces for the current script.
-		FindObjectNamesRecursive(NodeGraph, Usage, UsageId, EmitterUniqueName, VisitedGraphs, Result);
+		CollectGraphInputDataInterfacesRecursive(NodeGraph, Usage, UsageId, EmitterUniqueName, VisitedGraphs, Result);
 
 		if (Usage == ENiagaraScriptUsage::ParticleSpawnScriptInterpolated)
 		{
 			// The interpolated spawn script needs to include the particle update script's data interfaces as well.
-			FindObjectNamesRecursive(NodeGraph, ENiagaraScriptUsage::ParticleUpdateScript, FGuid(), EmitterUniqueName, VisitedGraphs, Result);
+			CollectGraphInputDataInterfacesRecursive(NodeGraph, ENiagaraScriptUsage::ParticleUpdateScript, FGuid(), EmitterUniqueName, VisitedGraphs, Result);
 		}
 		else if (Usage == ENiagaraScriptUsage::SystemSpawnScript || Usage == ENiagaraScriptUsage::SystemUpdateScript)
 		{
@@ -351,7 +350,7 @@ TMap<FName, UNiagaraDataInterface*> UNiagaraScriptSource::ComputeObjectNameMap(U
 				{
 					if (UNiagaraScriptSource* EmitterSource = Cast<UNiagaraScriptSource>(EmitterData->GraphSource); EmitterSource != nullptr)
 					{
-						FindObjectNamesRecursive(EmitterSource->NodeGraph, EmitterUsage, FGuid(), EmitterHandle.GetUniqueInstanceName(), VisitedGraphs, Result);
+						CollectGraphInputDataInterfacesRecursive(EmitterSource->NodeGraph, EmitterUsage, FGuid(), EmitterHandle.GetUniqueInstanceName(), VisitedGraphs, Result);
 					}
 				}
 			}
diff --git a/Engine/Plugins/FX/Niagara/Source/NiagaraEditor/Private/NiagaraValidationRules.cpp b/Engine/Plugins/FX/Niagara/Source/NiagaraEditor/Private/NiagaraValidationRules.cpp
--- a/Engine/Plugins/FX/Niagara/Source/NiagaraEditor/Private/NiagaraValidationRules.cpp
+++ b/Engine/Plugins/FX/Niagara/Source/NiagaraEditor/Private/NiagaraValidationRules.cpp
@@ -1452,11 +1452,11 @@ void UNiagaraValidationRule_TickDependencyCheck::CheckValidity(const FNiagaraVal
 
 				// Generate issue
 				UObject* StackObject = nullptr;
-				if (ResolvedDI.ResolvedSourceEmitterName.Len() > 0)
+				if (ResolvedDI.ResolvedParameterReference.EmitterHandleId.IsValid())
 				{
 					for (const TSharedRef<FNiagaraEmitterHandleViewModel>& EmitterViewModel : Context.ViewModel->GetEmitterHandleViewModels())
 					{
-						if (EmitterViewModel->GetName() == FName(ResolvedDI.ResolvedSourceEmitterName))
+						if (EmitterViewModel->GetId() == ResolvedDI.ResolvedParameterReference.EmitterHandleId)
 						{
 							StackObject = NiagaraValidation::GetStackEntry<UNiagaraStackEmitterPropertiesItem>(EmitterViewModel->GetEmitterStackViewModel());
 							break;
@@ -1474,7 +1474,7 @@ void UNiagaraValidationRule_TickDependencyCheck::CheckValidity(const FNiagaraVal
 				}
 
 				const FText DIClassText = FText::FromName(RuntimeDI->GetClass()->GetFName());
-				const FText DIVariableText = FText::FromName(ResolvedDI.Name);
+				const FText DIVariableText = FText::FromName(ResolvedDI.ParameterStoreVariable.GetName());
 				FNiagaraValidationResult& ValiationResult = OutResults.Emplace_GetRef(
 					Severity,
 					LOCTEXT("TickDependencyCheckFormat", "Performance issue due to late ticking which may cause waits on the game thread."),
diff --git a/Engine/Plugins/FX/Niagara/Source/NiagaraEditor/Private/ViewModels/NiagaraScriptInputCollectionViewModel.cpp b/Engine/Plugins/FX/Niagara/Source/NiagaraEditor/Private/ViewModels/NiagaraScriptInputCollectionViewModel.cpp
--- a/Engine/Plugins/FX/Niagara/Source/NiagaraEditor/Private/ViewModels/NiagaraScriptInputCollectionViewModel.cpp
+++ b/Engine/Plugins/FX/Niagara/Source/NiagaraEditor/Private/ViewModels/NiagaraScriptInputCollectionViewModel.cpp
@@ -280,22 +280,6 @@ void FNiagaraScriptInputCollectionViewModel::RefreshParameterViewModels()
 			else if (InputNode->Input.IsDataInterface())
 			{
 				UNiagaraDataInterface* EmitterDataInterface = InputNode->GetDataInterface();
-				for (FVersionedNiagaraScriptWeakPtr& ScriptWeakPtr : Scripts)
-				{
-					if (EmitterDataInterface != nullptr)
-					{
-						break;
-					}
-
-					for (FNiagaraScriptDataInterfaceInfo& DataInterfaceInfoItem : ScriptWeakPtr.Script->GetCachedDefaultDataInterfaces())
-					{
-						if (DataInterfaceInfoItem.Name == InputNode->Input.GetName())
-						{
-							EmitterDataInterface = DataInterfaceInfoItem.DataInterface;
-							break;
-						}
-					}
-				}
 				ParameterViewModel = MakeShareable(new FNiagaraScriptParameterViewModel(GraphVariable, *InputNode, EmitterDataInterface, ParameterEditMode));
 			}
 			else if (InputNode->Input.IsUObject())
@@ -528,23 +512,6 @@ void FNiagaraScriptInputCollectionViewModel::OnParameterValueChangedInternal(TSh
 						DataInterface->CopyTo(DataInterfaceToUpdate);
 					}
 				}
-
-				// Synchronize script variables...
-				for (FVersionedNiagaraScriptWeakPtr& VersionedScript : Scripts)
-				{
-					if (!VersionedScript.Script.IsValid() || !VersionedScript.Script->GetVMExecutableData().IsValid())
-					{
-						continue;
-					}
-					for (FNiagaraScriptDataInterfaceInfo& Info : VersionedScript.Script->GetCachedDefaultDataInterfaces())
-					{
-						if (Info.Name == ChangedParameter->GetName() && DataInterface != Info.DataInterface)
-						{
-							DataInterface->CopyTo(Info.DataInterface);
-							break;
-						}
-					}
-				}
 			}
 		}
 		else if (ChangedParameter->GetDefaultValueType() == INiagaraParameterViewModel::EDefaultValueType::ObjectAsset)
diff --git a/Engine/Plugins/FX/Niagara/Source/NiagaraEditor/Public/NiagaraCompilationBridge.h b/Engine/Plugins/FX/Niagara/Source/NiagaraEditor/Public/NiagaraCompilationBridge.h
--- a/Engine/Plugins/FX/Niagara/Source/NiagaraEditor/Public/NiagaraCompilationBridge.h
+++ b/Engine/Plugins/FX/Niagara/Source/NiagaraEditor/Public/NiagaraCompilationBridge.h
@@ -167,6 +167,7 @@ struct FNiagaraCompilationGraphBridge
 	static ENiagaraScriptUsage GetEmitterUsage(const FEmitterNode* EmitterNode);
 	static FString GetEmitterName(const FEmitterNode* EmitterNode);
 	static FString GetEmitterPathName(const FEmitterNode* EmitterNode);
+	static FGuid GetEmitterHandleId(const FEmitterNode* EmitterNode);
 	static FString GetEmitterHandleIdString(const FEmitterNode* EmitterNode);
 	static const FGraph* GetFunctionNodeGraph(const FFunctionCallNode* FunctionCall);
 	static FString GetFunctionFullName(const FFunctionCallNode* FunctionCall);
@@ -188,6 +189,7 @@ struct FNiagaraCompilationGraphBridge
 	static FNiagaraVariable GetPinVariable(const FPin* Pin, bool bNeedsValue, ENiagaraStructConversion Conversion);
 	static const FInputPin* GetPinAsInput(const FPin* Pin);
 	static FNiagaraVariable GetInputVariable(const FInputNode* InputNode);
+	static FNiagaraGraphInputDataInterfaceKey GetGraphInputDataInterfaceKey(const FInputNode* InputNode, const FCompilationCopy* CompilationCopy);
 	static const TArray<FNiagaraVariable>& GetOutputVariables(const FOutputNode* OutputNode);
 	static TArray<FNiagaraVariable> GetGraphOutputNodeVariables(const FGraph* Graph, ENiagaraScriptUsage Usage);
 	static TArray<const FInputNode*> GetGraphInputNodes(const FGraph* Graph, const FNiagaraFindInputNodeOptions& Options);
@@ -299,6 +301,7 @@ struct FNiagaraCompilationDigestBridge
 	static ENiagaraScriptUsage GetEmitterUsage(const FEmitterNode* EmitterNode);
 	static FString GetEmitterName(const FEmitterNode* EmitterNode);
 	static FString GetEmitterPathName(const FEmitterNode* EmitterNode);
+	static FGuid GetEmitterHandleId(const FEmitterNode* EmitterNode);
 	static FString GetEmitterHandleIdString(const FEmitterNode* EmitterNode);
 	static const FGraph* GetFunctionNodeGraph(const FFunctionCallNode* FunctionCall);
 	static FString GetFunctionFullName(const FFunctionCallNode* FunctionCall);
@@ -320,6 +323,7 @@ struct FNiagaraCompilationDigestBridge
 	static FNiagaraVariable GetPinVariable(const FPin* Pin, bool bNeedsValue, ENiagaraStructConversion Conversion);
 	static const FInputPin* GetPinAsInput(const FPin* Pin);
 	static FNiagaraVariable GetInputVariable(const FInputNode* InputNode);
+	static FNiagaraGraphInputDataInterfaceKey GetGraphInputDataInterfaceKey(const FInputNode* InputNode, const FCompilationCopy* CompilationCopy);
 	static const TArray<FNiagaraVariable>& GetOutputVariables(const FOutputNode* OutputNode);
 	static TArray<FNiagaraVariable> GetGraphOutputNodeVariables(const FGraph* Graph, ENiagaraScriptUsage Usage);
 	static TArray<const FInputNode*> GetGraphInputNodes(const FGraph* Graph, const FNiagaraFindInputNodeOptions& Options);
diff --git a/Engine/Plugins/FX/Niagara/Source/NiagaraEditor/Public/NiagaraScriptSource.h b/Engine/Plugins/FX/Niagara/Source/NiagaraEditor/Public/NiagaraScriptSource.h
--- a/Engine/Plugins/FX/Niagara/Source/NiagaraEditor/Public/NiagaraScriptSource.h
+++ b/Engine/Plugins/FX/Niagara/Source/NiagaraEditor/Public/NiagaraScriptSource.h
@@ -56,7 +56,7 @@ public:
 
 	virtual void PostLoadFromEmitter(FVersionedNiagaraEmitter OwningEmitter) override;
 
-	virtual TMap<FName, UNiagaraDataInterface*> ComputeObjectNameMap(UNiagaraSystem& System, ENiagaraScriptUsage Usage, FGuid UsageId, FString EmitterUniqueName) const override;
+	virtual TMap<FNiagaraGraphInputDataInterfaceKey, UNiagaraDataInterface*> CollectGraphInputDataInterfaces(UNiagaraSystem& System, ENiagaraScriptUsage Usage, FGuid UsageId, FString EmitterUniqueName) const override;
 
 	NIAGARAEDITOR_API virtual bool AddModuleIfMissing(FString ModulePath, ENiagaraScriptUsage Usage, bool& bOutFoundModule)override;
 
diff --git a/Engine/Source/Runtime/Core/Private/UObject/FortniteMainBranchObjectVersion.cpp b/Engine/Source/Runtime/Core/Private/UObject/FortniteMainBranchObjectVersion.cpp
--- a/Engine/Source/Runtime/Core/Private/UObject/FortniteMainBranchObjectVersion.cpp
+++ b/Engine/Source/Runtime/Core/Private/UObject/FortniteMainBranchObjectVersion.cpp
@@ -12,7 +12,7 @@ TMap<FGuid, FGuid> FFortniteMainBranchObjectVersion::GetSystemGuids()
 	SystemGuids.Add(DevGuids.LANDSCAPE_MOBILE_COOK_VERSION, FGuid("32D02EF867C74B71A0D4E0FA41392732"));
 	SystemGuids.Add(DevGuids.MATERIALSHADERMAP_DERIVEDDATA_VER, FGuid("4DCE94C9E9994B3DA256F6A015B87CA9"));
 	SystemGuids.Add(DevGuids.NIAGARASHADERMAP_DERIVEDDATA_VER, FGuid("CC229D919A214E59B79D01AC2E5BF981"));
-	SystemGuids.Add(DevGuids.Niagara_LatestScriptCompileVersion, FGuid("9D86DE4A5F0646AF8AFF2ACE830CE679"));
+	SystemGuids.Add(DevGuids.Niagara_LatestScriptCompileVersion, FGuid("C19F5C3F8A704C4A83F3077F3236AE2B"));
 	SystemGuids.Add(DevGuids.MaterialTranslationDDCVersion, FGuid("FD43494F1D8443D79A06E6E28F503A42"));
 
 	// These GUIDs are in AppendRestrictedFortniteMainStreamObjectVersionGuids as they should rarely be changed.
