- Fix some static analysis issues

diff --git a/Engine/Plugins/FX/Niagara/Source/NiagaraShader/Private/NiagaraBatchedElements.cpp b/Engine/Plugins/FX/Niagara/Source/NiagaraShader/Private/NiagaraBatchedElements.cpp
--- a/Engine/Plugins/FX/Niagara/Source/NiagaraShader/Private/NiagaraBatchedElements.cpp
+++ b/Engine/Plugins/FX/Niagara/Source/NiagaraShader/Private/NiagaraBatchedElements.cpp
@@ -146,7 +146,7 @@ void FBatchedElementNiagaraVolumeAttribute::BindShaders(FRHICommandList& RHICmdL
 		PassParameters.InAttributeUVAndChannel[i].X = AttributeUV.X;
 		PassParameters.InAttributeUVAndChannel[i].Y = AttributeUV.Y;
 		PassParameters.InAttributeUVAndChannel[i].Z = AttributeUV.Z;
-		reinterpret_cast<int32&>(PassParameters.InAttributeUVAndChannel[i].W) = AttributeChannel;
+		FMemory::Memcpy(&PassParameters.InAttributeUVAndChannel[i].W, &AttributeChannel, sizeof(int32));
 	}
 	PassParameters.InTileUV			= TileUVs;
 	PassParameters.InColorWeights	= FMatrix44f(ColorWeights);
diff --git a/Engine/Plugins/FX/Niagara/Source/NiagaraShader/Private/NiagaraDistanceFieldHelper.cpp b/Engine/Plugins/FX/Niagara/Source/NiagaraShader/Private/NiagaraDistanceFieldHelper.cpp
--- a/Engine/Plugins/FX/Niagara/Source/NiagaraShader/Private/NiagaraDistanceFieldHelper.cpp
+++ b/Engine/Plugins/FX/Niagara/Source/NiagaraShader/Private/NiagaraDistanceFieldHelper.cpp
@@ -28,7 +28,7 @@ void FNiagaraDistanceFieldHelper::SetGlobalDistanceFieldParameters(const FGlobal
 		ShaderParameters.GlobalDistanceFieldInvPageAtlasSize = (FVector3f)OptionalParameterData->InvPageAtlasSize;
 		ShaderParameters.GlobalDistanceFieldInvCoverageAtlasSize = (FVector3f)OptionalParameterData->InvCoverageAtlasSize;
 		ShaderParameters.GlobalVolumeDimension = OptionalParameterData->GlobalDFResolution;
-		ShaderParameters.GlobalVolumeTexelSize = 1.0f / OptionalParameterData->GlobalDFResolution;
+		ShaderParameters.GlobalVolumeTexelSize = FMath::IsNearlyZero(OptionalParameterData->GlobalDFResolution) ? 0.0f : 1.0f / OptionalParameterData->GlobalDFResolution;
 		ShaderParameters.MaxGlobalDFAOConeDistance = OptionalParameterData->MaxDFAOConeDistance;
 		ShaderParameters.NumGlobalSDFClipmaps = OptionalParameterData->NumGlobalSDFClipmaps;
 	}
diff --git a/Engine/Plugins/FX/Niagara/Source/NiagaraShader/Private/NiagaraShaderCompilationManager.cpp b/Engine/Plugins/FX/Niagara/Source/NiagaraShader/Private/NiagaraShaderCompilationManager.cpp
--- a/Engine/Plugins/FX/Niagara/Source/NiagaraShader/Private/NiagaraShaderCompilationManager.cpp
+++ b/Engine/Plugins/FX/Niagara/Source/NiagaraShader/Private/NiagaraShaderCompilationManager.cpp
@@ -96,7 +96,7 @@ void FNiagaraShaderCompilationManager::ProcessAsyncResults()
 		if (Results.FinishedJobs.Num() == Results.NumJobsQueued)
 		{
 			PendingFinalizeNiagaraShaderMaps.Add(It.Key(), FNiagaraShaderMapFinalizeResults(Results));
-			NiagaraShaderMapJobs.Remove(It.Key());
+			It.RemoveCurrent();
 		}
 	}
 
@@ -270,7 +270,7 @@ void FNiagaraShaderCompilationManager::ProcessCompiledNiagaraShaderMaps(TMap<int
 
 				// Cleanup shader jobs and compile tracking structures
 				ResultArray.Empty();
-				CompiledShaderMaps.Remove(ShaderMap->GetCompilingId());
+				ProcessIt.RemoveCurrent();
 			}
 
 			if (TimeBudget < 0)
diff --git a/Engine/Plugins/FX/Niagara/Source/NiagaraShader/Public/NiagaraShader.h b/Engine/Plugins/FX/Niagara/Source/NiagaraShader/Public/NiagaraShader.h
--- a/Engine/Plugins/FX/Niagara/Source/NiagaraShader/Public/NiagaraShader.h
+++ b/Engine/Plugins/FX/Niagara/Source/NiagaraShader/Public/NiagaraShader.h
@@ -1,9 +1,5 @@
 // Copyright Epic Games, Inc. All Rights Reserved.
 
-/*=============================================================================
-	MaterialShader.h: Shader base classes
-=============================================================================*/
-
 #pragma once
 
 #include "ShaderParameters.h"
@@ -23,7 +19,7 @@ struct FAppendToClassSchemaContext;
 
 template<typename TBufferStruct> class TUniformBufferRef;
 
-/** Base class of all shaders that need material parameters. */
+/** Base class for all generated Niagara compute shaders. */
 class FNiagaraShader : public FShader
 {
 public:
@@ -79,9 +75,9 @@ public:
 		SHADER_PARAMETER(int32,		PREV_Engine_System_NumEmittersAlive)
 		SHADER_PARAMETER(int32,		PREV_Engine_System_SignificanceIndex)
 		SHADER_PARAMETER(int32,		PREV_Engine_System_RandomSeed)
-		SHADER_PARAMETER(int32,		PREV_CurrentTimeStep)
-		SHADER_PARAMETER(int32,		PREV_NumTimeSteps)
-		SHADER_PARAMETER(float,		PREV_TimeStepFraction)
+		SHADER_PARAMETER(int32,		PREV_Engine_System_CurrentTimeStep)
+		SHADER_PARAMETER(int32,		PREV_Engine_System_NumTimeSteps)
+		SHADER_PARAMETER(float,		PREV_Engine_System_TimeStepFraction)
 		SHADER_PARAMETER(int32,		PREV_Engine_System_NumParticles)
 		SHADER_PARAMETER(int32,		PREV_System_Pad0)
 		SHADER_PARAMETER(int32,		PREV_System_Pad1)
diff --git a/Engine/Plugins/FX/Niagara/Source/NiagaraShader/Public/NiagaraShaderMap.h b/Engine/Plugins/FX/Niagara/Source/NiagaraShader/Public/NiagaraShaderMap.h
deleted file mode 100644
diff --git a/Engine/Plugins/FX/Niagara/Source/NiagaraShader/Public/NiagaraShaderModule.h b/Engine/Plugins/FX/Niagara/Source/NiagaraShader/Public/NiagaraShaderModule.h
--- a/Engine/Plugins/FX/Niagara/Source/NiagaraShader/Public/NiagaraShaderModule.h
+++ b/Engine/Plugins/FX/Niagara/Source/NiagaraShader/Public/NiagaraShaderModule.h
@@ -33,7 +33,7 @@ public:
 		Singleton = nullptr;
 	}
 
-	// If you know that you wnat the shader module and have access to the header you don't have to find the module.
+	// If you know that you want the shader module and have access to the header you don't have to find the module.
 	static INiagaraShaderModule* Get()
 	{
 		return Singleton;
