- Remove some unused members from FNiagaraShaderScriptParametersMetadata

#rnx

diff --git a/Engine/Plugins/FX/Niagara/Source/Niagara/Private/NiagaraCompilationTypes.cpp b/Engine/Plugins/FX/Niagara/Source/Niagara/Private/NiagaraCompilationTypes.cpp
--- a/Engine/Plugins/FX/Niagara/Source/Niagara/Private/NiagaraCompilationTypes.cpp
+++ b/Engine/Plugins/FX/Niagara/Source/Niagara/Private/NiagaraCompilationTypes.cpp
@@ -102,8 +102,6 @@ bool operator==(const FNiagaraShaderScriptParametersMetadata& Lhs, const FNiagar
 {
 	return Lhs.DataInterfaceParamInfo == Rhs.DataInterfaceParamInfo
 		&& Lhs.LooseMetadataNames == Rhs.LooseMetadataNames
-		&& Lhs.bExternalConstantsInterpolated == Rhs.bExternalConstantsInterpolated
-		&& Lhs.ExternalConstants == Rhs.ExternalConstants
 		&& Lhs.StructIncludeInfos == Rhs.StructIncludeInfos
 		&& Lhs.ShaderParametersMetadata == Rhs.ShaderParametersMetadata;
 }
diff --git a/Engine/Plugins/FX/Niagara/Source/Niagara/Private/NiagaraFunctionLibrary.cpp b/Engine/Plugins/FX/Niagara/Source/Niagara/Private/NiagaraFunctionLibrary.cpp
--- a/Engine/Plugins/FX/Niagara/Source/Niagara/Private/NiagaraFunctionLibrary.cpp
+++ b/Engine/Plugins/FX/Niagara/Source/Niagara/Private/NiagaraFunctionLibrary.cpp
@@ -309,27 +309,6 @@ UNiagaraComponent* UNiagaraFunctionLibrary::SpawnSystemAttachedWithParams(const
 	return PSC;
 }
 
-/**
-* Set a constant in an emitter of a Niagara System
-void UNiagaraFunctionLibrary::SetUpdateScriptConstant(UNiagaraComponent* Component, FName EmitterName, FName ConstantName, FVector Value)
-{
-	TArray<TSharedPtr<FNiagaraEmitterInstance>> &Emitters = Component->GetSystemInstance()->GetEmitters();
-
-	for (TSharedPtr<FNiagaraEmitterInstance> &Emitter : Emitters)
-	{		
-		if(UNiagaraEmitter* PinnedProps = Emitter->GetProperties().Get())
-		{
-			FName CurName = *PinnedProps->EmitterName;
-			if (CurName == EmitterName)
-			{
-				Emitter->GetProperties()->UpdateScriptProps.ExternalConstants.SetOrAdd(FNiagaraTypeDefinition::GetVec4Def(), ConstantName, Value);
-				break;
-			}
-		}
-	}
-}
-*/
-
 void UNiagaraFunctionLibrary::OverrideSystemUserVariableStaticMeshComponent(UNiagaraComponent* NiagaraSystem, const FString& OverrideName, UStaticMeshComponent* StaticMeshComponent)
 {
 	if (!NiagaraSystem)
diff --git a/Engine/Plugins/FX/Niagara/Source/NiagaraShader/Public/NiagaraShared.h b/Engine/Plugins/FX/Niagara/Source/NiagaraShader/Public/NiagaraShared.h
--- a/Engine/Plugins/FX/Niagara/Source/NiagaraShader/Public/NiagaraShared.h
+++ b/Engine/Plugins/FX/Niagara/Source/NiagaraShader/Public/NiagaraShared.h
@@ -285,12 +285,6 @@ struct FNiagaraShaderScriptParametersMetadata
 	UPROPERTY()
 	TArray<FString> LooseMetadataNames;
 
-	UPROPERTY()
-	bool bExternalConstantsInterpolated = false;
-
-	UPROPERTY()
-	TArray<FNiagaraShaderScriptExternalConstant> ExternalConstants;
-
 	// Runtime generated bindings, not serialized
 	TArray<FNiagaraDataInterfaceStructIncludeInfo> StructIncludeInfos;
 	TSharedPtr<FShaderParametersMetadata> ShaderParametersMetadata;
