Niagara - Work around windows clipboard issues to prevent random failures in the copy/paste tests.

diff --git a/Engine/Plugins/FX/Niagara/Source/NiagaraEditor/Private/NiagaraClipboard.cpp b/Engine/Plugins/FX/Niagara/Source/NiagaraEditor/Private/NiagaraClipboard.cpp
--- a/Engine/Plugins/FX/Niagara/Source/NiagaraEditor/Private/NiagaraClipboard.cpp
+++ b/Engine/Plugins/FX/Niagara/Source/NiagaraEditor/Private/NiagaraClipboard.cpp
@@ -342,6 +342,19 @@ FNiagaraClipboard::FNiagaraClipboard()
 }
 
 void FNiagaraClipboard::SetClipboardContent(UNiagaraClipboardContent* ClipboardContent)
+{
+	FPlatformApplicationMisc::ClipboardCopy(*ClipboardContentToString(ClipboardContent));
+}
+
+const UNiagaraClipboardContent* FNiagaraClipboard::GetClipboardContent() const
+{
+	// Get the text from the clipboard.
+	FString ClipboardText;
+	FPlatformApplicationMisc::ClipboardPaste(ClipboardText);
+	return ClipboardContentFromString(ClipboardText);
+}
+
+FString FNiagaraClipboard::ClipboardContentToString(UNiagaraClipboardContent* ClipboardContent)
 {
 	// Clear the mark state for saving.
 	UnMarkAllObjects(EObjectMark(OBJECTMARK_TagExp | OBJECTMARK_TagImp));
@@ -350,15 +363,11 @@ void FNiagaraClipboard::SetClipboardContent(UNiagaraClipboardContent* ClipboardC
 	FStringOutputDevice Archive;
 	const FExportObjectInnerContext Context;
 	UExporter::ExportToOutputDevice(&Context, ClipboardContent, nullptr, Archive, TEXT("copy"), 0, PPF_ExportsNotFullyQualified | PPF_Copy | PPF_Delimited, false, ClipboardContent->GetOuter());
-	FPlatformApplicationMisc::ClipboardCopy(*Archive);
+	return *Archive;
 }
 
-const UNiagaraClipboardContent* FNiagaraClipboard::GetClipboardContent() const
+UNiagaraClipboardContent* FNiagaraClipboard::ClipboardContentFromString(const FString& ClipboardText)
 {
-	// Get the text from the clipboard.
-	FString ClipboardText;
-	FPlatformApplicationMisc::ClipboardPaste(ClipboardText);
-
 	// Try to create niagara clipboard content from that.
 	FNiagaraClipboardContentTextObjectFactory ClipboardContentFactory;
 	if (ClipboardContentFactory.CanCreateObjectsFromText(ClipboardText))
@@ -370,7 +379,7 @@ const UNiagaraClipboardContent* FNiagaraClipboard::GetClipboardContent() const
 	if (ClipboardText.IsEmpty() == false)
 	{
 		// If the clipboard text wasn't a niagara clipboard object, it's likely been copied from elsewhere in the editor and may be 
-		// a valid portable value so we consturct an empty clipboard object and set the portable data.
+		// a valid portable value so we construct an empty clipboard object and set the portable data.
 		UNiagaraClipboardContent* ClipboardContent = UNiagaraClipboardContent::Create();
 		FNiagaraClipboardPortableValue& PortableValue = ClipboardContent->PortableValues.AddDefaulted_GetRef();
 		PortableValue.ValueString = ClipboardText;
diff --git a/Engine/Plugins/FX/Niagara/Source/NiagaraEditor/Public/NiagaraClipboard.h b/Engine/Plugins/FX/Niagara/Source/NiagaraEditor/Public/NiagaraClipboard.h
--- a/Engine/Plugins/FX/Niagara/Source/NiagaraEditor/Public/NiagaraClipboard.h
+++ b/Engine/Plugins/FX/Niagara/Source/NiagaraEditor/Public/NiagaraClipboard.h
@@ -292,6 +292,10 @@ public:
 	NIAGARAEDITOR_API void SetClipboardContent(UNiagaraClipboardContent* ClipboardContent);
 
 	NIAGARAEDITOR_API const UNiagaraClipboardContent* GetClipboardContent() const;
+
+	NIAGARAEDITOR_API static FString ClipboardContentToString(UNiagaraClipboardContent* ClipboardContent);
+
+	NIAGARAEDITOR_API static UNiagaraClipboardContent* ClipboardContentFromString(const FString& ClipboardText);
 };
 
 UCLASS(MinimalAPI)
