- Fix for RHI validation warning about null overlap UAV

diff --git a/Engine/Plugins/FX/Niagara/Source/Niagara/Private/NiagaraDataChannelData.cpp b/Engine/Plugins/FX/Niagara/Source/Niagara/Private/NiagaraDataChannelData.cpp
--- a/Engine/Plugins/FX/Niagara/Source/Niagara/Private/NiagaraDataChannelData.cpp
+++ b/Engine/Plugins/FX/Niagara/Source/Niagara/Private/NiagaraDataChannelData.cpp
@@ -806,17 +806,19 @@ void FNiagaraDataChannelDataProxy::AddTransition(FRDGBuilder& GraphBuilder, ERHI
 {
 	if (Buffer)
 	{
-		TArray<FRHITransitionInfo, TInlineAllocator<6>> Transitions;
-		Transitions.Reserve(3);
-
-		Transitions.Emplace(Buffer->GetGPUBufferFloat().UAV, AccessBefore, AccessAfter);
-		Transitions.Emplace(Buffer->GetGPUBufferInt().UAV, AccessBefore, AccessAfter);
-		//TODO: Half Support | Transitions.Emplace(Buffer->GetGPUBufferHalf().UAV, AccessBefore, AccessAfter);
-
-		TArray<FRHIUnorderedAccessView*> UAVsToOverlap;
-		UAVsToOverlap.Emplace(Buffer->GetGPUBufferFloat().UAV);
-		UAVsToOverlap.Emplace(Buffer->GetGPUBufferInt().UAV);
-		//TODO: Half Support | UAVsToOverlap.Emplace(Buffer->GetGPUBufferHalf().UAV);
+		TArray<FRHITransitionInfo, TInlineAllocator<3>> Transitions;
+		TArray<FRHIUnorderedAccessView*, TInlineAllocator<2>> UAVsToOverlap;
+		if (Buffer->GetGPUBufferFloat().UAV)
+		{
+			Transitions.Emplace(Buffer->GetGPUBufferFloat().UAV, AccessBefore, AccessAfter);
+			UAVsToOverlap.Emplace(Buffer->GetGPUBufferFloat().UAV);
+		}
+		if (Buffer->GetGPUBufferInt().UAV)
+		{
+			Transitions.Emplace(Buffer->GetGPUBufferInt().UAV, AccessBefore, AccessAfter);
+			UAVsToOverlap.Emplace(Buffer->GetGPUBufferInt().UAV);
+		}
+		//TODO: Half Support
 
 	GraphBuilder.AddPass(
 		RDG_EVENT_NAME("NDC Proxy - Transition Buffers"),
