Add (supressable) log errors for critical errors in GPUTick generation and have the tick skipped entirely.
If we try to run a tick with critical errors it will assert inside the shader validation etc.

diff --git a/Engine/Plugins/FX/Niagara/Source/Niagara/Classes/NiagaraGPUSystemTick.h b/Engine/Plugins/FX/Niagara/Source/Niagara/Classes/NiagaraGPUSystemTick.h
--- a/Engine/Plugins/FX/Niagara/Source/Niagara/Classes/NiagaraGPUSystemTick.h
+++ b/Engine/Plugins/FX/Niagara/Source/Niagara/Classes/NiagaraGPUSystemTick.h
@@ -112,6 +112,8 @@ public:
 		return ExternalUnformBuffers_RT[BufferIndex];
 	}
 
+	bool IsValid()const{ return !bHasErrors; }
+
 public:
 	// Transient data used by the RT
 	TArray<FUniformBufferRHIRef> ExternalUnformBuffers_RT;
@@ -128,6 +130,8 @@ public:
 	uint32 TotalDispatches = 0;
 	bool bIsFinalTick = false;
 	bool bHasInterpolatedParameters = false;
+	
+	bool bHasErrors = false;
 
 #if !UE_BUILD_SHIPPING && !UE_BUILD_TEST
 	// Debugging only
diff --git a/Engine/Plugins/FX/Niagara/Source/Niagara/Private/NiagaraGPUSystemTick.cpp b/Engine/Plugins/FX/Niagara/Source/Niagara/Private/NiagaraGPUSystemTick.cpp
--- a/Engine/Plugins/FX/Niagara/Source/Niagara/Private/NiagaraGPUSystemTick.cpp
+++ b/Engine/Plugins/FX/Niagara/Source/Niagara/Private/NiagaraGPUSystemTick.cpp
@@ -6,6 +6,17 @@
 #include "NiagaraSystemInstance.h"
 #include "NiagaraSystem.h"
 
+namespace NiagaraLoggingCVars
+{
+	static bool bEmitGPUTickGenErrors = true;
+	static FAutoConsoleVariableRef CVarEmitGPUTickGenErrors(
+		TEXT("fx.Niagara.Logging.EmitGPUGenTickErrors"),
+		bEmitGPUTickGenErrors,
+		TEXT("If true, we will emit errors to the log when Niagara GPU Ticks have critical errors."),
+		ECVF_Default
+	);
+}
+
 #if NIAGARA_VALIDATE_NDIPROXY_REFS
 FNiagaraComputeDataInterfaceInstanceData::~FNiagaraComputeDataInterfaceInstanceData()
 {
@@ -24,6 +35,8 @@ void FNiagaraGPUSystemTick::Init(FNiagaraSystemInstance* InSystemInstance)
 	SystemInstanceID = InSystemInstance->GetId();
 	SystemGpuComputeProxy = InSystemInstance->GetSystemGpuComputeProxy();
  
+	bHasErrors = false;
+
 	uint32 DataSizeForGPU = InSystemInstance->GPUDataInterfaceInstanceDataSize;
 
 	if (DataSizeForGPU > 0)
@@ -199,12 +212,24 @@ void FNiagaraGPUSystemTick::Init(FNiagaraSystemInstance* InSystemInstance)
 			for (UNiagaraDataInterface* DI : DataInterfaces)
 			{
 				FNiagaraDataInterfaceProxy* DIProxy = DI->GetProxy();
-				check(DIProxy);
-				InstanceData->DataInterfaceProxies.Add(DIProxy);
 
-				if ( FNiagaraDataInterfaceProxyRW* RWProxy = DIProxy->AsIterationProxy() )
+				if(DIProxy)
 				{
-					InstanceData->IterationDataInterfaceProxies.Add(RWProxy);
+					InstanceData->DataInterfaceProxies.Add(DIProxy);
+
+					if ( FNiagaraDataInterfaceProxyRW* RWProxy = DIProxy->AsIterationProxy() )
+					{
+						InstanceData->IterationDataInterfaceProxies.Add(RWProxy);
+					}
+				}	
+				else
+				{
+					bHasErrors = true;
+
+					if(NiagaraLoggingCVars::bEmitGPUTickGenErrors)
+					{
+						UE_LOG(LogNiagara, Error, TEXT("Error Generating GPU Tick. Null DI Proxy!\nUse cmd \"fx.Niagara.Logging.EmitGPUGenTickErrors false\" to suppress these errors.\nNiagara System: %s\nData Interface:%s"), *InSystemInstance->GetSystem()->GetPathName(), DI ? *DI->GetPathName() : TEXT("null DI!"));
+					}
 				}
 			}
 
diff --git a/Engine/Plugins/FX/Niagara/Source/Niagara/Private/NiagaraSystemGpuComputeProxy.cpp b/Engine/Plugins/FX/Niagara/Source/Niagara/Private/NiagaraSystemGpuComputeProxy.cpp
--- a/Engine/Plugins/FX/Niagara/Source/Niagara/Private/NiagaraSystemGpuComputeProxy.cpp
+++ b/Engine/Plugins/FX/Niagara/Source/Niagara/Private/NiagaraSystemGpuComputeProxy.cpp
@@ -166,6 +166,11 @@ void FNiagaraSystemGpuComputeProxy::QueueTick(const FNiagaraGPUSystemTick& Tick)
 {
 	check(IsInRenderingThread());
 
+	if(Tick.IsValid() == false)
+	{
+		return;
+	}
+
 	//if ( !FNiagaraUtilities::AllowGPUParticles(GetShaderPlatform()) )
 	//{
 	//	return;
