Fix leak when bypassing critial error in GPUTick creation.

diff --git a/Engine/Plugins/FX/Niagara/Source/Niagara/Private/NiagaraSystemGpuComputeProxy.cpp b/Engine/Plugins/FX/Niagara/Source/Niagara/Private/NiagaraSystemGpuComputeProxy.cpp
--- a/Engine/Plugins/FX/Niagara/Source/Niagara/Private/NiagaraSystemGpuComputeProxy.cpp
+++ b/Engine/Plugins/FX/Niagara/Source/Niagara/Private/NiagaraSystemGpuComputeProxy.cpp
@@ -162,12 +162,13 @@ void FNiagaraSystemGpuComputeProxy::ClearTicksFromRenderThread(FNiagaraGpuComput
 	);
 }
 
-void FNiagaraSystemGpuComputeProxy::QueueTick(const FNiagaraGPUSystemTick& Tick)
+void FNiagaraSystemGpuComputeProxy::QueueTick(FNiagaraGPUSystemTick& Tick)
 {
 	check(IsInRenderingThread());
 
 	if(Tick.IsValid() == false)
 	{
+		Tick.Destroy();
 		return;
 	}
 
diff --git a/Engine/Plugins/FX/Niagara/Source/Niagara/Public/NiagaraSystemGpuComputeProxy.h b/Engine/Plugins/FX/Niagara/Source/Niagara/Public/NiagaraSystemGpuComputeProxy.h
--- a/Engine/Plugins/FX/Niagara/Source/Niagara/Public/NiagaraSystemGpuComputeProxy.h
+++ b/Engine/Plugins/FX/Niagara/Source/Niagara/Public/NiagaraSystemGpuComputeProxy.h
@@ -25,7 +25,7 @@ public:
 
 	FNiagaraSystemInstanceID GetSystemInstanceID() const { return SystemInstanceID; }
 	ENiagaraGpuComputeTickStage::Type GetComputeTickStage() const { return ComputeTickStage; }
-	void QueueTick(const FNiagaraGPUSystemTick& Tick);
+	void QueueTick(FNiagaraGPUSystemTick& Tick);
 	void ReleaseTicks(int32 NumTicksToRelease, bool bLastViewFamily);
 
 	bool RequiresGlobalDistanceField() const { return bRequiresGlobalDistanceField; }
