[Backout] - CL46727218
#rnx
Original CL Desc
-----------------------------------------------------------------
Niagara - Data interface routing fixes and expanded support.
+ When collecting data interfaces from scripts pre and post translate, associate them with a key struct with complete source information instead of just a name to avoid ambibuity and string manipulation.
+ Remove redundant data from the cached default data interfaces struct and rename it to clarify it's purpose.
+ Create a new struct for parameter references which includes emitter handle information and uses flags for namespaces which are extracted from the parameter name at translation time to avoid string wrangling at runtime.
+ Track the DI souce at a more granular level when translating so that the resolve can be done more reliably.
+ Refactor the DI resolve algorithm to better encapsulate handling data from one script at a time and to better track parameter references for emitters, to prevent sharing of script-local DIs, and to enable sharing of unresolved DIs create at resolve time.
+ Refactor testing to be more granular and added additional test cases.

+ Fixes an issue where a single data interface written to multiple parameters might end up with the wrong name in the execution parameter store.
+ Fixes an issue where a single data interfaces written to the same parameter in multiple modules would generate invalid placeholder data interfaces.
+ Adds support for binding different source data interface parameters to the same target data interface parameter if all of the source data interfaces resolve to the same instance.

#virtualized

diff --git a/Engine/Plugins/FX/Niagara/Source/Niagara/Classes/NiagaraConstants.h b/Engine/Plugins/FX/Niagara/Source/Niagara/Classes/NiagaraConstants.h
--- a/Engine/Plugins/FX/Niagara/Source/Niagara/Classes/NiagaraConstants.h
+++ b/Engine/Plugins/FX/Niagara/Source/Niagara/Classes/NiagaraConstants.h
@@ -250,6 +250,7 @@ struct FNiagaraConstants
 	static NIAGARA_API const FString StackContextNamespaceString;
 	static NIAGARA_API const FString SystemNamespaceString;
 	static NIAGARA_API const FString UserNamespaceString;
+	static NIAGARA_API const FString InternalNamespaceString;
 	static NIAGARA_API const FString PreviousNamespaceString;
 
 	static NIAGARA_API const int32 MaxCategoryNameLength;
diff --git a/Engine/Plugins/FX/Niagara/Source/Niagara/Classes/NiagaraScript.h b/Engine/Plugins/FX/Niagara/Source/Niagara/Classes/NiagaraScript.h
--- a/Engine/Plugins/FX/Niagara/Source/Niagara/Classes/NiagaraScript.h
+++ b/Engine/Plugins/FX/Niagara/Source/Niagara/Classes/NiagaraScript.h
@@ -1248,7 +1248,7 @@ public:
 	NIAGARA_API FOnPropertyChanged& OnPropertyChanged();
 
 	/** External call used to identify the values for a successful VM script compilation. OnVMScriptCompiled will be issued in this case.*/
-	void SetVMCompilationResults(const FNiagaraVMExecutableDataId& InCompileId, FNiagaraVMExecutableData& InScriptVM, const TMap<FNiagaraGraphInputDataInterfaceKey, UNiagaraDataInterface*>& GraphInputDataInterfaces, bool ApplyRapidIterationParameters);
+	void SetVMCompilationResults(const FNiagaraVMExecutableDataId& InCompileId, FNiagaraVMExecutableData& InScriptVM, const FString& EmitterUniqueName, const TMap<FName, UNiagaraDataInterface*>& ObjectNameMap, bool ApplyRapidIterationParameters);
 
 	void SetComputeCompilationResults(
 		const ITargetPlatform* TargetPlatform,
@@ -1289,8 +1289,8 @@ public:
 	NIAGARA_API TArray<TObjectPtr<UNiagaraParameterCollection>>& GetCachedParameterCollectionReferences();
 
 #if WITH_EDITORONLY_DATA
-	TArray<FNiagaraScriptDataInterfaceDefaultInfo>& GetCachedGraphDefaultDataInterfaces() { return CachedGraphDefaultDataInterfaces; }
-	TConstArrayView<FNiagaraScriptDataInterfaceDefaultInfo> GetCachedGraphDefaultDataInterfaces() const { return MakeArrayView(CachedGraphDefaultDataInterfaces); }
+	TArray<FNiagaraScriptDataInterfaceInfo>& GetCachedDefaultDataInterfaces() { return CachedDefaultDataInterfaces; }
+	TConstArrayView<FNiagaraScriptDataInterfaceInfo> GetCachedDefaultDataInterfaces() const { return MakeArrayView(CachedDefaultDataInterfaces); }
 
 	TArrayView<FNiagaraScriptUObjectCompileInfo> GetCachedDefaultUObjects() { return CachedDefaultUObjects; }
 	TConstArrayView<FNiagaraScriptUObjectCompileInfo> GetCachedDefaultUObjects() const { return CachedDefaultUObjects; }
@@ -1396,7 +1396,7 @@ public:
 
 #if WITH_EDITORONLY_DATA
 	UPROPERTY()
-	TArray<FNiagaraScriptDataInterfaceDefaultInfo> CachedGraphDefaultDataInterfaces;
+	TArray<FNiagaraScriptDataInterfaceInfo> CachedDefaultDataInterfaces;
 #endif
 
 #if WITH_EDITORONLY_DATA
diff --git a/Engine/Plugins/FX/Niagara/Source/Niagara/Classes/NiagaraScriptSourceBase.h b/Engine/Plugins/FX/Niagara/Source/Niagara/Classes/NiagaraScriptSourceBase.h
--- a/Engine/Plugins/FX/Niagara/Source/Niagara/Classes/NiagaraScriptSourceBase.h
+++ b/Engine/Plugins/FX/Niagara/Source/Niagara/Classes/NiagaraScriptSourceBase.h
@@ -58,7 +58,7 @@ public:
 	virtual ~FNiagaraCompileRequestDuplicateDataBase() {}
 	virtual bool IsDuplicateDataFor(UNiagaraSystem* InSystem, UNiagaraEmitter* InEmitter, UNiagaraScript* InScript) const = 0;
 	virtual void GetDuplicatedObjects(TArray<UObject*>& Objects) = 0;
-	virtual const TMap<FNiagaraGraphInputDataInterfaceKey, UNiagaraDataInterface*>& GetGraphInputDataInterfaces() = 0;
+	virtual const TMap<FName, UNiagaraDataInterface*>& GetObjectNameMap() = 0;
 	virtual const UNiagaraScriptSourceBase* GetScriptSource() const =  0;
 	virtual int32 GetDependentRequestCount() const = 0;
 	virtual TSharedPtr<FNiagaraCompileRequestDuplicateDataBase, ESPMode::ThreadSafe> GetDependentRequest(int32 Index) = 0;
@@ -136,7 +136,7 @@ class UNiagaraScriptSourceBase : public UObject
 	virtual void ReportAnalyticsData(FNiagaraScriptSourceAnalytics& InData) const {};
 #endif
 
-	virtual TMap<FNiagaraGraphInputDataInterfaceKey, UNiagaraDataInterface*> CollectGraphInputDataInterfaces(UNiagaraSystem& System, ENiagaraScriptUsage Usage, FGuid UsageId, FString EmitterUniqueName) const { return TMap<FNiagaraGraphInputDataInterfaceKey, UNiagaraDataInterface*>(); }
+	virtual TMap<FName, UNiagaraDataInterface*> ComputeObjectNameMap(UNiagaraSystem& System, ENiagaraScriptUsage Usage, FGuid UsageId, FString EmitterUniqueName) const { return TMap<FName, UNiagaraDataInterface*>(); }
 
 	virtual FGuid GetCompileBaseId(ENiagaraScriptUsage InUsage, const FGuid& InUsageId) const { return FGuid(); }
 
diff --git a/Engine/Plugins/FX/Niagara/Source/Niagara/Private/NiagaraAsyncCompile.cpp b/Engine/Plugins/FX/Niagara/Source/Niagara/Private/NiagaraAsyncCompile.cpp
--- a/Engine/Plugins/FX/Niagara/Source/Niagara/Private/NiagaraAsyncCompile.cpp
+++ b/Engine/Plugins/FX/Niagara/Source/Niagara/Private/NiagaraAsyncCompile.cpp
@@ -880,19 +880,40 @@ void FNiagaraActiveCompilationDefault::Apply(const FNiagaraQueryCompilationOptio
 			TSharedPtr<FNiagaraVMExecutableData> ExeData = EmitterCompiledScriptPair.CompileResults;
 			UNiagaraScript* CompiledScript = EmitterCompiledScriptPair.CompiledScript;
 
-			// generate the GraphInputDataInterfaces from the source (from the duplicated data if available).
-			TMap<FNiagaraGraphInputDataInterfaceKey, UNiagaraDataInterface*> GraphInputDataInterfaces;
+			// generate the ObjectNameMap from the source (from the duplicated data if available).
+			TMap<FName, UNiagaraDataInterface*> ObjectNameMap;
 			if (AsyncTask->ComputedPrecompileDuplicateData.IsValid())
 			{
-				GraphInputDataInterfaces = AsyncTask->ComputedPrecompileDuplicateData->GetGraphInputDataInterfaces();
+				if (const UNiagaraScriptSourceBase* ScriptSource = AsyncTask->ComputedPrecompileDuplicateData->GetScriptSource())
+				{
+					ObjectNameMap = ScriptSource->ComputeObjectNameMap(*Options.System, CompiledScript->GetUsage(), CompiledScript->GetUsageId(), AsyncTask->UniqueEmitterName);
+				}
+				else
+				{
+					checkf(false, TEXT("Failed to get ScriptSource from ComputedPrecompileDuplicateData"));
+				}
+
+				// we need to replace the DI in the above generated map with the duplicates that we've created as a part of the duplication process
+				TMap<FName, UNiagaraDataInterface*> DuplicatedObjectNameMap = AsyncTask->ComputedPrecompileDuplicateData->GetObjectNameMap();
+				for (auto ObjectMapIt = ObjectNameMap.CreateIterator(); ObjectMapIt; ++ObjectMapIt)
+				{
+					if (UNiagaraDataInterface** Replacement = DuplicatedObjectNameMap.Find(ObjectMapIt.Key()))
+					{
+						ObjectMapIt.Value() = *Replacement;
+					}
+					else
+					{
+						ObjectMapIt.RemoveCurrent();
+					}
+				}
 			}
 			else
 			{
 				const UNiagaraScriptSourceBase* ScriptSource = CompiledScript->GetLatestSource();
-				GraphInputDataInterfaces = ScriptSource->CollectGraphInputDataInterfaces(*Options.System, CompiledScript->GetUsage(), CompiledScript->GetUsageId(), AsyncTask->UniqueEmitterName);
+				ObjectNameMap = ScriptSource->ComputeObjectNameMap(*Options.System, CompiledScript->GetUsage(), CompiledScript->GetUsageId(), AsyncTask->UniqueEmitterName);
 			}
 
-			EmitterCompiledScriptPair.CompiledScript->SetVMCompilationResults(EmitterCompiledScriptPair.CompileId, *ExeData, GraphInputDataInterfaces, bEvaluateParametersPending);
+			EmitterCompiledScriptPair.CompiledScript->SetVMCompilationResults(EmitterCompiledScriptPair.CompileId, *ExeData, AsyncTask->UniqueEmitterName, ObjectNameMap, bEvaluateParametersPending);
 
 			// Synchronize the variables that we actually encountered during precompile so that we can expose them to the end user.
 			FNiagaraUserRedirectionParameterStore& ExposedParameters = Options.System->GetExposedParameters();
diff --git a/Engine/Plugins/FX/Niagara/Source/Niagara/Private/NiagaraCommon.cpp b/Engine/Plugins/FX/Niagara/Source/Niagara/Private/NiagaraCommon.cpp
--- a/Engine/Plugins/FX/Niagara/Source/Niagara/Private/NiagaraCommon.cpp
+++ b/Engine/Plugins/FX/Niagara/Source/Niagara/Private/NiagaraCommon.cpp
@@ -1188,10 +1188,31 @@ FNiagaraVariable FNiagaraUtilities::ConvertVariableToRapidIterationConstantName(
 	return Var;
 }
 
+bool FNiagaraScriptDataInterfaceInfo::IsUserDataInterface() const
+{
+	TStringBuilder<128> NameBuilder;
+	Name.ToString(NameBuilder);
+	return FCString::Strnicmp(NameBuilder.ToString(), PARAM_MAP_USER_STR, 5) == 0;
+}
+
 bool FNiagaraScriptResolvedDataInterfaceInfo::NeedsPerInstanceBinding() const
 {
-	return ResolvedParameterReference.HasFlags(ENiagaraCompileParameterReferenceFlags::User) ||
-		(ResolvedDataInterface != nullptr && ResolvedDataInterface->PerInstanceDataSize() > 0);
+	FNameBuilder NameBuilder;
+	Name.ToString(NameBuilder);
+	return NameBuilder.ToView().StartsWith(PARAM_MAP_USER_STR) 
+		|| (ResolvedDataInterface != nullptr && ResolvedDataInterface->PerInstanceDataSize() > 0);
+}
+
+bool FNiagaraScriptDataInterfaceCompileInfo::CanExecuteOnTarget(ENiagaraSimTarget SimTarget) const
+{
+	// Note that this can be called on non-game threads. We ensure that the data interface CDO object is already in existence at application init time.
+	UNiagaraDataInterface* Obj = GetDefaultDataInterface();
+	if (Obj)
+	{
+		return Obj->CanExecuteOnTarget(SimTarget);
+	}
+	UE_LOG(LogNiagara, Error, TEXT("Failed to call CanExecuteOnTarget for DataInterface \"%s\". Perhaps missing a plugin for your project?"), *Name.ToString());
+	return false;
 }
 
 UNiagaraDataInterface* FNiagaraScriptDataInterfaceCompileInfo::GetDefaultDataInterface() const
@@ -1215,6 +1236,23 @@ UNiagaraDataInterface* FNiagaraScriptDataInterfaceCompileInfo::GetDefaultDataInt
 	return nullptr;
 }
 
+bool FNiagaraScriptDataInterfaceCompileInfo::NeedsPerInstanceBinding()const
+{
+	FNameBuilder NameBuilder;
+	Name.ToString(NameBuilder);
+	if (NameBuilder.ToView().StartsWith(TEXT("User.")))
+	{
+		return true;
+	}
+
+	UNiagaraDataInterface* Obj = GetDefaultDataInterface();
+	if (Obj && Obj->PerInstanceDataSize() > 0)
+	{
+		return true;
+	}
+	return false;
+}
+
 bool FNiagaraScriptDataInterfaceCompileInfo::MatchesClass(const UClass* InClass) const
 {
 	UNiagaraDataInterface* Obj = GetDefaultDataInterface();
diff --git a/Engine/Plugins/FX/Niagara/Source/Niagara/Private/NiagaraCompilationTypes.cpp b/Engine/Plugins/FX/Niagara/Source/Niagara/Private/NiagaraCompilationTypes.cpp
--- a/Engine/Plugins/FX/Niagara/Source/Niagara/Private/NiagaraCompilationTypes.cpp
+++ b/Engine/Plugins/FX/Niagara/Source/Niagara/Private/NiagaraCompilationTypes.cpp
@@ -87,15 +87,13 @@ bool operator==(const FVMExternalFunctionBindingInfo& Lhs, const FVMExternalFunc
 
 bool operator==(const FNiagaraScriptDataInterfaceCompileInfo& Lhs, const FNiagaraScriptDataInterfaceCompileInfo& Rhs)
 {
-	return Lhs.Id == Rhs.Id
-		&& Lhs.Name == Rhs.Name
+	return Lhs.Name == Rhs.Name
 		&& Lhs.UserPtrIdx == Rhs.UserPtrIdx
 		&& Lhs.Type == Rhs.Type
 		&& Lhs.RegisteredFunctions == Rhs.RegisteredFunctions
-		&& Lhs.SourceGraphInputDataInterfaceKey == Rhs.SourceGraphInputDataInterfaceKey
-		&& Lhs.SourceExternalParameter == Rhs.SourceExternalParameter
-		&& Lhs.SourceScriptDIIndex == Rhs.SourceScriptDIIndex
-		&& Lhs.TargetParameter == Rhs.TargetParameter;
+		&& Lhs.RegisteredParameterMapRead == Rhs.RegisteredParameterMapRead
+		&& Lhs.RegisteredParameterMapWrite == Rhs.RegisteredParameterMapWrite
+		&& Lhs.bIsPlaceholder == Rhs.bIsPlaceholder;
 }
 
 bool operator==(const FNiagaraShaderScriptParametersMetadata& Lhs, const FNiagaraShaderScriptParametersMetadata& Rhs)
@@ -334,7 +332,7 @@ public:
 					FNiagaraVMExecutableDataId UpdatedCompileId;
 					TargetScript->ComputeVMCompilationId(UpdatedCompileId, FGuid());
 
-					TMap<FNiagaraGraphInputDataInterfaceKey, UNiagaraDataInterface*> GraphInputDataInterfaces;
+					TMap<FName, UNiagaraDataInterface*> ObjectNameMap;
 
 					// The original implementation would generate DI references from the compilation data (unless things were pulled from
 					// the DDC).  We will always pull the data from the target scripts so that we can avoid any weird caching issues with
@@ -343,14 +341,15 @@ public:
 					// when the compilation results are applied.
 					if (const UNiagaraScriptSourceBase* ScriptSource = TargetScript->GetLatestSource())
 					{
-						GraphInputDataInterfaces = ScriptSource->CollectGraphInputDataInterfaces(*Options.System, TargetScript->GetUsage(), TargetScript->GetUsageId(), ScriptCompileData.UniqueEmitterName);
+						ObjectNameMap = ScriptSource->ComputeObjectNameMap(*Options.System, TargetScript->GetUsage(), TargetScript->GetUsageId(), ScriptCompileData.UniqueEmitterName);
 					}
 
 					constexpr bool bApplyRapidIterationParameters = false;
 					TargetScript->SetVMCompilationResults(
 						UpdatedCompileId,
 						*ScriptCompileData.ExeData,
-						GraphInputDataInterfaces,
+						ScriptCompileData.UniqueEmitterName,
+						ObjectNameMap,
 						bApplyRapidIterationParameters);
 
 					if (!ScriptCompileData.CompiledShaders.IsEmpty())
diff --git a/Engine/Plugins/FX/Niagara/Source/Niagara/Private/NiagaraConstants.cpp b/Engine/Plugins/FX/Niagara/Source/Niagara/Private/NiagaraConstants.cpp
--- a/Engine/Plugins/FX/Niagara/Source/Niagara/Private/NiagaraConstants.cpp
+++ b/Engine/Plugins/FX/Niagara/Source/Niagara/Private/NiagaraConstants.cpp
@@ -62,6 +62,7 @@ const FString FNiagaraConstants::RapidIterationParametersNamespaceString(TEXT("C
 const FString FNiagaraConstants::StackContextNamespaceString(TEXT("StackContext"));
 const FString FNiagaraConstants::SystemNamespaceString(TEXT("System"));
 const FString FNiagaraConstants::UserNamespaceString(TEXT("User"));
+const FString FNiagaraConstants::InternalNamespaceString(TEXT("__INTERNAL__"));
 const FString FNiagaraConstants::PreviousNamespaceString(TEXT("Previous"));
 
 const FName FNiagaraConstants::InputScopeName(TEXT("Input"));
diff --git a/Engine/Plugins/FX/Niagara/Source/Niagara/Private/NiagaraDataInterfaceEmitterBinding.cpp b/Engine/Plugins/FX/Niagara/Source/Niagara/Private/NiagaraDataInterfaceEmitterBinding.cpp
--- a/Engine/Plugins/FX/Niagara/Source/Niagara/Private/NiagaraDataInterfaceEmitterBinding.cpp
+++ b/Engine/Plugins/FX/Niagara/Source/Niagara/Private/NiagaraDataInterfaceEmitterBinding.cpp
@@ -5,7 +5,6 @@
 #include "NiagaraEmitterInstance.h"
 #include "NiagaraScriptRuntimeData.h"
 #include "NiagaraSystem.h"
-#include "NiagaraSystemImpl.h"
 #include "NiagaraSystemInstance.h"
 
 #include UE_INLINE_GENERATED_CPP_BY_NAME(NiagaraDataInterfaceEmitterBinding)
@@ -28,7 +27,7 @@ FNiagaraEmitterInstance* FNiagaraDataInterfaceEmitterBinding::Resolve(const FNia
 		// If we came from an emitter script we need to look over the default data interfaces to find ourself
 		else
 		{
-			FGuid SourceEmitterHandleId;
+			FString SourceEmitterName;
 			UNiagaraSystem* NiagaraSystem = SystemInstance->GetSystem();
 			for (UNiagaraScript* NiagaraScript : { NiagaraSystem->GetSystemUpdateScript(), NiagaraSystem->GetSystemSpawnScript() })
 			{
@@ -39,23 +38,26 @@ FNiagaraEmitterInstance* FNiagaraDataInterfaceEmitterBinding::Resolve(const FNia
 					{
 						continue;
 					}
-					const FNiagaraScriptResolvedDataInterfaceInfo* MatchingResolvedDataInterfaceInfo = ScriptRuntimeData->GetResolvedDataInterfaces().FindByPredicate([&DataInterface]
-						(const FNiagaraScriptResolvedDataInterfaceInfo& ResolvedDataInterfaceInfo) { return ResolvedDataInterfaceInfo.ResolvedDataInterface == DataInterface; });
-					if (MatchingResolvedDataInterfaceInfo != nullptr)
+					const FNiagaraScriptResolvedDataInterfaceInfo* MatchingResolvedDataInterface = ScriptRuntimeData->GetResolvedDataInterfaces().FindByPredicate(
+						[&DataInterface](const FNiagaraScriptResolvedDataInterfaceInfo& ResolvedDataInterface) { return ResolvedDataInterface.ResolvedDataInterface == DataInterface; });
+					if (MatchingResolvedDataInterface != nullptr)
 					{
-						SourceEmitterHandleId = MatchingResolvedDataInterfaceInfo->OwningEmitterHandleId;
+						SourceEmitterName = MatchingResolvedDataInterface->ResolvedSourceEmitterName;
 						break;
 					}
 				}
 			}
 
-			if (SourceEmitterHandleId.IsValid())
+			if (SourceEmitterName.IsEmpty() == false)
 			{
 				for (const FNiagaraEmitterInstanceRef& EmitterInstance : SystemInstance->GetEmitters())
 				{
-					if (EmitterInstance->GetEmitterHandle().GetId() == SourceEmitterHandleId)
+					if (const UNiagaraEmitter* CachedEmitter = EmitterInstance->GetEmitter())
 					{
-						return &EmitterInstance.Get();
+						if (CachedEmitter->GetUniqueEmitterName() == SourceEmitterName)
+						{
+							return &EmitterInstance.Get();
+						}
 					}
 				}
 			}
@@ -123,51 +125,83 @@ const FNiagaraEmitterHandle* FNiagaraDataInterfaceEmitterBinding::ResolveHandle(
 		}
 		else
 		{
-			// Try to find this data interface in one of the scripts resolved data interfaces.
-			const FNiagaraEmitterHandle* FoundEmitterHandle = nullptr;
-			OwnerSystem->ForEachScriptWithOwningContext([DataInterface, &FoundEmitterHandle](const UNiagaraSystem* OwningSystem, const FNiagaraEmitterHandle* OwningEmitterHandle, UNiagaraScript* Script)
+			// Try to find this data interface in one of the emitter's particle scripts.
+			for (const FNiagaraEmitterHandle& EmitterHandle : OwnerSystem->GetEmitterHandles())
 			{
-				if (FoundEmitterHandle != nullptr)
+				FVersionedNiagaraEmitterData* EmitterData = EmitterHandle.GetEmitterData();
+				if (EmitterHandle.GetIsEnabled() && EmitterData)
 				{
-					return;
-				}
+					bool bFoundDI = false;
+					bool bIsEmitterNamespace = false;
+					EmitterData->ForEachScript(
+						[OwnerSystem, &EmitterHandle, DataInterface, &bFoundDI, &bIsEmitterNamespace](const UNiagaraScript* Script)
+						{
+							const FNiagaraScriptRuntimeData* ScriptRuntimeData = OwnerSystem->GetScriptRuntimeData(&EmitterHandle, *Script).Get();
+							if (ScriptRuntimeData == nullptr)
+							{
+								return;
+							}
+							for (const FNiagaraScriptResolvedDataInterfaceInfo& ResolvedDI : ScriptRuntimeData->GetResolvedDataInterfaces())
+							{
+								if (ResolvedDI.ResolvedDataInterface == DataInterface)
+								{
+									bFoundDI = true;
+									bIsEmitterNamespace = FNiagaraVariableBase::IsInNameSpace(FNiagaraConstants::EmitterNamespaceString, ResolvedDI.CompileName);
+									return;
+								}
+							}
+						}
+					);
 
-				FGuid EmitterHandleId = OwningEmitterHandle != nullptr ? OwningEmitterHandle->GetId() : FGuid();
-				const FNiagaraScriptRuntimeData* ScriptRuntimeData = OwningSystem->GetScriptRuntimeData(EmitterHandleId, Script->GetUsage(), Script->GetUsageId()).Get();
+					if (bFoundDI)
+					{
+						return bIsEmitterNamespace ? &EmitterHandle : nullptr;
+					}
+				}
+			}
+			// Check the system scripts since they contain the compiled emitter data.
+			TArray<UNiagaraScript*> SystemScripts = { OwnerSystem->GetSystemSpawnScript(), OwnerSystem->GetSystemUpdateScript() };
+			for (UNiagaraScript* SystemScript : SystemScripts)
+			{
+				const FNiagaraScriptRuntimeData* ScriptRuntimeData = OwnerSystem->GetScriptRuntimeData(nullptr, *SystemScript).Get();
 				if (ScriptRuntimeData == nullptr)
 				{
-					return;
+					continue;
 				}
-
 				for (const FNiagaraScriptResolvedDataInterfaceInfo& ResolvedDI : ScriptRuntimeData->GetResolvedDataInterfaces())
 				{
 					if (ResolvedDI.ResolvedDataInterface == DataInterface)
 					{
-						if (OwningEmitterHandle != nullptr)
-						{
-							FoundEmitterHandle = OwningEmitterHandle;
-						}
-						else if (ResolvedDI.OwningEmitterHandleId.IsValid())
+						if (ResolvedDI.ResolvedSourceEmitterName.IsEmpty() == false)
 						{
-							FoundEmitterHandle = OwningSystem->GetEmitterHandles().FindByPredicate([&ResolvedDI](const FNiagaraEmitterHandle& EmitterHandle)
-								{ return EmitterHandle.GetId() == ResolvedDI.OwningEmitterHandleId; });
-						}
-						if (FoundEmitterHandle != nullptr)
-						{
-							break;
+							return OwnerSystem->GetEmitterHandles().FindByPredicate([&ResolvedDI](const FNiagaraEmitterHandle& EmitterHandle)
+								{ return EmitterHandle.GetUniqueInstanceName() == ResolvedDI.ResolvedSourceEmitterName; });
 						}
 					}
 				}
-			});
-			if (FoundEmitterHandle != nullptr)
-			{
-				return FoundEmitterHandle;
 			}
 		}
 	}
 	else if (BindingMode == ENiagaraDataInterfaceEmitterBindingMode::Other)
 	{
-		return ResolveHandleOtherBinding(OwnerSystem);
+		if (!EmitterName.IsNone())
+		{
+			FNameBuilder EmitterNameString;
+			EmitterName.ToString(EmitterNameString);
+			FStringView EmitterNameStringView = EmitterNameString.ToView();
+
+			for (const FNiagaraEmitterHandle& EmitterHandle : OwnerSystem->GetEmitterHandles())
+			{
+				if (UNiagaraEmitter* NiagaraEmitter = EmitterHandle.GetInstance().Emitter)
+				{
+					//-TODO: UniqueEmitterName should probably be a FName?
+					if (EmitterNameStringView.Equals(NiagaraEmitter->GetUniqueEmitterName(), ESearchCase::IgnoreCase))
+					{
+						return &EmitterHandle;
+					}
+				}
+			}
+		}
 	}
 	return nullptr;
 }
@@ -182,33 +216,21 @@ const FNiagaraEmitterHandle* FNiagaraDataInterfaceEmitterBinding::ResolveHandle(
 	}
 	else if (BindingMode == ENiagaraDataInterfaceEmitterBindingMode::Other)
 	{
-		return ResolveHandleOtherBinding(OwnerSystem);
-	}
-	return nullptr;
-}
-
-UNiagaraEmitter* FNiagaraDataInterfaceEmitterBinding::Resolve(const UNiagaraDataInterface* DataInterface) const
-{
-	const FNiagaraEmitterHandle* EmitterHandle = ResolveHandle(DataInterface);
-	return EmitterHandle ? EmitterHandle->GetInstance().Emitter : nullptr;
-}
-
-const FNiagaraEmitterHandle* FNiagaraDataInterfaceEmitterBinding::ResolveHandleOtherBinding(const UNiagaraSystem* OwnerSystem) const
-{
-	if (!EmitterName.IsNone())
-	{
-		FNameBuilder EmitterNameString;
-		EmitterName.ToString(EmitterNameString);
-		FStringView EmitterNameStringView = EmitterNameString.ToView();
-
-		for (const FNiagaraEmitterHandle& EmitterHandle : OwnerSystem->GetEmitterHandles())
+		if (!EmitterName.IsNone())
 		{
-			if (UNiagaraEmitter* NiagaraEmitter = EmitterHandle.GetInstance().Emitter)
+			FNameBuilder EmitterNameString;
+			EmitterName.ToString(EmitterNameString);
+			FStringView EmitterNameStringView = EmitterNameString.ToView();
+
+			for (const FNiagaraEmitterHandle& EmitterHandle : OwnerSystem->GetEmitterHandles())
 			{
-				//-TODO: UniqueEmitterName should probably be a FName?
-				if (EmitterNameStringView.Equals(NiagaraEmitter->GetUniqueEmitterName(), ESearchCase::IgnoreCase))
+				if (UNiagaraEmitter* NiagaraEmitter = EmitterHandle.GetInstance().Emitter)
 				{
-					return &EmitterHandle;
+					//-TODO: UniqueEmitterName should probably be a FName?
+					if (EmitterNameStringView.Equals(NiagaraEmitter->GetUniqueEmitterName(), ESearchCase::IgnoreCase))
+					{
+						return &EmitterHandle;
+					}
 				}
 			}
 		}
@@ -216,6 +238,11 @@ const FNiagaraEmitterHandle* FNiagaraDataInterfaceEmitterBinding::ResolveHandleO
 	return nullptr;
 }
 
+UNiagaraEmitter* FNiagaraDataInterfaceEmitterBinding::Resolve(const UNiagaraDataInterface* DataInterface) const
+{
+	const FNiagaraEmitterHandle* EmitterHandle = ResolveHandle(DataInterface);
+	return EmitterHandle ? EmitterHandle->GetInstance().Emitter : nullptr;
+}
 
 FString FNiagaraDataInterfaceEmitterBinding::ResolveUniqueName(const UNiagaraDataInterface* DataInterface) const
 {
diff --git a/Engine/Plugins/FX/Niagara/Source/Niagara/Private/NiagaraDataInterfaceGrid2DCollectionReader.cpp b/Engine/Plugins/FX/Niagara/Source/Niagara/Private/NiagaraDataInterfaceGrid2DCollectionReader.cpp
--- a/Engine/Plugins/FX/Niagara/Source/Niagara/Private/NiagaraDataInterfaceGrid2DCollectionReader.cpp
+++ b/Engine/Plugins/FX/Niagara/Source/Niagara/Private/NiagaraDataInterfaceGrid2DCollectionReader.cpp
@@ -81,24 +81,21 @@ bool UNiagaraDataInterfaceGrid2DCollectionReader::InitPerInstanceData(void* PerI
 		{
 			TConstArrayView<FNiagaraScriptResolvedDataInterfaceInfo> DataInterfaceInfo = ExecContext->GetScriptRuntimeData()->GetResolvedDataInterfaces();
 			const TArray<UNiagaraDataInterface*>& DataInterfaces = ExecContext->CombinedParamStore.GetDataInterfaces();
-			if (DataInterfaceInfo.Num() != DataInterfaces.Num())
-			{
-				UE_LOG(LogNiagara, Error, TEXT("Grid2DCollectionReader failed to bind to collection for Emitter: %s and DataInterface: %s.  Resolved data interface count didn't match exec context data interface count."), *this->EmitterName, *this->DIName);
-				return false;
-			}
 
 			FString FullName = FString("Emitter.") + DIName;
-			for (int32 Index = 0; Index < DataInterfaces.Num(); ++Index)
+			int Index = 0;
+
+			// #todo(dmp): we are looking at the UObjects that define the DIs here 
+
+			for (UNiagaraDataInterface* Interface : DataInterfaces)
 			{
-				UNiagaraDataInterface* Interface = DataInterfaces[Index];
-				const FNiagaraScriptResolvedDataInterfaceInfo& Info = DataInterfaceInfo[Index];
-				if (Info.TargetParameterReference.ParameterFullName.GetPlainNameString() == FullName ||
-					Info.ResolvedParameterReference.ParameterFullName.GetPlainNameString() == FullName)
+				if (DataInterfaceInfo[Index].CompileName.GetPlainNameString() == FullName)
 				{
-					GridInterface = Cast<UNiagaraDataInterfaceGrid2DCollection>(Interface);
+					GridInterface = static_cast<UNiagaraDataInterfaceGrid2DCollection*>(Interface);
 					ProxyToUse = Interface->GetProxy();
 					break;
 				}
+				++Index;
 			}
 		}
 
diff --git a/Engine/Plugins/FX/Niagara/Source/Niagara/Private/NiagaraDataInterfaceGrid3DCollectionReader.cpp b/Engine/Plugins/FX/Niagara/Source/Niagara/Private/NiagaraDataInterfaceGrid3DCollectionReader.cpp
--- a/Engine/Plugins/FX/Niagara/Source/Niagara/Private/NiagaraDataInterfaceGrid3DCollectionReader.cpp
+++ b/Engine/Plugins/FX/Niagara/Source/Niagara/Private/NiagaraDataInterfaceGrid3DCollectionReader.cpp
@@ -80,24 +80,21 @@ bool UNiagaraDataInterfaceGrid3DCollectionReader::InitPerInstanceData(void* PerI
 		{
 			TConstArrayView<FNiagaraScriptResolvedDataInterfaceInfo> DataInterfaceInfo = ExecContext->GetScriptRuntimeData()->GetResolvedDataInterfaces();
 			const TArray<UNiagaraDataInterface*>& DataInterfaces = ExecContext->CombinedParamStore.GetDataInterfaces();
-			if (DataInterfaceInfo.Num() != DataInterfaces.Num())
-			{
-				UE_LOG(LogNiagara, Error, TEXT("Grid3DCollectionReader failed to bind to collection for Emitter: %s and DataInterface: %s.  Resolved data interface count didn't match exec context data interface count."), *this->EmitterName, *this->DIName);
-				return false;
-			}
 
 			FString FullName = FString("Emitter.") + DIName;
-			for (int32 Index = 0; Index < DataInterfaces.Num(); ++Index)
+			int Index = 0;
+
+			// #todo(dmp): we are looking at the UObjects that define the DIs here 
+			
+			for (UNiagaraDataInterface* Interface : DataInterfaces)
 			{
-				UNiagaraDataInterface* Interface = DataInterfaces[Index];
-				const FNiagaraScriptResolvedDataInterfaceInfo& Info = DataInterfaceInfo[Index];
-				if (Info.TargetParameterReference.ParameterFullName.GetPlainNameString() == FullName ||
-					Info.ResolvedParameterReference.ParameterFullName.GetPlainNameString() == FullName)
+				if (DataInterfaceInfo[Index].CompileName.GetPlainNameString() == FullName)
 				{
-					GridInterface = Cast<UNiagaraDataInterfaceGrid3DCollection>(Interface);
+					GridInterface = static_cast<UNiagaraDataInterfaceGrid3DCollection*>(Interface);
 					ProxyToUse = Interface->GetProxy();
 					break;
 				}
+				++Index;
 			}
 		}
 
diff --git a/Engine/Plugins/FX/Niagara/Source/Niagara/Private/NiagaraDataInterfaceParticleRead.cpp b/Engine/Plugins/FX/Niagara/Source/Niagara/Private/NiagaraDataInterfaceParticleRead.cpp
--- a/Engine/Plugins/FX/Niagara/Source/Niagara/Private/NiagaraDataInterfaceParticleRead.cpp
+++ b/Engine/Plugins/FX/Niagara/Source/Niagara/Private/NiagaraDataInterfaceParticleRead.cpp
@@ -8,7 +8,6 @@
 #include "NiagaraGPUSystemTick.h"
 #include "NiagaraRenderer.h"
 #include "NiagaraSimStageData.h"
-#include "NiagaraSystemImpl.h"
 #include "NiagaraSystemInstance.h"
 #include "NiagaraShaderParametersBuilder.h"
 
@@ -2408,24 +2407,20 @@ void UNiagaraDataInterfaceParticleRead::GetFeedback(UNiagaraSystem* Asset, UNiag
 		return false;
 	}();
 
-	Asset->ForEachScriptWithRuntimeData([&](const UNiagaraScript* Script, const FNiagaraScriptRuntimeData* ScriptRuntimeData)
+
+	for (const auto Script : Scripts)
 	{
-		const TArray<FNiagaraScriptDataInterfaceCompileInfo>& CompileInfos = Script->GetVMExecutableData().DataInterfaceInfo;
-		TConstArrayView<FNiagaraScriptResolvedDataInterfaceInfo> ResolvedInfos = ScriptRuntimeData->GetResolvedDataInterfaces();
-		if (CompileInfos.Num() != ResolvedInfos.Num())
-		{
-			return;
-		}
+		const TArray<FNiagaraScriptDataInterfaceInfo>& CachedDefaultDIs = Script->GetCachedDefaultDataInterfaces();
 
-		for (int32 Idx = 0; Idx < CompileInfos.Num(); Idx++)
+		for (int32 Idx = 0; Idx < Script->GetVMExecutableData().DataInterfaceInfo.Num(); Idx++)
 		{
-			const auto& DIInfo = CompileInfos[Idx];
-			if (DIInfo.Type.GetClass() != GetClass())
+			const auto& DIInfo = Script->GetVMExecutableData().DataInterfaceInfo[Idx];
+			if (!DIInfo.MatchesClass(GetClass()))
 			{
 				continue;
 			}
 
-			if (ResolvedInfos[Idx].ResolvedDataInterface == nullptr || !ResolvedInfos[Idx].ResolvedDataInterface->Equals(this))
+			if (!CachedDefaultDIs.IsValidIndex(Idx) || CachedDefaultDIs[Idx].DataInterface == nullptr || !CachedDefaultDIs[Idx].DataInterface->Equals(this))
 			{
 				continue;
 			}
@@ -2493,7 +2488,7 @@ void UNiagaraDataInterfaceParticleRead::GetFeedback(UNiagaraSystem* Asset, UNiag
 				}
 			}
 		}
-	});
+	}
 
 	// If we found persistent ID functions in use and the target emitter isn't set to expose them, trigger a fixable warning.
 	FVersionedNiagaraEmitterData* EmitterData = FoundSourceEmitter.GetEmitterData();
diff --git a/Engine/Plugins/FX/Niagara/Source/Niagara/Private/NiagaraDataInterfaceVectorField.cpp b/Engine/Plugins/FX/Niagara/Source/Niagara/Private/NiagaraDataInterfaceVectorField.cpp
--- a/Engine/Plugins/FX/Niagara/Source/Niagara/Private/NiagaraDataInterfaceVectorField.cpp
+++ b/Engine/Plugins/FX/Niagara/Source/Niagara/Private/NiagaraDataInterfaceVectorField.cpp
@@ -291,7 +291,7 @@ void UNiagaraDataInterfaceVectorField::GetFeedback(UNiagaraSystem* InAsset, UNia
 					{
 						// We assume that if the properties match or we are referencing an external variable whose name is in the list of candidates found in the prior search, it's a valid match for us.
 						if (ResolvedDIs.IsValidIndex(Idx) && ResolvedDIs[Idx].ResolvedDataInterface != nullptr &&
-							(ResolvedDIs[Idx].ResolvedDataInterface->Equals(this) || DIAliases.Contains(ResolvedDIs[Idx].ResolvedVariable.GetName())))
+							(ResolvedDIs[Idx].ResolvedDataInterface->Equals(this) || DIAliases.Contains(ResolvedDIs[Idx].CompileName)))
 						{
 							FVersionedNiagaraEmitter OuterEmitter = Script->GetOuterEmitter();
 							if (OuterEmitter.Emitter && OuterEmitter.GetEmitterData() && (OuterEmitter.GetEmitterData()->SimTarget == ENiagaraSimTarget::CPUSim || Script->IsSystemScript(Script->Usage)))
diff --git a/Engine/Plugins/FX/Niagara/Source/Niagara/Private/NiagaraEmitterInstanceImpl.cpp b/Engine/Plugins/FX/Niagara/Source/Niagara/Private/NiagaraEmitterInstanceImpl.cpp
--- a/Engine/Plugins/FX/Niagara/Source/Niagara/Private/NiagaraEmitterInstanceImpl.cpp
+++ b/Engine/Plugins/FX/Niagara/Source/Niagara/Private/NiagaraEmitterInstanceImpl.cpp
@@ -15,7 +15,6 @@
 #include "NiagaraSimulationStageBase.h"
 #include "NiagaraComponentSettings.h"
 #include "NiagaraSystemSimulation.h"
-#include "NiagaraScriptRuntimeData.h"
 
 DECLARE_DWORD_COUNTER_STAT(TEXT("Num Custom Events"), STAT_NiagaraNumCustomEvents, STATGROUP_Niagara);
 
@@ -664,26 +663,30 @@ void FNiagaraEmitterInstanceImpl::UnbindParameters(bool bExternalOnly)
 {
 	if (bExternalOnly && !IsDisabled())
 	{
-		ParentSystemInstance->UnbindScriptParameterCollections(SpawnExecContext.Script, SpawnExecContext.Parameters);
-		ParentSystemInstance->UnbindScriptParameterCollections(UpdateExecContext.Script, UpdateExecContext.Parameters);
-
-		FVersionedNiagaraEmitterData* EmitterData = VersionedEmitter.GetEmitterData();
-		if (EmitterData->SimTarget == ENiagaraSimTarget::GPUComputeSim)
+		for (UNiagaraParameterCollection* Collection : SpawnExecContext.Script->GetCachedParameterCollectionReferences())
 		{
-			ParentSystemInstance->UnbindScriptParameterCollections(SpawnExecContext.Script, GPUExecContext->CombinedParamStore);
-			ParentSystemInstance->UnbindScriptParameterCollections(UpdateExecContext.Script, GPUExecContext->CombinedParamStore);
-			for (const UNiagaraSimulationStageBase* SimStage : EmitterData->GetSimulationStages())
+			if (UNiagaraParameterCollectionInstance* NPCInst = ParentSystemInstance->GetParameterCollectionInstance(Collection))
 			{
-				if (SimStage->bEnabled)
-				{
-					ParentSystemInstance->UnbindScriptParameterCollections(SimStage->Script, GPUExecContext->CombinedParamStore);
-				}
+				NPCInst->GetParameterStore().Unbind(&SpawnExecContext.Parameters);
+			}
+		}
+		for (UNiagaraParameterCollection* Collection : UpdateExecContext.Script->GetCachedParameterCollectionReferences())
+		{
+			if (UNiagaraParameterCollectionInstance* NPCInst = ParentSystemInstance->GetParameterCollectionInstance(Collection))
+			{
+				NPCInst->GetParameterStore().Unbind(&UpdateExecContext.Parameters);
 			}
 		}
 
 		for (FNiagaraScriptExecutionContext& EventContext : GetEventExecutionContexts())
 		{
-			ParentSystemInstance->UnbindScriptParameterCollections(EventContext.Script, EventContext.Parameters);
+			for (UNiagaraParameterCollection* Collection : EventContext.Script->GetCachedParameterCollectionReferences())
+			{
+				if (UNiagaraParameterCollectionInstance* NPCInst = ParentSystemInstance->GetParameterCollectionInstance(Collection))
+				{
+					NPCInst->GetParameterStore().Unbind(&EventContext.Parameters);
+				}
+			}
 		}
 
 		ParentSystemInstance->UnbindFromParameterStore(RendererBindings);
@@ -713,26 +716,67 @@ void FNiagaraEmitterInstanceImpl::BindParameters(bool bExternalOnly)
 		return;
 	}
 
-	ParentSystemInstance->BindScriptParameterCollections(SpawnExecContext.Script, SpawnExecContext.GetScriptRuntimeData(), SpawnExecContext.Parameters);
-	ParentSystemInstance->BindScriptParameterCollections(UpdateExecContext.Script, UpdateExecContext.GetScriptRuntimeData(), UpdateExecContext.Parameters);
+	auto BindToParameterCollection = [&](UNiagaraParameterCollection* Collection, FNiagaraParameterStore& DestStore)
+	{
+		if (Collection)
+		{
+			if (UNiagaraParameterCollectionInstance* Inst = ParentSystemInstance->GetParameterCollectionInstance(Collection))
+			{
+				Inst->GetParameterStore().Bind(&DestStore);
+			}
+			else
+			{
+				UE_LOG(LogNiagara, Error, TEXT("Emitter attempting to bind to a null Parameter Collection Instance.\nEmitter:%s\nCollection:%s"), VersionedEmitter.Emitter ? *VersionedEmitter.Emitter->GetPathName() : TEXT("null emitter!"), *Collection->GetPathName());
+			}
+		}
+		else
+		{
+			UE_LOG(LogNiagara, Error, TEXT("Emitter attempting to bind to a null Parameter Collection.\nEmitter:%s"), VersionedEmitter.Emitter ? *VersionedEmitter.Emitter->GetPathName() : TEXT("null emitter!" ));
+		}
+	};
+
+	for (UNiagaraParameterCollection* Collection : SpawnExecContext.Script->GetCachedParameterCollectionReferences())
+	{
+		BindToParameterCollection(Collection, SpawnExecContext.Parameters);
+	}
+	for (UNiagaraParameterCollection* Collection : UpdateExecContext.Script->GetCachedParameterCollectionReferences())
+	{
+		BindToParameterCollection(Collection, UpdateExecContext.Parameters);
+	}
 
 	FVersionedNiagaraEmitterData* EmitterData = VersionedEmitter.GetEmitterData();
 	if (EmitterData->SimTarget == ENiagaraSimTarget::GPUComputeSim)
 	{
-		ParentSystemInstance->BindScriptParameterCollections(SpawnExecContext.Script, GPUExecContext->GetScriptRuntimeData(), GPUExecContext->CombinedParamStore);
-		ParentSystemInstance->BindScriptParameterCollections(UpdateExecContext.Script, GPUExecContext->GetScriptRuntimeData(), GPUExecContext->CombinedParamStore);
-		for (const UNiagaraSimulationStageBase* SimStage : EmitterData->GetSimulationStages())
+		for (UNiagaraParameterCollection* Collection : SpawnExecContext.Script->GetCachedParameterCollectionReferences())
 		{
-			if (SimStage->bEnabled)
-			{
-				ParentSystemInstance->BindScriptParameterCollections(SimStage->Script, GPUExecContext->GetScriptRuntimeData(), GPUExecContext->CombinedParamStore);
-			}
+			BindToParameterCollection(Collection, GPUExecContext->CombinedParamStore);
+		}
+		for (UNiagaraParameterCollection* Collection : UpdateExecContext.Script->GetCachedParameterCollectionReferences())
+		{
+			BindToParameterCollection(Collection, GPUExecContext->CombinedParamStore);
 		}
 	}
 
 	for (FNiagaraScriptExecutionContext& EventContext : GetEventExecutionContexts())
 	{
-		ParentSystemInstance->BindScriptParameterCollections(EventContext.Script, EventContext.GetScriptRuntimeData(), EventContext.Parameters);
+		for (UNiagaraParameterCollection* Collection : EventContext.Script->GetCachedParameterCollectionReferences())
+		{
+			BindToParameterCollection(Collection, EventContext.Parameters);
+		}
+	}
+
+	FNiagaraScriptInstanceParameterStore& TargetParamStore = EmitterData->SimTarget == ENiagaraSimTarget::GPUComputeSim ? GPUExecContext->CombinedParamStore : UpdateExecContext.Parameters;
+	for (const UNiagaraSimulationStageBase* SimStage : EmitterData->GetSimulationStages())
+	{
+		if (SimStage->bEnabled == false)
+		{
+			continue;
+		}
+
+		for (UNiagaraParameterCollection* Collection : SimStage->Script->GetCachedParameterCollectionReferences())
+		{
+			BindToParameterCollection(Collection, TargetParamStore);
+		}
 	}
 
 	if (!bExternalOnly)
@@ -1417,6 +1461,8 @@ void FNiagaraEmitterInstanceImpl::Tick(float DeltaSeconds)
 			}
 		}
 
+		//GPUExecContext.UpdateInterfaces = VersionedEmitter->UpdateScriptProps.Script->GetCachedDefaultDataInterfaces();
+
 		// copy over the constants for the render thread
 		//
 		if (GbDumpParticleData || System->bDumpDebugEmitterInfo)
diff --git a/Engine/Plugins/FX/Niagara/Source/Niagara/Private/NiagaraResolveDIHelpers.cpp b/Engine/Plugins/FX/Niagara/Source/Niagara/Private/NiagaraResolveDIHelpers.cpp
--- a/Engine/Plugins/FX/Niagara/Source/Niagara/Private/NiagaraResolveDIHelpers.cpp
+++ b/Engine/Plugins/FX/Niagara/Source/Niagara/Private/NiagaraResolveDIHelpers.cpp
@@ -9,7 +9,6 @@
 #include "NiagaraScript.h"
 #include "NiagaraSimulationStageBase.h"
 #include "NiagaraSystem.h"
-#include "NiagaraSystemImpl.h"
 
 #if WITH_EDITORONLY_DATA
 
@@ -17,520 +16,477 @@
 
 namespace FNiagaraResolveDIHelpers
 {
-	struct FScriptDefaultData
+	struct FDataInterfaceSourceEmitterNamePair
 	{
-		TMap<FNiagaraScriptDataInterfaceCompileId, UNiagaraDataInterface*> InternalDefaultMap;
-		TMap<FNiagaraCompileParameterReference, UNiagaraDataInterface*> ParameterDefaultMap;
+		UNiagaraDataInterface* DataInterface;
+		FName SourceName;
+		FString SourceEmitterName;
 	};
 
-	/** Collects all valid DI instances from a set of data interface compile infos and maps them by compile id or parameter reference,
-		depending on the source. */
-	void CollectScriptInternalDefaultDIs(
-		const TArray<FNiagaraScriptDataInterfaceCompileInfo>& ScriptDataInterfaceCompileInfo,
-		TConstArrayView<FNiagaraScriptDataInterfaceDefaultInfo> ScriptCachedDefaultDataInterfaces,
-		FScriptDefaultData& OutScriptDefaultData)
+	void CollectDIReadsAssignmentsAndBindingsForScript(
+		const UNiagaraScript* TargetScript,
+		const UNiagaraSystem* System,
+		const FString& EmitterName,
+		TMap<FNiagaraVariableBase, FDataInterfaceSourceEmitterNamePair>& OutVariableReadMap,
+		TMap<FNiagaraVariableBase, FDataInterfaceSourceEmitterNamePair>& OutVariableAssignmentMap,
+		TMap<FNiagaraVariableBase, FNiagaraVariableBase>& OutVariableBindingMap,
+		TArray<FText>& OutErrorMessages)
 	{
-		for (const FNiagaraScriptDataInterfaceCompileInfo& DataInterfaceCompileInfo : ScriptDataInterfaceCompileInfo)
+		if (TargetScript == nullptr)
 		{
-			// First collect the default data interfaces defined in the graphs.
-			const FNiagaraScriptDataInterfaceDefaultInfo* DataInterfaceDefaultInfo = ScriptCachedDefaultDataInterfaces.FindByPredicate(
-				[&DataInterfaceCompileInfo](const FNiagaraScriptDataInterfaceDefaultInfo& CachedDefaultInfo)
-				{ return CachedDefaultInfo.Id == DataInterfaceCompileInfo.Id && CachedDefaultInfo.DataInterface != nullptr; });
-			if (DataInterfaceDefaultInfo != nullptr)
-			{
-				OutScriptDefaultData.InternalDefaultMap.Add(DataInterfaceCompileInfo.Id, DataInterfaceDefaultInfo->DataInterface);
-				if (DataInterfaceCompileInfo.TargetParameter.IsValid())
-				{
-					// NOTE: We shouldn't need to collect these here since it should be handled in the resolve, but this is needed to handle legacy 
-					// content where an update script defines a data interface that's used in spawn.  This should be handled by fail if not 
-					// set for data interfaces if that's ever implemented.
-					OutScriptDefaultData.ParameterDefaultMap.Add(DataInterfaceCompileInfo.TargetParameter, DataInterfaceDefaultInfo->DataInterface);
-				}
-			}
+			return;
 		}
-	}
 
-	/** Creates default instances of DIs for NPC data interfaces reads.  The NPC data interfaces are not available at compile time and will likely
-		be replaced at runtime, but we need valid instances for safety and to avoid creating them at runtime if they're missing. */
-	void CreateScriptNPCDefaults(UObject& OwnerForCreatedDIs, const TArray<FNiagaraScriptDataInterfaceCompileInfo>& ScriptDataInterfaceCompileInfo,
-		TMap<FNiagaraCompileParameterReference, UNiagaraDataInterface*>& OutSystemExternalParameterDefaultMap)
-	{
-		for (const FNiagaraScriptDataInterfaceCompileInfo& DataInterfaceCompileInfo : ScriptDataInterfaceCompileInfo)
+		for (const FNiagaraScriptDataInterfaceInfo& CachedDefaultDataInterface : TargetScript->GetCachedDefaultDataInterfaces())
 		{
-			if (DataInterfaceCompileInfo.SourceExternalParameter.HasFlags(ENiagaraCompileParameterReferenceFlags::NPC) &&
-				OutSystemExternalParameterDefaultMap.Contains(DataInterfaceCompileInfo.SourceExternalParameter) == false)
+			if (CachedDefaultDataInterface.RegisteredParameterMapWrite != NAME_None)
 			{
-				UNiagaraDataInterface* NPCDefault = NewObject<UNiagaraDataInterface>(&OwnerForCreatedDIs, DataInterfaceCompileInfo.Type.GetClass());
-				OutSystemExternalParameterDefaultMap.Add(DataInterfaceCompileInfo.SourceExternalParameter, NPCDefault);
+				FNiagaraVariable WriteVariable(CachedDefaultDataInterface.Type, CachedDefaultDataInterface.RegisteredParameterMapWrite);
+				if (EmitterName.IsEmpty() == false)
+				{
+					WriteVariable = FNiagaraUtilities::ResolveAliases(WriteVariable, FNiagaraAliasContext().ChangeEmitterNameToEmitter(EmitterName));
+				}
+
+				// Handle reads and bindings.
+				if (CachedDefaultDataInterface.RegisteredParameterMapRead != NAME_None)
+				{
+					FNiagaraVariable ReadVariable(CachedDefaultDataInterface.Type, CachedDefaultDataInterface.RegisteredParameterMapRead);
+					if (EmitterName.IsEmpty() == false)
+					{
+						ReadVariable = FNiagaraUtilities::ResolveAliases(ReadVariable, FNiagaraAliasContext().ChangeEmitterToEmitterName(EmitterName));
+					}
+
+					FDataInterfaceSourceEmitterNamePair* CurrentRead = OutVariableReadMap.Find(ReadVariable);
+					if (CurrentRead == nullptr)
+					{
+						FDataInterfaceSourceEmitterNamePair& NewRead = OutVariableReadMap.Add(ReadVariable);
+						NewRead.DataInterface = CachedDefaultDataInterface.DataInterface;
+						NewRead.SourceName = CachedDefaultDataInterface.Name;
+						NewRead.SourceEmitterName = CachedDefaultDataInterface.SourceEmitterName;
+					}
+
+					if (ReadVariable != WriteVariable)
+					{
+						FNiagaraVariableBase* CurrentBinding = OutVariableBindingMap.Find(WriteVariable);
+						if (CurrentBinding != nullptr)
+						{
+							if (*CurrentBinding != ReadVariable)
+							{
+								OutErrorMessages.Add(FText::Format(
+									LOCTEXT("MultipleBindingsFormat", "A data interface parameter was the target of multiple bindings in a single system.  The data interface used in the simulation may be incorrect.  Target Parameter: {0} First Read: {1} Current Read: {2}"),
+									FText::FromName(WriteVariable.GetName()),
+									FText::FromName(CurrentBinding->GetName()),
+									FText::FromName(ReadVariable.GetName())));
+							}
+						}
+						else
+						{
+							{
+								OutVariableBindingMap.Add(WriteVariable, ReadVariable);
+							}
+						}
+					}
+				}
+		
+				// Handle assignments
+				FDataInterfaceSourceEmitterNamePair* CurrentAssignment = OutVariableAssignmentMap.Find(WriteVariable);
+				if (CurrentAssignment != nullptr)
+				{
+					if (CurrentAssignment->DataInterface != CachedDefaultDataInterface.DataInterface)
+					{
+						OutErrorMessages.Add(FText::Format(
+							LOCTEXT("MultipleAssignmentsFormat", "A data interface parameter was the target of an assignment multiple times in a single system.  The data interface used in the simulation may be incorrect.  Target Parameter: {0} First Assignment: {1} Current Assignment: {2}"),
+							FText::FromName(WriteVariable.GetName()),
+							FText::FromName(CurrentAssignment->SourceName),
+							FText::FromName(CachedDefaultDataInterface.Name)));
+					}
+				}
+				else
+				{
+					FDataInterfaceSourceEmitterNamePair& NewAssignment = OutVariableAssignmentMap.Add(WriteVariable);
+					NewAssignment.DataInterface = CachedDefaultDataInterface.DataInterface;
+					NewAssignment.SourceName = CachedDefaultDataInterface.Name;
+					NewAssignment.SourceEmitterName = CachedDefaultDataInterface.SourceEmitterName;
+				}
 			}
 		}
 	}
 
-	/** Collects default DI instances for all scripts in a system. */
-	void CollectScriptDIResolveDataForSystem(
-		const UNiagaraSystem& System,
-		UObject& OwnerForCreatedDIs,
-		TMap<FGuid, FScriptDefaultData>& OutScriptDefaultDataMap,
-		TMap<FNiagaraCompileParameterReference, UNiagaraDataInterface*>& OutSystemExternalParameterDefaultMap,
+	void CollectDIReadsAssignmentsAndBindings(
+		const UNiagaraSystem* System,
+		TMap<FGuid, TMap<FNiagaraVariableBase, FDataInterfaceSourceEmitterNamePair>>& OutEmitterIdToVariableReadsMap,
+		TMap<FGuid, TMap<FNiagaraVariableBase, FDataInterfaceSourceEmitterNamePair>>& OutEmitterIdToVariableAssignmentsMap,
+		TMap<FGuid, TMap<FNiagaraVariableBase, FNiagaraVariableBase>>& OutEmitterIdToVariableBindingsMap,
 		TArray<FText>& OutErrorMessages)
 	{
-		// First add the user parameters as external defaults.
-		const FNiagaraParameterStore& InstanceParameters = System.GetSystemCompiledData().InstanceParamStore;
+		TMap<FNiagaraVariableBase, FDataInterfaceSourceEmitterNamePair> VariableReadMap;
+		TMap<FNiagaraVariableBase, FDataInterfaceSourceEmitterNamePair> VariableAssignmentMap;
+		TMap<FNiagaraVariableBase, FNiagaraVariableBase> VariableBindingMap;
+
+		// Add user parameters as assignments.
+		const FNiagaraParameterStore& InstanceParameters = System->GetSystemCompiledData().InstanceParamStore;
 		for (const FNiagaraVariableWithOffset& UserParameterWithOffset : InstanceParameters.ReadParameterVariables())
 		{
 			if (UserParameterWithOffset.IsDataInterface())
 			{
-				FNiagaraCompileParameterReference UserParameterReference = FNiagaraCompileParameterReference(
-					UserParameterWithOffset.GetName(), FGuid(), ENiagaraCompileParameterReferenceFlags::User);
 				UNiagaraDataInterface* DataInterface = InstanceParameters.GetDataInterface(UserParameterWithOffset.Offset);
-				if (DataInterface == nullptr)
+				if (DataInterface != nullptr)
 				{
-					DataInterface = NewObject<UNiagaraDataInterface>(&OwnerForCreatedDIs, UserParameterWithOffset.GetType().GetClass());
+					FDataInterfaceSourceEmitterNamePair& NewAssignment = VariableAssignmentMap.Add(UserParameterWithOffset);
+					NewAssignment.DataInterface = DataInterface;
+					NewAssignment.SourceEmitterName = FString();
 				}
-				OutSystemExternalParameterDefaultMap.Add(UserParameterReference, DataInterface);
 			}
 		}
 
-		auto HandleScript = [&OwnerForCreatedDIs,&OutScriptDefaultDataMap,&OutSystemExternalParameterDefaultMap]
-			(const FNiagaraEmitterHandle* OwningEmitterHandle, const UNiagaraScript* Script, FScriptDefaultData& ScriptDefaultData)
-		{
-			CollectScriptInternalDefaultDIs(Script->GetVMExecutableData().DataInterfaceInfo, Script->GetCachedGraphDefaultDataInterfaces(), ScriptDefaultData);
-			CreateScriptNPCDefaults(OwnerForCreatedDIs, Script->GetVMExecutableData().DataInterfaceInfo, OutSystemExternalParameterDefaultMap);
-		};
+		// Collect system and emitter scripts.
+		CollectDIReadsAssignmentsAndBindingsForScript(System->GetSystemSpawnScript(), System, FString(), VariableReadMap, VariableAssignmentMap, VariableBindingMap, OutErrorMessages);
+		CollectDIReadsAssignmentsAndBindingsForScript(System->GetSystemUpdateScript(), System, FString(), VariableReadMap, VariableAssignmentMap, VariableBindingMap, OutErrorMessages);
+		OutEmitterIdToVariableReadsMap.Add(FGuid(), VariableReadMap);
+		OutEmitterIdToVariableAssignmentsMap.Add(FGuid(), VariableAssignmentMap);
+		OutEmitterIdToVariableBindingsMap.Add(FGuid(), VariableBindingMap);
 
-		FScriptDefaultData& SystemScriptDefaultData = OutScriptDefaultDataMap.Add(FGuid());
-		HandleScript(nullptr, System.GetSystemSpawnScript(), SystemScriptDefaultData);
-		HandleScript(nullptr, System.GetSystemUpdateScript(), SystemScriptDefaultData);
-
-		for (const FNiagaraEmitterHandle& EmitterHandle : System.GetEmitterHandles())
+		// Collect emitter scripts.
+		for (const FNiagaraEmitterHandle& EmitterHandle : System->GetEmitterHandles())
 		{
-			FVersionedNiagaraEmitterData* VersionedNiagaraEmitterData = EmitterHandle.GetIsEnabled() ? EmitterHandle.GetEmitterData() : nullptr;
-			if (VersionedNiagaraEmitterData == nullptr)
+			if (EmitterHandle.GetIsEnabled() == false)
 			{
 				continue;
 			}
 
-			FScriptDefaultData& EmitterScriptDefaultData = OutScriptDefaultDataMap.Add(EmitterHandle.GetId());
-			HandleScript(&EmitterHandle, VersionedNiagaraEmitterData->SpawnScriptProps.Script, EmitterScriptDefaultData);
-			HandleScript(&EmitterHandle, VersionedNiagaraEmitterData->UpdateScriptProps.Script, EmitterScriptDefaultData);
-			for (const FNiagaraEventScriptProperties& EventHandler : VersionedNiagaraEmitterData->GetEventHandlers())
+			FVersionedNiagaraEmitterData* VersionedNiagaraEmitterData = EmitterHandle.GetEmitterData();
+			if (VersionedNiagaraEmitterData == nullptr)
 			{
-				HandleScript(&EmitterHandle, EventHandler.Script, EmitterScriptDefaultData);
+				continue;
 			}
-			if (VersionedNiagaraEmitterData->SimTarget == ENiagaraSimTarget::GPUComputeSim)
+
+			TMap<FNiagaraVariableBase, FDataInterfaceSourceEmitterNamePair> EmitterVariableReadMap = VariableReadMap;
+			TMap<FNiagaraVariableBase, FDataInterfaceSourceEmitterNamePair> EmitterVariableAssignmentMap = VariableAssignmentMap;
+			TMap<FNiagaraVariableBase, FNiagaraVariableBase> EmitterVariableBindingMap = VariableBindingMap;
+
+			CollectDIReadsAssignmentsAndBindingsForScript(VersionedNiagaraEmitterData->SpawnScriptProps.Script, System, EmitterHandle.GetUniqueInstanceName(), EmitterVariableReadMap, EmitterVariableAssignmentMap, EmitterVariableBindingMap, OutErrorMessages);
+			CollectDIReadsAssignmentsAndBindingsForScript(VersionedNiagaraEmitterData->UpdateScriptProps.Script, System, EmitterHandle.GetUniqueInstanceName(), EmitterVariableReadMap, EmitterVariableAssignmentMap, EmitterVariableBindingMap, OutErrorMessages);
+
+			for (const FNiagaraEventScriptProperties& EventHandler : VersionedNiagaraEmitterData->GetEventHandlers())
 			{
-				HandleScript(&EmitterHandle, VersionedNiagaraEmitterData->GetGPUComputeScript(), EmitterScriptDefaultData);
+				CollectDIReadsAssignmentsAndBindingsForScript(EventHandler.Script, System, EmitterHandle.GetUniqueInstanceName(), EmitterVariableReadMap, EmitterVariableAssignmentMap, EmitterVariableBindingMap, OutErrorMessages);
 			}
-		}
-	}
 
-	/** Initializes the runtime compiled data map for all scrips in the system. */
-	void ResetResolvedDIsForSystem(const UNiagaraSystem& System, TMap<FNiagaraScriptDataKey, FNiagaraScriptRuntimeCompiledData>& OutScriptRuntimeCompiledDataMap)
-	{
-		System.ForEachScriptWithOwningContext([&OutScriptRuntimeCompiledDataMap]
-			(const UNiagaraSystem* OwningSystem, const FNiagaraEmitterHandle* OwningEmitterHandle, UNiagaraScript* Script)
-		{
-			FGuid EmitterHandleId = OwningEmitterHandle != nullptr ? OwningEmitterHandle->GetId() : FGuid();
-			FNiagaraScriptDataKey ScriptDataKey = FNiagaraScriptDataKey(EmitterHandleId, Script->GetUsage(), Script->GetUsageId());
-			FNiagaraScriptRuntimeCompiledData& ScriptRuntimeCompiledData = OutScriptRuntimeCompiledDataMap.FindOrAdd(ScriptDataKey);
-			ScriptRuntimeCompiledData.ResolvedDataInterfaces.Empty(Script->GetVMExecutableData().DataInterfaceInfo.Num());
-			ScriptRuntimeCompiledData.ResolvedDataInterfaces.AddDefaulted(Script->GetVMExecutableData().DataInterfaceInfo.Num());
-			ScriptRuntimeCompiledData.ResolvedUserDataInterfaceBindings.Empty();
-			ScriptRuntimeCompiledData.ResolvedNpcDataInterfaceBindings.Empty();
-		});
-	}
-
-	/** Converts the system external parameter default map to an emitter specific parameter default map for use by particle scripts in that
-		emitter. This step is necessary since emitter parameter data interfaces are collected from the system script and their parameter references
-		include the emitter guid to disambiguate them.  Emitter parameter references in particle scripts do not include emitter specific information
-		so the emitter guid must be stripped from them so that they can correctly resolve.  We also filter out emitter parameters DIs for other emitters
-		here. */
-	void CreateEmitterSpecificParameterDefaults(
-		FGuid TargetEmitterHandleId,
-		const TMap<FNiagaraCompileParameterReference, UNiagaraDataInterface*>& InExternalParameterDefaultMap,
-		TMap<FNiagaraCompileParameterReference, UNiagaraDataInterface*>& OutEmitterExternalParameterDefaultMap)
-	{
-		for (const TPair<FNiagaraCompileParameterReference, UNiagaraDataInterface*>& ExternalParameterDefaultPair : InExternalParameterDefaultMap)
-		{
-			const FNiagaraCompileParameterReference& ExternalParameterReference = ExternalParameterDefaultPair.Key;
-			UNiagaraDataInterface* ExternalDataInterface = ExternalParameterDefaultPair.Value;
-			if (ExternalParameterReference.HasFlags(ENiagaraCompileParameterReferenceFlags::Emitter))
+			for (const UNiagaraSimulationStageBase* SimulationStage : VersionedNiagaraEmitterData->GetSimulationStages())
 			{
-				if (ExternalParameterReference.EmitterHandleId == TargetEmitterHandleId)
+				if (SimulationStage == nullptr || SimulationStage->bEnabled == false)
 				{
-					FNiagaraCompileParameterReference EmitterParameterReference(ExternalParameterReference.ParameterFullName, FGuid(), ExternalParameterReference.Flags);
-					OutEmitterExternalParameterDefaultMap.Add(EmitterParameterReference, ExternalDataInterface);
+					continue;
 				}
+				CollectDIReadsAssignmentsAndBindingsForScript(SimulationStage->Script, System, EmitterHandle.GetUniqueInstanceName(), EmitterVariableReadMap, EmitterVariableAssignmentMap, EmitterVariableBindingMap, OutErrorMessages);
 			}
-			else
-			{
-				OutEmitterExternalParameterDefaultMap.Add(ExternalParameterDefaultPair);
-			}
+
+			OutEmitterIdToVariableReadsMap.Add(EmitterHandle.GetId(), EmitterVariableReadMap);
+			OutEmitterIdToVariableAssignmentsMap.Add(EmitterHandle.GetId(), EmitterVariableAssignmentMap);
+			OutEmitterIdToVariableBindingsMap.Add(EmitterHandle.GetId(), EmitterVariableBindingMap);
 		}
 	}
 
-	/** Collects resolved data interfaces for parameters that can be referenced external to the script and maps them by parameter reference for use by
-		dependent scripts. */
-	void CollectExternalParameterDefaults(
-		const TArray<FNiagaraScriptResolvedDataInterfaceInfo>& ResolvedDataInterfaces,
-		TMap<FNiagaraCompileParameterReference, UNiagaraDataInterface*>& ExternalParameterDefaultMap)
+	void CompactBindings(const TMap<FNiagaraVariableBase, FNiagaraVariableBase>& InVariableBindingMap, TMap<FNiagaraVariableBase, FNiagaraVariableBase>& OutCompactedVariableBindingMap, TArray<FText>& OutErrorMessages)
 	{
-		for (const FNiagaraScriptResolvedDataInterfaceInfo& ResolvedDataInterfaceData : ResolvedDataInterfaces)
+		TArray<FNiagaraVariableBase> SourceVariables;
+		InVariableBindingMap.GetKeys(SourceVariables);
+		TMap<FNiagaraVariableBase, FNiagaraVariableBase> CompactedVariableBindingMap;
+		for (const FNiagaraVariableBase& SourceVariable : SourceVariables)
 		{
-			if (ResolvedDataInterfaceData.TargetParameterReference.IsValid() &&
-				ResolvedDataInterfaceData.TargetParameterReference.IsScriptInternalReference() == false &&
-				ResolvedDataInterfaceData.ResolvedDataInterface != nullptr &&
-				ExternalParameterDefaultMap.Contains(ResolvedDataInterfaceData.TargetParameterReference) == false)
+			FNiagaraVariableBase BoundVariable;
+			const FNiagaraVariableBase* CurrentBoundVariable = &SourceVariable;
+			TSet<FNiagaraVariableBase> SeenBoundVariables;
+			bool bCircularReferenceDetected = false;
+			while (CurrentBoundVariable != nullptr && bCircularReferenceDetected == false)
 			{
-				ExternalParameterDefaultMap.Add(ResolvedDataInterfaceData.TargetParameterReference, ResolvedDataInterfaceData.ResolvedDataInterface);
+				SeenBoundVariables.Add(*CurrentBoundVariable);
+				BoundVariable = *CurrentBoundVariable;
+				CurrentBoundVariable = InVariableBindingMap.Find(BoundVariable);
+				if (CurrentBoundVariable != nullptr && SeenBoundVariables.Contains(*CurrentBoundVariable))
+				{
+					bCircularReferenceDetected = true;
+					OutErrorMessages.Add(FText::Format(
+						LOCTEXT("CircularDependencyWarningFormat", "A data interface parameter circular dependency was found.  The data interface used in the simulation may be incorrect.  Target Parameter: {0} Resolved Parameter: {1}"),
+						FText::FromName(SourceVariable.GetName()),
+						FText::FromName(CurrentBoundVariable->GetName())));
+				}
 			}
+			OutCompactedVariableBindingMap.Add(SourceVariable, BoundVariable);
 		}
 	}
 
-	/** Collects default DI instances for unresolved system parameters which were created while resolving particle scripts within an emitter.  The instances are added to a map
-		for the system so that they will be shared among particle scripts in different emitters referencing the same system parameter. */
-	void CollectSystemUnresolvedExternalParameterDefaults(
-		const TMap<FNiagaraCompileParameterReference, UNiagaraDataInterface*>& EmitterUnresolvedExternalParameterDefaultMap,
-		TMap<FNiagaraCompileParameterReference, UNiagaraDataInterface*>& SystemUnresolvedExternalParameterDefaultMap)
+	void ResolveDIsForScript(
+		const UNiagaraSystem* System,
+		UNiagaraScript* TargetScript,
+		const FString& EmitterName,
+		const TMap<FNiagaraVariableBase, FDataInterfaceSourceEmitterNamePair>& VariableReadMap,
+		const TMap<FNiagaraVariableBase, FDataInterfaceSourceEmitterNamePair>& VariableAssignmentMap,
+		const TMap<FNiagaraVariableBase, FNiagaraVariableBase>& CompactedVariableBindingMap,
+		FNiagaraScriptRuntimeCompiledData& OutScriptRuntimeCompiledData,
+		TArray<FText>& OutErrorMessages)
 	{
-		for (const TPair<FNiagaraCompileParameterReference, UNiagaraDataInterface*>& EmitterUnresolvedExternalParameterDefaultPair : EmitterUnresolvedExternalParameterDefaultMap)
+		if (TargetScript == nullptr)
 		{
-			const FNiagaraCompileParameterReference& EmitterUnresolvedExternalParameterReference = EmitterUnresolvedExternalParameterDefaultPair.Key;
-			if (EmitterUnresolvedExternalParameterReference.HasFlags(ENiagaraCompileParameterReferenceFlags::System) && 
-				SystemUnresolvedExternalParameterDefaultMap.Contains(EmitterUnresolvedExternalParameterReference) == false)
-			{
-				UNiagaraDataInterface* UnresolvedExternalDataInterface = EmitterUnresolvedExternalParameterDefaultPair.Value;
-				SystemUnresolvedExternalParameterDefaultMap.Add(EmitterUnresolvedExternalParameterReference, UnresolvedExternalDataInterface);
-			}
+			return;
 		}
-	}
-
-	/** Resolves data interfaces for a set of data interface compile infos using previously discovered or created instances. */
-	void ResolveDIParametersAndDefaults(
-		const TArray<FNiagaraScriptDataInterfaceCompileInfo>& ScriptDataInterfaceCompileInfo,
-		const TMap<FNiagaraScriptDataInterfaceCompileId, UNiagaraDataInterface*>& InternalDefaultMap,
-		const TMap<FNiagaraCompileParameterReference, UNiagaraDataInterface*>& ParameterDefaultMap,
-		const TMap<FNiagaraCompileParameterReference, UNiagaraDataInterface*>& ExternalParameterDefaultMap,
-		const TMap<FNiagaraCompileParameterReference, UNiagaraDataInterface*>& UnresolvedExternalParameterDefaultMap,
-		const TOptional<ENiagaraScriptUsage>& ScriptUsageFilter,
-		TArray<FNiagaraScriptResolvedDataInterfaceInfo>& OutResolvedDataInterfaces)
-	{
-		for (int32 ScriptDIIndex = 0; ScriptDIIndex < ScriptDataInterfaceCompileInfo.Num(); ScriptDIIndex++)
+		TArray<FNiagaraScriptResolvedDataInterfaceInfo> ResolvedDataInterfaces;
+		ResolvedDataInterfaces.Reserve(TargetScript->GetCachedDefaultDataInterfaces().Num());
+		TArray<FNiagaraResolvedUserDataInterfaceBinding> UserDataInterfaceBindings;
+		int32 ResolvedDataInterfaceIndex = 0;
+		TSet<FNiagaraVariableBase> ParameterStoreVariables;
+		for (const FNiagaraScriptDataInterfaceInfo& CachedDefaultDataInterface : TargetScript->GetCachedDefaultDataInterfaces())
 		{
-			const FNiagaraScriptDataInterfaceCompileInfo& DataInterfaceCompileInfo = ScriptDataInterfaceCompileInfo[ScriptDIIndex];
-			if (ScriptUsageFilter.IsSet() && DataInterfaceCompileInfo.Id.GetScriptUsage() != ScriptUsageFilter.GetValue())
-			{
-				continue;
-			}
+			FNiagaraScriptResolvedDataInterfaceInfo& ResolvedDataInterface = ResolvedDataInterfaces.AddDefaulted_GetRef();
+			ResolvedDataInterface.Name = CachedDefaultDataInterface.Name;
+			ResolvedDataInterface.CompileName = CachedDefaultDataInterface.CompileName;
+			ResolvedDataInterface.UserPtrIdx = CachedDefaultDataInterface.UserPtrIdx;
 
-			UNiagaraDataInterface* ResolvedDataInterfaceInstance = nullptr;
-			FNiagaraCompileParameterReference ResolvedParameterReference;
-			if (DataInterfaceCompileInfo.SourceExternalParameter.IsValid())
+			if (CachedDefaultDataInterface.RegisteredParameterMapRead != NAME_None)
 			{
-				UNiagaraDataInterface*const* ExternalDefaultPtr = ExternalParameterDefaultMap.Find(DataInterfaceCompileInfo.SourceExternalParameter);
-				if (ExternalDefaultPtr != nullptr && *ExternalDefaultPtr != nullptr)
+				// If this DI is read from a parameter try to resolve it through the binding and assignment or read maps.
+				FNiagaraVariable ReadVariable(CachedDefaultDataInterface.Type, CachedDefaultDataInterface.RegisteredParameterMapRead);
+				if (EmitterName.IsEmpty() == false)
 				{
-					ResolvedDataInterfaceInstance = *ExternalDefaultPtr;
-					ResolvedParameterReference = DataInterfaceCompileInfo.SourceExternalParameter;
+					ReadVariable = FNiagaraUtilities::ResolveAliases(ReadVariable, FNiagaraAliasContext().ChangeEmitterToEmitterName(EmitterName));
 				}
-				else
+
+				const FNiagaraVariableBase* BoundVariable = CompactedVariableBindingMap.Find(ReadVariable);
+				if (BoundVariable == nullptr)
+				{
+					BoundVariable = &ReadVariable;
+				}
+
+				const FDataInterfaceSourceEmitterNamePair* BoundDataInterface = VariableAssignmentMap.Find(*BoundVariable);
+				if (BoundDataInterface == nullptr)
+				{
+					// If the data interface was not written to explicitly, then it needs to use the default which was cached in the first read.
+					BoundDataInterface = VariableReadMap.Find(*BoundVariable);
+				}
+
+				/*
+				if (BoundDataInterface == nullptr)
+				{
+					OutErrorMessages.Add(FText::Format(
+						LOCTEXT("MissingReadDataInterface", "A data interface was read from the parameter map, but its default could not be found.  The data interface used in the simulation may be incorrect.  Read Parameter: {0} Bound Parameter: {1}"),
+						FText::FromName(ReadVariable.GetName()),
+						FText::FromName(BoundVariable->GetName())));
+				}
+				*/
+
+				if (BoundDataInterface != nullptr && ParameterStoreVariables.Contains(ReadVariable) == false)
 				{
-					// NOTE: This case handles legacy content where an update script defines a data interface that's used in spawn.  This
-					// should be handled by fail if not set for data interfaces if that's ever implemented.
-					UNiagaraDataInterface*const* ParameterDefaultPtr = ParameterDefaultMap.Find(DataInterfaceCompileInfo.SourceExternalParameter);
-					if (ParameterDefaultPtr != nullptr && *ParameterDefaultPtr != nullptr)
+					ResolvedDataInterface.ResolvedVariable = *BoundVariable;
+					ResolvedDataInterface.ParameterStoreVariable = ReadVariable;
+					ResolvedDataInterface.bIsInternal = false;
+					ResolvedDataInterface.ResolvedDataInterface = BoundDataInterface->DataInterface;
+					ResolvedDataInterface.ResolvedSourceEmitterName = BoundDataInterface->SourceEmitterName;
+					ParameterStoreVariables.Add(ResolvedDataInterface.ParameterStoreVariable);
+				}
+
+				if (FNiagaraUserRedirectionParameterStore::IsUserParameter(*BoundVariable))
+				{
+					const FNiagaraParameterStore& InstanceParameters = System->GetSystemCompiledData().InstanceParamStore;
+					const int32* UserParameterIndex = InstanceParameters.FindParameterOffset(*BoundVariable);
+					if (UserParameterIndex != nullptr && *UserParameterIndex != INDEX_NONE)
 					{
-						ResolvedDataInterfaceInstance = *ParameterDefaultPtr;
-						ResolvedParameterReference = DataInterfaceCompileInfo.SourceExternalParameter;
-					}
-					else
-					{
-						// NOTE: This case handles legacy content where there are multiple reads for a data interface that was never defined.  This
-						// should be handled by fail if not set for data interfaces if that's ever implemented.
-						UNiagaraDataInterface*const* UnresolvedExternalDefault = UnresolvedExternalParameterDefaultMap.Find(DataInterfaceCompileInfo.SourceExternalParameter);
-						if (UnresolvedExternalDefault != nullptr && *UnresolvedExternalDefault != nullptr)
-						{
-							// We don't set the resolved parameter here since it wasn't actually resolved, we only use the data interface so that
-							// we don't create multiple independent data interfaces for the same unresolved read.
-							ResolvedDataInterfaceInstance = *UnresolvedExternalDefault;
-						}
+						UserDataInterfaceBindings.Add(FNiagaraResolvedUserDataInterfaceBinding(*UserParameterIndex, ResolvedDataInterfaceIndex));
 					}
 				}
 			}
-			if (ResolvedDataInterfaceInstance == nullptr && DataInterfaceCompileInfo.SourceScriptDIIndex != INDEX_NONE)
+			else if (CachedDefaultDataInterface.RegisteredParameterMapWrite != NAME_None)
 			{
-				if (ensureMsgf(DataInterfaceCompileInfo.SourceScriptDIIndex < ScriptDIIndex, TEXT("Data interface ordering is incorrect.")))
+				FNiagaraVariable WriteVariable(CachedDefaultDataInterface.Type, CachedDefaultDataInterface.RegisteredParameterMapWrite);
+				const FDataInterfaceSourceEmitterNamePair* DataInterfaceAssignment = VariableAssignmentMap.Find(WriteVariable);
+				if (DataInterfaceAssignment != nullptr && ParameterStoreVariables.Contains(WriteVariable) == false)
 				{
-					ResolvedDataInterfaceInstance = OutResolvedDataInterfaces[DataInterfaceCompileInfo.SourceScriptDIIndex].ResolvedDataInterface;
-					ResolvedParameterReference = OutResolvedDataInterfaces[DataInterfaceCompileInfo.SourceScriptDIIndex].ResolvedParameterReference;
-				}	
+					ResolvedDataInterface.ResolvedVariable = WriteVariable;
+					ResolvedDataInterface.ParameterStoreVariable = WriteVariable;
+					ResolvedDataInterface.bIsInternal = false;
+					ResolvedDataInterface.ResolvedDataInterface = DataInterfaceAssignment->DataInterface;
+					ResolvedDataInterface.ResolvedSourceEmitterName = DataInterfaceAssignment->SourceEmitterName;
+					ParameterStoreVariables.Add(ResolvedDataInterface.ParameterStoreVariable);
+				}
 			}
-			if (ResolvedDataInterfaceInstance == nullptr && DataInterfaceCompileInfo.SourceGraphInputDataInterfaceKey.IsValid())
+
+			if (ResolvedDataInterface.ResolvedDataInterface == nullptr)
 			{
-				UNiagaraDataInterface*const* InternalSourceDataInterfacePtr = InternalDefaultMap.Find(DataInterfaceCompileInfo.Id);
-				if (InternalSourceDataInterfacePtr != nullptr && *InternalSourceDataInterfacePtr != nullptr)
+				// NPC variables will fail to resolve but we want to main that it's both external & not append the internal namespace
+				if ( FNiagaraVariable::IsInNameSpace(FNiagaraConstants::ParameterCollectionNamespaceString, ResolvedDataInterface.Name) )
 				{
-					ResolvedDataInterfaceInstance = *InternalSourceDataInterfacePtr;
+					FNiagaraVariable NPCVariable(CachedDefaultDataInterface.Type, ResolvedDataInterface.Name);
+					ResolvedDataInterface.ResolvedVariable = NPCVariable;
+					ResolvedDataInterface.ParameterStoreVariable = NPCVariable;
+					ResolvedDataInterface.bIsInternal = false;
+					ResolvedDataInterface.ResolvedDataInterface = CachedDefaultDataInterface.DataInterface;
+					ResolvedDataInterface.ResolvedSourceEmitterName = CachedDefaultDataInterface.SourceEmitterName;
+				}
+				else
+				{
+					// If the DI was not read from a parameter or couldn't be found, use the one cached during compilation, and give it an internal
+					// name to prevent it from being bound incorrectly.
+					FNameBuilder NameBuilder;
+					NameBuilder.Append(FNiagaraConstants::InternalNamespaceString);
+					NameBuilder.AppendChar(TEXT('.'));
+					ResolvedDataInterface.Name.AppendString(NameBuilder);
+
+					FNiagaraVariable InternalVariable(CachedDefaultDataInterface.Type, FName(NameBuilder.ToString()));
+					ResolvedDataInterface.ResolvedVariable = InternalVariable;
+					ResolvedDataInterface.ParameterStoreVariable = InternalVariable;
+					ResolvedDataInterface.bIsInternal = true;
+					ResolvedDataInterface.ResolvedDataInterface = CachedDefaultDataInterface.DataInterface;
+					ResolvedDataInterface.ResolvedSourceEmitterName = CachedDefaultDataInterface.SourceEmitterName;
 				}
 			}
 
-			FNiagaraScriptResolvedDataInterfaceInfo& ResolvedDataInterfaceData = OutResolvedDataInterfaces[ScriptDIIndex];
-			ResolvedDataInterfaceData.OwningEmitterHandleId = DataInterfaceCompileInfo.Id.GetEmitterHandleId();
-			ResolvedDataInterfaceData.CompileName = DataInterfaceCompileInfo.Name;
-			ResolvedDataInterfaceData.Type = DataInterfaceCompileInfo.Type;
-			ResolvedDataInterfaceData.ResolvedParameterReference = ResolvedParameterReference;
-			ResolvedDataInterfaceData.TargetParameterReference = DataInterfaceCompileInfo.TargetParameter;
-			ResolvedDataInterfaceData.ResolvedDataInterface = ResolvedDataInterfaceInstance;
-			ResolvedDataInterfaceData.UserPtrIdx = DataInterfaceCompileInfo.UserPtrIdx;
+			ResolvedDataInterfaceIndex++;
 		}
+		OutScriptRuntimeCompiledData.ResolvedDataInterfaces = ResolvedDataInterfaces;
+		OutScriptRuntimeCompiledData.ResolvedUserDataInterfaceBindings = UserDataInterfaceBindings;
 	}
 
-	/** Creates default DI instances for any data interface compile infos which couldn't be resolved to a previously defined
-		or created DI instance. */
-	void GenerateUnresolvedDIs(
-		const TArray<FNiagaraScriptDataInterfaceCompileInfo>& ScriptDataInterfaceCompileInfo,
-		const TOptional<ENiagaraScriptUsage>& ScriptUsageFilter,
-		UObject& OwnerForCreatedDIs,
-		TArray<FNiagaraScriptResolvedDataInterfaceInfo>& ResolvedDataInterfaces,
-		TMap<FNiagaraCompileParameterReference, UNiagaraDataInterface*>& UnresolvedExternalParameterDefaultMap)
+	void SynchronizeMatchingInternalResolvedDataInterfaces(TArrayView<const FNiagaraScriptResolvedDataInterfaceInfo> SourceResolvedDataInterfaces, TArray<FNiagaraScriptResolvedDataInterfaceInfo>& TargetResolvedDataInterfaces)
 	{
-		for (int32 ScriptDIIndex = 0; ScriptDIIndex < ScriptDataInterfaceCompileInfo.Num(); ScriptDIIndex++)
+		for (const FNiagaraScriptResolvedDataInterfaceInfo& SourceResolvedDataInterface : SourceResolvedDataInterfaces)
 		{
-			const FNiagaraScriptDataInterfaceCompileInfo& DataInterfaceCompileInfo = ScriptDataInterfaceCompileInfo[ScriptDIIndex];
-			FNiagaraScriptResolvedDataInterfaceInfo& ResolvedDataInterfaceData = ResolvedDataInterfaces[ScriptDIIndex];
-			if (ScriptUsageFilter.IsSet() && DataInterfaceCompileInfo.Id.GetScriptUsage() != ScriptUsageFilter.GetValue())
+			for (FNiagaraScriptResolvedDataInterfaceInfo& TargetResolvedDataInterface : TargetResolvedDataInterfaces)
 			{
-				continue;
-			}
-
-			if (ResolvedDataInterfaceData.ResolvedDataInterface == nullptr)
-			{
-				ResolvedDataInterfaceData.ResolvedDataInterface = NewObject<UNiagaraDataInterface>(&OwnerForCreatedDIs, DataInterfaceCompileInfo.Type.GetClass());
-				if (DataInterfaceCompileInfo.SourceExternalParameter.IsValid())
+				if (TargetResolvedDataInterface.bIsInternal && TargetResolvedDataInterface.Name == SourceResolvedDataInterface.Name)
 				{
-					// NOTE: This missing default is an unresolved external read and we don't want to create multiple duplicates for multiple
-					// reads, so it needs to be collected here. This should be handled by fail if not set for data interfaces if that's implemented.
-					UnresolvedExternalParameterDefaultMap.Add(DataInterfaceCompileInfo.SourceExternalParameter, ResolvedDataInterfaceData.ResolvedDataInterface);
+					TargetResolvedDataInterface.ResolvedDataInterface = SourceResolvedDataInterface.ResolvedDataInterface;
+					TargetResolvedDataInterface.ResolvedSourceEmitterName = SourceResolvedDataInterface.ResolvedSourceEmitterName;
+					break;
 				}
 			}
 		}
 	}
 
-	/** Generates resolved DI infos for a specific script using previously collected data from defaults or from higher level script resolves. */
-	void GenerateResolvedDIsForScript(
-		const FNiagaraEmitterHandle* EmitterHandle,
-		const UNiagaraScript* Script,
-		const FScriptDefaultData& ScriptDefaultData,
-		const TOptional<ENiagaraScriptUsage>& ScriptUsageFilter,
-		UObject& OwnerForCreatedDIs,
-		TMap<FNiagaraScriptDataKey, FNiagaraScriptRuntimeCompiledData>& ScriptRuntimeCompiledDataMap,
-		TMap<FNiagaraCompileParameterReference, UNiagaraDataInterface*>& ExternalParameterDefaultMap,
-		TMap<FNiagaraCompileParameterReference, UNiagaraDataInterface*>& UnresolvedExternalParameterDefaultMap)
+	/** Handles the special case where internal data interfaces defined in certain scripts need to be synchronized with other scripts that contain them. */
+	void ResolveInternalDataInterfaces(TArray<FNiagaraScriptRuntimeCompiledData*> SourceScriptsRuntimeCompiledData, FNiagaraScriptRuntimeCompiledData& TargetScriptRuntimeCompiledData)
 	{
-		FGuid EmitterHandleId = EmitterHandle != nullptr ? EmitterHandle->GetId() : FGuid();
-		FNiagaraScriptDataKey ScriptDataKey = FNiagaraScriptDataKey(EmitterHandleId, Script->GetUsage(), Script->GetUsageId());
-		FNiagaraScriptRuntimeCompiledData& ScriptRuntimeCompiledData = ScriptRuntimeCompiledDataMap[ScriptDataKey];
-		const TArray<FNiagaraScriptDataInterfaceCompileInfo>& ScriptDataInterfaceCompileInfo = Script->GetVMExecutableData().DataInterfaceInfo;
-		
-		ResolveDIParametersAndDefaults(ScriptDataInterfaceCompileInfo, ScriptDefaultData.InternalDefaultMap, ScriptDefaultData.ParameterDefaultMap,
-			ExternalParameterDefaultMap, UnresolvedExternalParameterDefaultMap, ScriptUsageFilter, ScriptRuntimeCompiledData.ResolvedDataInterfaces);
-		GenerateUnresolvedDIs(ScriptDataInterfaceCompileInfo, ScriptUsageFilter, OwnerForCreatedDIs, ScriptRuntimeCompiledData.ResolvedDataInterfaces,
-			UnresolvedExternalParameterDefaultMap);
-		CollectExternalParameterDefaults(ScriptRuntimeCompiledData.ResolvedDataInterfaces, ExternalParameterDefaultMap);
-	};
+		TArray<FNiagaraScriptResolvedDataInterfaceInfo> ResolvedTargetDataInterfaces;
+		ResolvedTargetDataInterfaces.Append(TargetScriptRuntimeCompiledData.ResolvedDataInterfaces);
+		for (FNiagaraScriptRuntimeCompiledData* SourceScriptRuntimeCompiledData : SourceScriptsRuntimeCompiledData)
+		{
+			TArrayView<const FNiagaraScriptResolvedDataInterfaceInfo> ResolvedSourceParticleDataInterfaces = MakeConstArrayView(SourceScriptRuntimeCompiledData->ResolvedDataInterfaces);
+			SynchronizeMatchingInternalResolvedDataInterfaces(ResolvedSourceParticleDataInterfaces, ResolvedTargetDataInterfaces);
+		}
+		TargetScriptRuntimeCompiledData.ResolvedDataInterfaces = ResolvedTargetDataInterfaces;
+	}
 
-	/** Generates resolved DI infos for all scripts in a given system. */
-	void GenerateResolvedDIsForSystem(
-		const UNiagaraSystem& System,
-		UObject& OwnerForCreatedDIs,
-		const TMap<FGuid, FScriptDefaultData>& ScriptDefaultDataMap,
-		const TMap<FNiagaraCompileParameterReference, UNiagaraDataInterface*>& ExternalParameterDefaultMap,
+	void ResolveDIsInternal(
+		UNiagaraSystem* System,
+		const TMap<FGuid, TMap<FNiagaraVariableBase, FDataInterfaceSourceEmitterNamePair>>& EmitterIdToVariableReadsMap,
+		const TMap<FGuid, TMap<FNiagaraVariableBase, FDataInterfaceSourceEmitterNamePair>>& EmitterIdToVariableAssignmentsMap,
+		const TMap<FGuid, TMap<FNiagaraVariableBase, FNiagaraVariableBase>>& EmitterIdToVariableBindingsMap,
 		TMap<FNiagaraScriptDataKey, FNiagaraScriptRuntimeCompiledData>& OutScriptRuntimeCompiledDataMap,
 		TArray<FText>& OutErrorMessages)
 	{
-		const FScriptDefaultData& SystemScriptDefaultData = ScriptDefaultDataMap[FGuid()];
-		TMap<FNiagaraCompileParameterReference, UNiagaraDataInterface*> SystemExternalParameterDefaultMap = ExternalParameterDefaultMap;
-		TMap<FNiagaraCompileParameterReference, UNiagaraDataInterface*> SystemUnresolvedExternalParameterDefaultMap;
-		
-		auto HandleSystemScript = [&OwnerForCreatedDIs,&OutScriptRuntimeCompiledDataMap,&SystemScriptDefaultData,
-			&SystemExternalParameterDefaultMap,&SystemUnresolvedExternalParameterDefaultMap]
-			(const UNiagaraScript* Script, ENiagaraScriptUsage TargetUsage)
-		{
-			GenerateResolvedDIsForScript(nullptr, Script, SystemScriptDefaultData, TargetUsage, OwnerForCreatedDIs, OutScriptRuntimeCompiledDataMap,
-				SystemExternalParameterDefaultMap, SystemUnresolvedExternalParameterDefaultMap);
-		};
-		
-		// Run the resolve for the system spawn and system update scripts, but with filters for only system scripts so that bindings
-		// from system spawn to emitter update can be handled correctly.
-		HandleSystemScript(System.GetSystemSpawnScript(), ENiagaraScriptUsage::SystemSpawnScript);
-		HandleSystemScript(System.GetSystemUpdateScript(), ENiagaraScriptUsage::SystemUpdateScript);
-		
-		// Once all system data interfaces are resolved, we resolve the systems scripts again, but filtered to just emitter scripts data interfaces.
-		HandleSystemScript(System.GetSystemSpawnScript(), ENiagaraScriptUsage::EmitterSpawnScript);
-		HandleSystemScript(System.GetSystemUpdateScript(), ENiagaraScriptUsage::EmitterUpdateScript);
+		TMap<FNiagaraVariableBase, FDataInterfaceSourceEmitterNamePair> VariableReadMap = EmitterIdToVariableReadsMap[FGuid()];
+		TMap<FNiagaraVariableBase, FDataInterfaceSourceEmitterNamePair> VariableAssignmentMap = EmitterIdToVariableAssignmentsMap[FGuid()];
+		TMap<FNiagaraVariableBase, FNiagaraVariableBase> VariableBindingMap = EmitterIdToVariableBindingsMap[FGuid()];
+		TMap<FNiagaraVariableBase, FNiagaraVariableBase> CompactedVariableBindingMap;
+		CompactBindings(VariableBindingMap, CompactedVariableBindingMap, OutErrorMessages);
 
-		for (const FNiagaraEmitterHandle& EmitterHandle : System.GetEmitterHandles())
+		UNiagaraScript* SystemSpawnScript = System->GetSystemSpawnScript();
+		UNiagaraScript* SystemUpdateScript = System->GetSystemUpdateScript();
+
+		FNiagaraScriptRuntimeCompiledData SystemSpawnRuntimeCompiledData;
+		FNiagaraScriptRuntimeCompiledData SystemUpdateRuntimeCompiledData;
+
+		ResolveDIsForScript(System, SystemSpawnScript, FString(), VariableReadMap, VariableAssignmentMap, CompactedVariableBindingMap, SystemSpawnRuntimeCompiledData, OutErrorMessages);
+		ResolveDIsForScript(System, SystemUpdateScript, FString(), VariableReadMap, VariableAssignmentMap, CompactedVariableBindingMap, SystemUpdateRuntimeCompiledData, OutErrorMessages);
+
+		OutScriptRuntimeCompiledDataMap.Add(FNiagaraScriptDataKey(FGuid(), SystemSpawnScript->GetUsage(), SystemSpawnScript->GetUsageId()), SystemSpawnRuntimeCompiledData);
+		OutScriptRuntimeCompiledDataMap.Add(FNiagaraScriptDataKey(FGuid(), SystemUpdateScript->GetUsage(), SystemUpdateScript->GetUsageId()), SystemUpdateRuntimeCompiledData);
+
+		for (const FNiagaraEmitterHandle& EmitterHandle : System->GetEmitterHandles())
 		{
-			FVersionedNiagaraEmitterData* VersionedNiagaraEmitterData = EmitterHandle.GetIsEnabled() ? EmitterHandle.GetEmitterData() : nullptr;
+			if (EmitterHandle.GetIsEnabled() == false)
+			{
+				continue;
+			}
+
+			FVersionedNiagaraEmitterData* VersionedNiagaraEmitterData = EmitterHandle.GetEmitterData();
 			if (VersionedNiagaraEmitterData == nullptr)
 			{
 				continue;
 			}
 
-			const FScriptDefaultData& EmitterScriptDefaultData = ScriptDefaultDataMap[EmitterHandle.GetId()];
-			TMap<FNiagaraCompileParameterReference, UNiagaraDataInterface*> EmitterExternalParameterDefaultMap;
-			TMap<FNiagaraCompileParameterReference, UNiagaraDataInterface*> EmitterUnresolvedExternalParameterDefaultMap;
-			CreateEmitterSpecificParameterDefaults(EmitterHandle.GetId(), SystemExternalParameterDefaultMap, EmitterExternalParameterDefaultMap);
-			CreateEmitterSpecificParameterDefaults(EmitterHandle.GetId(), SystemUnresolvedExternalParameterDefaultMap, EmitterUnresolvedExternalParameterDefaultMap);
+			TMap<FNiagaraVariableBase, FDataInterfaceSourceEmitterNamePair> EmitterVariableReadMap = EmitterIdToVariableReadsMap[EmitterHandle.GetId()];
+			TMap<FNiagaraVariableBase, FDataInterfaceSourceEmitterNamePair> EmitterVariableAssignmentMap = EmitterIdToVariableAssignmentsMap[EmitterHandle.GetId()];
+			TMap<FNiagaraVariableBase, FNiagaraVariableBase> EmitterVariableBindingMap = EmitterIdToVariableBindingsMap[EmitterHandle.GetId()];
+			TMap<FNiagaraVariableBase, FNiagaraVariableBase> EmitterCompactedVariableBindingMap;
+			CompactBindings(EmitterVariableBindingMap, EmitterCompactedVariableBindingMap, OutErrorMessages);
 
-			auto HandleParticleScript = [&OwnerForCreatedDIs,&OutScriptRuntimeCompiledDataMap,&EmitterHandle,&EmitterScriptDefaultData,
-				&EmitterExternalParameterDefaultMap,&EmitterUnresolvedExternalParameterDefaultMap]
-				(const UNiagaraScript* Script)
-			{
-				GenerateResolvedDIsForScript(&EmitterHandle, Script, EmitterScriptDefaultData, TOptional<ENiagaraScriptUsage>(), OwnerForCreatedDIs,
-					OutScriptRuntimeCompiledDataMap, EmitterExternalParameterDefaultMap, EmitterUnresolvedExternalParameterDefaultMap);
-			};
+			UNiagaraScript* SpawnScript = VersionedNiagaraEmitterData->SpawnScriptProps.Script;
+			UNiagaraScript* UpdateScript = VersionedNiagaraEmitterData->UpdateScriptProps.Script;
+			UNiagaraScript* GpuScript = VersionedNiagaraEmitterData->GetGPUComputeScript();
 
-			HandleParticleScript(VersionedNiagaraEmitterData->SpawnScriptProps.Script);
-			HandleParticleScript(VersionedNiagaraEmitterData->UpdateScriptProps.Script);
-			for (const FNiagaraEventScriptProperties& EventHandler : VersionedNiagaraEmitterData->GetEventHandlers())
-			{
-				HandleParticleScript(EventHandler.Script);
-			}
-			if (VersionedNiagaraEmitterData->SimTarget == ENiagaraSimTarget::GPUComputeSim)
+			FNiagaraScriptRuntimeCompiledData SpawnRuntimeCompiledData;
+			FNiagaraScriptRuntimeCompiledData UpdateRuntimeCompiledData;
+
+			ResolveDIsForScript(System, SpawnScript, EmitterHandle.GetUniqueInstanceName(), EmitterVariableReadMap, EmitterVariableAssignmentMap, EmitterCompactedVariableBindingMap, SpawnRuntimeCompiledData, OutErrorMessages);
+			ResolveDIsForScript(System, UpdateScript, EmitterHandle.GetUniqueInstanceName(), EmitterVariableReadMap, EmitterVariableAssignmentMap, EmitterCompactedVariableBindingMap, UpdateRuntimeCompiledData, OutErrorMessages);
+
+			if (VersionedNiagaraEmitterData->UsesInterpolatedSpawning())
 			{
-				HandleParticleScript(VersionedNiagaraEmitterData->GetGPUComputeScript());
+				TArray<FNiagaraScriptRuntimeCompiledData*> SourceScriptRuntimeCompiledData = { &UpdateRuntimeCompiledData };
+				ResolveInternalDataInterfaces(SourceScriptRuntimeCompiledData, SpawnRuntimeCompiledData);
 			}
 
-			// Collect any unresolved data interfaces created for the system here so that they are shared across particle scripts.
-			CollectSystemUnresolvedExternalParameterDefaults(EmitterUnresolvedExternalParameterDefaultMap, SystemUnresolvedExternalParameterDefaultMap);
-		}
-	}
+			OutScriptRuntimeCompiledDataMap.Add(FNiagaraScriptDataKey(EmitterHandle.GetId(), SpawnScript->GetUsage(), SpawnScript->GetUsageId()), SpawnRuntimeCompiledData);
+			OutScriptRuntimeCompiledDataMap.Add(FNiagaraScriptDataKey(EmitterHandle.GetId(), UpdateScript->GetUsage(), UpdateScript->GetUsageId()), UpdateRuntimeCompiledData);
 
-	/** Generates the aliased parameter variables on the resolved DI infos using the source parameter references and the emitter id to name map. */
-	void GenerateDIParameterVariablesForScript(TArray<FNiagaraScriptResolvedDataInterfaceInfo>& ResolvedDataInterfaces,
-		const FNiagaraEmitterHandle* OwningEmitterHandle,
-		const TMap<FGuid, FString>& EmitterHandleIdToNameMap)
-	{
-		for (FNiagaraScriptResolvedDataInterfaceInfo& ResolvedDataInterfaceData : ResolvedDataInterfaces)
-		{
-			// This is a bit inconsistent since emitter read variables are expected to be in the form Emitter.[VariableName], but emitter write variables are required to be in the
-			// form [EmitterName].[VariableName] to disambiguate them in the system parameter store.
-			if (ResolvedDataInterfaceData.ResolvedParameterReference.IsValid())
+			for (const FNiagaraEventScriptProperties& EventHandler : VersionedNiagaraEmitterData->GetEventHandlers())
 			{
-				ResolvedDataInterfaceData.ResolvedVariable = FNiagaraVariableBase(ResolvedDataInterfaceData.Type, ResolvedDataInterfaceData.ResolvedParameterReference.ParameterFullName);
+				UNiagaraScript* EventScript = EventHandler.Script;
+				FNiagaraScriptRuntimeCompiledData EventRuntimeCompiledData;
+
+				ResolveDIsForScript(System, EventScript, EmitterHandle.GetUniqueInstanceName(), EmitterVariableReadMap, EmitterVariableAssignmentMap, EmitterCompactedVariableBindingMap, EventRuntimeCompiledData, OutErrorMessages);
+				OutScriptRuntimeCompiledDataMap.Add(FNiagaraScriptDataKey(EmitterHandle.GetId(), EventScript->GetUsage(), EventScript->GetUsageId()), EventRuntimeCompiledData);
 			}
-			if (ResolvedDataInterfaceData.TargetParameterReference.IsValid())
-			{ 
-				if (ResolvedDataInterfaceData.TargetParameterReference.HasFlags(ENiagaraCompileParameterReferenceFlags::Emitter))
+
+			if (VersionedNiagaraEmitterData->SimTarget == ENiagaraSimTarget::GPUComputeSim && GpuScript != nullptr)
+			{
+				FNiagaraScriptRuntimeCompiledData GpuRuntimeCompiledData;
+				ResolveDIsForScript(System, GpuScript, EmitterHandle.GetUniqueInstanceName(), EmitterVariableReadMap, EmitterVariableAssignmentMap, EmitterCompactedVariableBindingMap, GpuRuntimeCompiledData, OutErrorMessages);
+
+				TArray<FNiagaraScriptRuntimeCompiledData> SimulationStageScriptRuntimeCompiledDataItems;
+				for (const UNiagaraSimulationStageBase* SimulationStage : VersionedNiagaraEmitterData->GetSimulationStages())
 				{
-					FString EmitterName;
-					if (OwningEmitterHandle != nullptr)
+					if (SimulationStage == nullptr || SimulationStage->bEnabled == false)
 					{
-						EmitterName = OwningEmitterHandle->GetUniqueInstanceName();
+						continue;
 					}
-					else
-					{
-						const FString* FoundEmitterName = EmitterHandleIdToNameMap.Find(ResolvedDataInterfaceData.TargetParameterReference.EmitterHandleId);
-						EmitterName = FoundEmitterName != nullptr ? *FoundEmitterName : FString();
-					}
-					if (EmitterName.IsEmpty() == false)
-					{
-						ResolvedDataInterfaceData.ParameterStoreVariable = FNiagaraUtilities::ResolveAliases(
-							FNiagaraVariable(ResolvedDataInterfaceData.Type, ResolvedDataInterfaceData.TargetParameterReference.ParameterFullName),
-							FNiagaraAliasContext().ChangeEmitterToEmitterName(EmitterName));
-					}
-				}
-				else
-				{
-					ResolvedDataInterfaceData.ParameterStoreVariable = FNiagaraVariableBase(ResolvedDataInterfaceData.Type, ResolvedDataInterfaceData.TargetParameterReference.ParameterFullName);
-				}
-			}
-			else
-			{
-				ResolvedDataInterfaceData.ParameterStoreVariable = FNiagaraVariableBase(ResolvedDataInterfaceData.Type, ResolvedDataInterfaceData.CompileName);
-			}
-		}
-	}
+					UNiagaraScript* SimulationStageScript = SimulationStage->Script;
+					FNiagaraScriptRuntimeCompiledData& SimulationStageRuntimeCompiledData = SimulationStageScriptRuntimeCompiledDataItems.AddDefaulted_GetRef();
 
-	void GenerateUserDIBindingsForScript(const TArray<FNiagaraScriptResolvedDataInterfaceInfo>& ResolvedDataInterfaces,
-		const FNiagaraParameterStore& SystemInstanceParameters,
-		TArray<FNiagaraResolvedUserDataInterfaceBinding>& OutResolvedUserDataInterfaceBindings)
-	{
-		int32 ResolvedDataInterfaceIndex = 0;
-		for (const FNiagaraScriptResolvedDataInterfaceInfo& ResolvedDataInterfaceData : ResolvedDataInterfaces)
-		{
-			if (ResolvedDataInterfaceData.ResolvedParameterReference.HasFlags(ENiagaraCompileParameterReferenceFlags::User))
-			{
-				const int32* UserParameterIndex = SystemInstanceParameters.FindParameterOffset(ResolvedDataInterfaceData.ResolvedVariable);
-				if (UserParameterIndex != nullptr && *UserParameterIndex != INDEX_NONE)
-				{
-					OutResolvedUserDataInterfaceBindings.Add(FNiagaraResolvedUserDataInterfaceBinding(*UserParameterIndex, ResolvedDataInterfaceIndex));
+					ResolveDIsForScript(System, SimulationStageScript, EmitterHandle.GetUniqueInstanceName(), EmitterVariableReadMap, EmitterVariableAssignmentMap, EmitterCompactedVariableBindingMap, SimulationStageRuntimeCompiledData, OutErrorMessages);
 				}
-			}
-			ResolvedDataInterfaceIndex++;
-		}
-	}
 
-	void GenerateNPCDIBindingsForScript(const TArray<FNiagaraScriptResolvedDataInterfaceInfo>& ResolvedDataInterfaces,
-		TArray<FNiagaraResolvedNpcDataInterfaceBinding>& OutResolvedNpcDataInterfaceBindings)
-	{
-		int32 ResolvedDataInterfaceIndex = 0;
-		for (const FNiagaraScriptResolvedDataInterfaceInfo& ResolvedDataInterfaceData : ResolvedDataInterfaces)
-		{
-			if (ResolvedDataInterfaceData.ResolvedParameterReference.IsValid() &&
-				ResolvedDataInterfaceData.ResolvedParameterReference.HasFlags(ENiagaraCompileParameterReferenceFlags::NPC))
-			{
-				FNiagaraResolvedNpcDataInterfaceBinding* NpcBinding = OutResolvedNpcDataInterfaceBindings.FindByPredicate([&ResolvedDataInterfaceData]
-					(const FNiagaraResolvedNpcDataInterfaceBinding& NpcBinding)
-				{
-					return NpcBinding.SourceNpcParameter == ResolvedDataInterfaceData.ResolvedVariable;
-				});
-				if (NpcBinding == nullptr)
+				TArray<FNiagaraScriptRuntimeCompiledData*> SourceScriptRuntimeCompiledData;
+				SourceScriptRuntimeCompiledData.Add(&SpawnRuntimeCompiledData);
+				SourceScriptRuntimeCompiledData.Add(&UpdateRuntimeCompiledData);
+				for (FNiagaraScriptRuntimeCompiledData& SimulationStageScriptRuntimeCompiledData : SimulationStageScriptRuntimeCompiledDataItems)
 				{
-					NpcBinding = &OutResolvedNpcDataInterfaceBindings.AddDefaulted_GetRef();
-					NpcBinding->SourceNpcParameter = ResolvedDataInterfaceData.ResolvedVariable;
+					SourceScriptRuntimeCompiledData.Add(&SimulationStageScriptRuntimeCompiledData);
 				}
-				NpcBinding->TargetScriptParameterStoreDataInterfaceIndices.AddUnique(ResolvedDataInterfaceIndex);
-			}
-			ResolvedDataInterfaceIndex++;
-		}
-	}
+				ResolveInternalDataInterfaces(SourceScriptRuntimeCompiledData, GpuRuntimeCompiledData);
 
-	void UpdatePostResolveData(
-		const UNiagaraSystem& System,
-		TMap<FNiagaraScriptDataKey, FNiagaraScriptRuntimeCompiledData>& OutScriptRuntimeCompiledDataMap)
-	{
-		TMap<FGuid, FString> EmitterHandleIdToNameMap;
-		for (const FNiagaraEmitterHandle& EmitterHandle : System.GetEmitterHandles())
-		{
-			EmitterHandleIdToNameMap.Add(EmitterHandle.GetId(), EmitterHandle.GetUniqueInstanceName());
+				GpuRuntimeCompiledData.SimTarget = ENiagaraSimTarget::GPUComputeSim;
+				OutScriptRuntimeCompiledDataMap.Add(FNiagaraScriptDataKey(EmitterHandle.GetId(), GpuScript->GetUsage(), GpuScript->GetUsageId()), GpuRuntimeCompiledData);
+			}
 		}
-
-		const FNiagaraParameterStore& SystemInstanceParameters = System.GetSystemCompiledData().InstanceParamStore;
-
-		System.ForEachScriptWithOwningContext([&EmitterHandleIdToNameMap,&SystemInstanceParameters,&OutScriptRuntimeCompiledDataMap]
-			(const UNiagaraSystem* OwningSystem, const FNiagaraEmitterHandle* OwningEmitterHandle, UNiagaraScript* Script)
-		{
-			FGuid EmitterHandleId = OwningEmitterHandle != nullptr ? OwningEmitterHandle->GetId() : FGuid();
-			FNiagaraScriptDataKey ScriptDataKey = FNiagaraScriptDataKey(EmitterHandleId, Script->GetUsage(), Script->GetUsageId());
-			FNiagaraScriptRuntimeCompiledData& ScriptRuntimeCompiledData = OutScriptRuntimeCompiledDataMap[ScriptDataKey];
-
-			GenerateDIParameterVariablesForScript(ScriptRuntimeCompiledData.ResolvedDataInterfaces, OwningEmitterHandle, EmitterHandleIdToNameMap);
-			GenerateUserDIBindingsForScript(ScriptRuntimeCompiledData.ResolvedDataInterfaces, SystemInstanceParameters, ScriptRuntimeCompiledData.ResolvedUserDataInterfaceBindings);
-			GenerateNPCDIBindingsForScript(ScriptRuntimeCompiledData.ResolvedDataInterfaces, ScriptRuntimeCompiledData.ResolvedNpcDataInterfaceBindings);
-		});
 	}
 
-	void ResolveDIs(UNiagaraSystem& System, TMap<FNiagaraScriptDataKey, FNiagaraScriptRuntimeCompiledData>& OutScriptRuntimeCompiledDataMap, TArray<FText>& OutErrorMessages)
+	void ResolveDIs(UNiagaraSystem* System, TMap<FNiagaraScriptDataKey, FNiagaraScriptRuntimeCompiledData>& OutScriptRuntimeCompiledDataMap, TArray<FText>& OutErrorMessages)
 	{
-		TMap<FGuid, FScriptDefaultData> ScriptDefaultDataMap;
-		TMap<FNiagaraCompileParameterReference, UNiagaraDataInterface*> SystemExternalParameterDefaultMap;
-		UObject& OwnerForCreateDIs = System;
-		CollectScriptDIResolveDataForSystem(System, OwnerForCreateDIs, ScriptDefaultDataMap, SystemExternalParameterDefaultMap, OutErrorMessages);
-		ResetResolvedDIsForSystem(System, OutScriptRuntimeCompiledDataMap);
-		GenerateResolvedDIsForSystem(System, OwnerForCreateDIs, ScriptDefaultDataMap, SystemExternalParameterDefaultMap, OutScriptRuntimeCompiledDataMap, OutErrorMessages);
-		UpdatePostResolveData(System, OutScriptRuntimeCompiledDataMap);
+		TMap<FGuid, TMap<FNiagaraVariableBase, FDataInterfaceSourceEmitterNamePair>> EmitterIdToVariableReadsMap;
+		TMap<FGuid, TMap<FNiagaraVariableBase, FDataInterfaceSourceEmitterNamePair>> EmitterIdToVariableAssignmentsMap;
+		TMap<FGuid, TMap<FNiagaraVariableBase, FNiagaraVariableBase>> EmitterIdToVariableBindingsMap;
+		FNiagaraResolveDIHelpers::CollectDIReadsAssignmentsAndBindings(System, EmitterIdToVariableReadsMap, EmitterIdToVariableAssignmentsMap, EmitterIdToVariableBindingsMap, OutErrorMessages);
+		FNiagaraResolveDIHelpers::ResolveDIsInternal(System, EmitterIdToVariableReadsMap, EmitterIdToVariableAssignmentsMap, EmitterIdToVariableBindingsMap, OutScriptRuntimeCompiledDataMap, OutErrorMessages);
 	}
 }
 
diff --git a/Engine/Plugins/FX/Niagara/Source/Niagara/Private/NiagaraResolveDIHelpers.h b/Engine/Plugins/FX/Niagara/Source/Niagara/Private/NiagaraResolveDIHelpers.h
--- a/Engine/Plugins/FX/Niagara/Source/Niagara/Private/NiagaraResolveDIHelpers.h
+++ b/Engine/Plugins/FX/Niagara/Source/Niagara/Private/NiagaraResolveDIHelpers.h
@@ -5,17 +5,18 @@
 #include "HAL/Platform.h"
 #include "Misc/CoreMiscDefines.h"
 
+#include "NiagaraCommon.h"
+#include "NiagaraScriptRuntimeCompiledData.h"
 #include "NiagaraTypes.h"
 
-struct FNiagaraScriptRuntimeCompiledData;
-struct FNiagaraScriptDataKey;
-class  UNiagaraSystem;
+class UNiagaraDataInterface;
+class UNiagaraSystem;
 
 #if WITH_EDITORONLY_DATA
 
 namespace FNiagaraResolveDIHelpers
 {
-	void ResolveDIs(UNiagaraSystem& System, TMap<FNiagaraScriptDataKey, FNiagaraScriptRuntimeCompiledData>& OutScriptRuntimeCompiledDataMap, TArray<FText>& OutErrorMessages);
+	void ResolveDIs(UNiagaraSystem* System, TMap<FNiagaraScriptDataKey, FNiagaraScriptRuntimeCompiledData>& OutScriptRuntimeCompiledDataMap, TArray<FText>& OutErrorMessages);
 }
 
 #endif
\ No newline at end of file
diff --git a/Engine/Plugins/FX/Niagara/Source/Niagara/Private/NiagaraScript.cpp b/Engine/Plugins/FX/Niagara/Source/Niagara/Private/NiagaraScript.cpp
--- a/Engine/Plugins/FX/Niagara/Source/Niagara/Private/NiagaraScript.cpp
+++ b/Engine/Plugins/FX/Niagara/Source/Niagara/Private/NiagaraScript.cpp
@@ -1483,9 +1483,9 @@ void UNiagaraScript::ComputeVMCompilationId_EmitterShared(FNiagaraVMExecutableDa
 
 void UNiagaraScript::InvalidateAndEmptyDefaultDataInterfaces(int32 SlackCount, EDataInterfaceInvalidationMode InvalidationMode)
 {
-	DataInterfacesPendingInvalidation.Reserve(DataInterfacesPendingInvalidation.Num() + CachedGraphDefaultDataInterfaces.Num());
+	DataInterfacesPendingInvalidation.Reserve(DataInterfacesPendingInvalidation.Num() + CachedDefaultDataInterfaces.Num());
 
-	for (FNiagaraScriptDataInterfaceDefaultInfo& DefaultDataInterface : CachedGraphDefaultDataInterfaces)
+	for (FNiagaraScriptDataInterfaceInfo& DefaultDataInterface : CachedDefaultDataInterfaces)
 	{
 		if (DefaultDataInterface.DataInterface)
 		{
@@ -1496,7 +1496,7 @@ void UNiagaraScript::InvalidateAndEmptyDefaultDataInterfaces(int32 SlackCount, E
 		}
 	}
 
-	CachedGraphDefaultDataInterfaces.Empty(SlackCount);
+	CachedDefaultDataInterfaces.Empty(SlackCount);
 
 	if (InvalidationMode == EDataInterfaceInvalidationMode::ProcessInvaliation)
 	{
@@ -1945,7 +1945,7 @@ void UNiagaraScript::PostLoad()
 
 #if WITH_EDITORONLY_DATA
 	// Because we might be using these cached data interfaces, we need to make sure that they are properly postloaded.
-	for (FNiagaraScriptDataInterfaceDefaultInfo& Info : CachedGraphDefaultDataInterfaces)
+	for (FNiagaraScriptDataInterfaceInfo& Info : CachedDefaultDataInterfaces)
 	{
 		if (Info.DataInterface)
 		{
@@ -2650,6 +2650,24 @@ FVersionedNiagaraScriptData* FVersionedNiagaraScript::GetScriptData() const
 }
 #endif
 
+UNiagaraDataInterface* ResolveDataInterface(const TMap<FName, UNiagaraDataInterface*>& ObjectNameMap, FName VariableName)
+{
+	UNiagaraDataInterface* const* FoundDI = ObjectNameMap.Find(VariableName);
+	if (FoundDI && *(FoundDI))
+	{
+		return *FoundDI;
+	}
+	return nullptr;
+}
+
+void DumpNameMap(const TMap<FName, UNiagaraDataInterface*>& ObjectNameMap)
+{
+	for (const TPair<FName, UNiagaraDataInterface*>& Pair : ObjectNameMap)
+	{
+		UE_LOG(LogNiagara, Log, TEXT("%s -> %s"), *Pair.Key.ToString(), *GetPathNameSafe(Pair.Value));
+	}
+}
+
 FName ResolveEmitterAlias(const FName& InName, const FString& InAlias)
 {
 	// If the alias is empty than the name can't be resolved.
@@ -2665,7 +2683,7 @@ FName ResolveEmitterAlias(const FName& InName, const FString& InAlias)
 	return Var.GetName();
 }
 
-void UNiagaraScript::SetVMCompilationResults(const FNiagaraVMExecutableDataId& InCompileId, FNiagaraVMExecutableData& InScriptVM, const TMap<FNiagaraGraphInputDataInterfaceKey, UNiagaraDataInterface*>& GraphInputDataInterfaces, bool bApplyRapidIterationParameters)
+void UNiagaraScript::SetVMCompilationResults(const FNiagaraVMExecutableDataId& InCompileId, FNiagaraVMExecutableData& InScriptVM, const FString& EmitterUniqueName, const TMap<FName, UNiagaraDataInterface*>& ObjectNameMap, bool bApplyRapidIterationParameters)
 {
 	CachedScriptVMId = InCompileId;
 	CachedScriptVM = InScriptVM;
@@ -2688,24 +2706,37 @@ void UNiagaraScript::SetVMCompilationResults(const FNiagaraVMExecutableDataId& I
 	InvalidateAndEmptyDefaultDataInterfaces(CachedScriptVM.DataInterfaceInfo.Num(), EDataInterfaceInvalidationMode::ProcessInvaliation);
 	for (const FNiagaraScriptDataInterfaceCompileInfo& Info : CachedScriptVM.DataInterfaceInfo)
 	{
-		if (Info.SourceGraphInputDataInterfaceKey.IsValid())
+		int32 Idx = CachedDefaultDataInterfaces.AddDefaulted();
+		CachedDefaultDataInterfaces[Idx].UserPtrIdx = Info.UserPtrIdx;
+		CachedDefaultDataInterfaces[Idx].Name = ResolveEmitterAlias(Info.Name, EmitterUniqueName);
+		CachedDefaultDataInterfaces[Idx].CompileName = Info.Name;
+		CachedDefaultDataInterfaces[Idx].Type = Info.Type;
+		CachedDefaultDataInterfaces[Idx].SourceEmitterName = Info.SourceEmitterName;
+		CachedDefaultDataInterfaces[Idx].RegisteredParameterMapRead = ResolveEmitterAlias(Info.RegisteredParameterMapRead, EmitterUniqueName);
+		CachedDefaultDataInterfaces[Idx].RegisteredParameterMapWrite = ResolveEmitterAlias(Info.RegisteredParameterMapWrite, EmitterUniqueName);
+
+		// We compiled it just a bit ago, so we should be able to resolve it from the table that we passed in.
+		UNiagaraDataInterface* FindDIById = ResolveDataInterface(ObjectNameMap, CachedDefaultDataInterfaces[Idx].Name);
+		if (FindDIById != nullptr )
 		{
-			int32 Idx = CachedGraphDefaultDataInterfaces.AddDefaulted();
-			CachedGraphDefaultDataInterfaces[Idx].Id = Info.Id;
-			UNiagaraDataInterface*const* GraphInputDataInterfacePtr = GraphInputDataInterfaces.Find(Info.SourceGraphInputDataInterfaceKey);
-			if (GraphInputDataInterfacePtr != nullptr)
-			{
-				CachedGraphDefaultDataInterfaces[Idx].DataInterface = CopyDataInterface(*GraphInputDataInterfacePtr, this);
-			}
-			else
-			{
-				// Use the CDO since we didn't have a default..
-				UObject* Obj = Info.Type.GetClass()->GetDefaultObject(true);
-				CachedGraphDefaultDataInterfaces[Idx].DataInterface = Cast<UNiagaraDataInterface>(CopyDataInterface(CastChecked<UNiagaraDataInterface>(Obj), this));
+			CachedDefaultDataInterfaces[Idx].DataInterface = CopyDataInterface(FindDIById, this);
+			check(CachedDefaultDataInterfaces[Idx].DataInterface != nullptr);
+		}
 
-				UE_LOG(LogNiagara, Display, TEXT("We somehow ended up with an input data interface that we couldn't match post compile. This shouldn't happen. Creating a dummy to prevent crashes. DataInterfaceInfoName:%s Object:%s"), *Info.Name.ToString(), *GetPathNameSafe(this));
+		if (CachedDefaultDataInterfaces[Idx].DataInterface == nullptr)
+		{
+			// Use the CDO since we didn't have a default..
+			UObject* Obj = Info.Type.GetClass()->GetDefaultObject(true);
+			CachedDefaultDataInterfaces[Idx].DataInterface = Cast<UNiagaraDataInterface>(CopyDataInterface(CastChecked<UNiagaraDataInterface>(Obj), this));
+
+			if (Info.bIsPlaceholder == false)
+			{
+				UE_LOG(LogNiagara, Display, TEXT("We somehow ended up with a data interface that we couldn't match post compile. This shouldn't happen. Creating a dummy to prevent crashes. DataInterfaceInfoName:%s Object:%s"), *Info.Name.ToString(), *GetPathNameSafe(this));
+				UE_LOG(LogNiagara, Log, TEXT("Object to Name map contents:"));
+				DumpNameMap(ObjectNameMap);
 			}
 		}
+		check(CachedDefaultDataInterfaces[Idx].DataInterface != nullptr);
 	}
 
 	CachedDefaultUObjects = CachedScriptVM.UObjectInfos;
@@ -3009,7 +3040,7 @@ void UNiagaraScript::RequestCompile(const FGuid& ScriptVersion, bool bForceCompi
 			if (BinaryToExecData(this, OutData, ExeData))
 			{
 				COOK_STAT(Timer.AddHit(OutData.Num()));
-				SetVMCompilationResults(LastGeneratedVMId, ExeData, RequestDuplicateData->GetGraphInputDataInterfaces(), false);
+				SetVMCompilationResults(LastGeneratedVMId, ExeData, FString(), RequestDuplicateData->GetObjectNameMap(), false);
 				return;
 			}
 		}
@@ -3023,7 +3054,7 @@ void UNiagaraScript::RequestCompile(const FGuid& ScriptVersion, bool bForceCompi
 		TSharedPtr<FNiagaraVMExecutableData> ExeData = NiagaraModule.GetCompileJobResult(JobHandle, true, ScriptMetrics);
 		if (ExeData)
 		{
-			SetVMCompilationResults(LastGeneratedVMId, *ExeData, RequestDuplicateData->GetGraphInputDataInterfaces(), false);
+			SetVMCompilationResults(LastGeneratedVMId, *ExeData, FString(), RequestDuplicateData->GetObjectNameMap(), false);
 			// save result to the ddc
 			if (ExecToBinaryData(this, OutData, *ExeData))
 			{
@@ -3263,11 +3294,11 @@ bool UNiagaraScript::ShouldCompile(EShaderPlatform Platform) const
 
 	// Add all data interfaces
 	TSet<UClass*> DIUniqueClasses;
-	for (const FNiagaraScriptDataInterfaceCompileInfo& DataInterfaceInfo : CachedScriptVM.DataInterfaceInfo)
+	for (const FNiagaraScriptDataInterfaceInfo& DataInterfaceInfo : CachedDefaultDataInterfaces)
 	{
-		if (UClass* DataInterfaceClass = DataInterfaceInfo.Type.GetClass())
+		if (UNiagaraDataInterface* DataInterface = DataInterfaceInfo.DataInterface)
 		{
-			DIUniqueClasses.Add(DataInterfaceClass);
+			DIUniqueClasses.Add(DataInterface->GetClass());
 		}
 	}
 
@@ -3293,11 +3324,11 @@ void UNiagaraScript::ModifyCompilationEnvironment(EShaderPlatform Platform, stru
 #if WITH_EDITOR
 	// Add all data interfaces
 	TSet<UClass*> DIUniqueClasses;
-	for (const FNiagaraScriptDataInterfaceCompileInfo& DataInterfaceInfo : CachedScriptVM.DataInterfaceInfo)
+	for ( const FNiagaraScriptDataInterfaceInfo& DataInterfaceInfo : CachedDefaultDataInterfaces )
 	{
-		if (UClass* DataInterfaceClass = DataInterfaceInfo.Type.GetClass())
+		if ( UNiagaraDataInterface* DataInterface = DataInterfaceInfo.DataInterface )
 		{
-			DIUniqueClasses.Add(DataInterfaceClass);
+			DIUniqueClasses.Add(DataInterface->GetClass());
 		}
 	}
 
@@ -3653,6 +3684,12 @@ void UNiagaraScript::SyncAliases(const FNiagaraAliasContext& ResolveAliasesConte
 	// Sync up any simulation stage name references.
 	for (FSimulationStageMetaData& SimStageMetaData : GetVMExecutableData().SimulationStageMetaData)
 	{
+		if ( !SimStageMetaData.IterationDataInterface.IsNone() )
+		{
+			FNiagaraVariable Var(FNiagaraTypeDefinition(UNiagaraDataInterface::StaticClass()), SimStageMetaData.IterationDataInterface);
+			SimStageMetaData.IterationDataInterface = FNiagaraUtilities::ResolveAliases(Var, ResolveAliasesContext).GetName();
+		}
+
 		if (!SimStageMetaData.EnabledBinding.IsNone())
 		{
 			FNiagaraVariable Var(FNiagaraTypeDefinition::GetIntDef(), SimStageMetaData.EnabledBinding);
@@ -3772,13 +3809,13 @@ bool UNiagaraScript::SynchronizeExecutablesWithCompilation(const UNiagaraScript*
 		CachedScriptVMId = Script->CachedScriptVMId;
 		CachedParameterCollectionReferences = Script->CachedParameterCollectionReferences;
 		CachedDefaultUObjects = Script->CachedDefaultUObjects;
-		InvalidateAndEmptyDefaultDataInterfaces(Script->CachedGraphDefaultDataInterfaces.Num(), EDataInterfaceInvalidationMode::ProcessInvaliation);
-		for (const FNiagaraScriptDataInterfaceDefaultInfo& Info : Script->CachedGraphDefaultDataInterfaces)
+		InvalidateAndEmptyDefaultDataInterfaces(Script->CachedDefaultDataInterfaces.Num(), EDataInterfaceInvalidationMode::ProcessInvaliation);
+		for (const FNiagaraScriptDataInterfaceInfo& Info : Script->CachedDefaultDataInterfaces)
 		{
-			FNiagaraScriptDataInterfaceDefaultInfo AddInfo;
+			FNiagaraScriptDataInterfaceInfo AddInfo;
 			AddInfo = Info;
 			AddInfo.DataInterface = CopyDataInterface(Info.DataInterface, this);
-			CachedGraphDefaultDataInterfaces.Add(AddInfo);
+			CachedDefaultDataInterfaces.Add(AddInfo);
 		}
 
 		GenerateStatIDs();
diff --git a/Engine/Plugins/FX/Niagara/Source/Niagara/Private/NiagaraScriptRuntimeCompiledData.cpp b/Engine/Plugins/FX/Niagara/Source/Niagara/Private/NiagaraScriptRuntimeCompiledData.cpp
--- a/Engine/Plugins/FX/Niagara/Source/Niagara/Private/NiagaraScriptRuntimeCompiledData.cpp
+++ b/Engine/Plugins/FX/Niagara/Source/Niagara/Private/NiagaraScriptRuntimeCompiledData.cpp
@@ -22,7 +22,7 @@ bool FNiagaraScriptRuntimeCompiledData::ValidateWithScript(UNiagaraScript& InScr
 		for (const FNiagaraScriptResolvedDataInterfaceInfo& ResolvedDataInterfaceData : ResolvedDataInterfaces)
 		{
 			UE_LOG(LogNiagara, Log, TEXT("Name:%s, Type: %s, Path:%s"),
-				*ResolvedDataInterfaceData.CompileName.ToString(), *ResolvedDataInterfaceData.Type.GetName(),
+				*ResolvedDataInterfaceData.CompileName.ToString(), *ResolvedDataInterfaceData.ParameterStoreVariable.GetType().GetName(),
 				ResolvedDataInterfaceData.ResolvedDataInterface != nullptr
 				? *ResolvedDataInterfaceData.ResolvedDataInterface->GetPathName()
 				: TEXT("None"));
diff --git a/Engine/Plugins/FX/Niagara/Source/Niagara/Private/NiagaraSystem.cpp b/Engine/Plugins/FX/Niagara/Source/Niagara/Private/NiagaraSystem.cpp
--- a/Engine/Plugins/FX/Niagara/Source/Niagara/Private/NiagaraSystem.cpp
+++ b/Engine/Plugins/FX/Niagara/Source/Niagara/Private/NiagaraSystem.cpp
@@ -1956,7 +1956,7 @@ void UNiagaraSystem::GetAssetRegistryTags(FAssetRegistryTagsContext Context) con
 			{
 				for (FNiagaraScriptDataInterfaceCompileInfo& Info : Script->GetVMExecutableData().DataInterfaceInfo)
 				{
-					const UNiagaraDataInterface* DefaultDataInterface = Info.GetDefaultDataInterface();
+					UNiagaraDataInterface* DefaultDataInterface = Info.GetDefaultDataInterface();
 					DataInterfaces.AddUnique(DefaultDataInterface);
 				}
 			}
@@ -2596,7 +2596,7 @@ void UNiagaraSystem::CacheFromCompiledData()
 				FStaticBuffersCachingFlags& CachingFlagsForDataInterface = DataInterfaceToStaticCachingFlagsMap.FindOrAdd(DataInterface);
 				CachingFlagsForDataInterface.bIsUsedByCpu |= bUsedByCPU;
 				CachingFlagsForDataInterface.bIsUsedByGpu |= bUsedByGPU;
-				CachingFlagsForDataInterface.bIsExternalDataInterface |= DataInterfaceData.ResolvedParameterReference.HasAnyFlags(ENiagaraCompileParameterReferenceFlags::User | ENiagaraCompileParameterReferenceFlags::NPC);
+				CachingFlagsForDataInterface.bIsExternalDataInterface |= DataInterfaceData.ResolvedVariable.IsInExternalNamespace();
 			}
 		}
 	});
@@ -2648,7 +2648,7 @@ void UNiagaraSystem::InitScriptCompiledData()
 	ScriptRuntimeCompiledDataForEditor = NewObject<UNiagaraScriptRuntimeCompiledDataEditorReference>(this);
 
 	TArray<FText> ErrorMessages;
-	FNiagaraResolveDIHelpers::ResolveDIs(*this, ScriptRuntimeCompiledDataForEditor->ScriptRuntimeCompiledDataMap, ErrorMessages);
+	FNiagaraResolveDIHelpers::ResolveDIs(this, ScriptRuntimeCompiledDataForEditor->ScriptRuntimeCompiledDataMap, ErrorMessages);
 	if (ErrorMessages.Num() > 0)
 	{
 		INiagaraModule& NiagaraModule = FModuleManager::GetModuleChecked<INiagaraModule>("Niagara");
diff --git a/Engine/Plugins/FX/Niagara/Source/Niagara/Private/NiagaraSystemInstance.cpp b/Engine/Plugins/FX/Niagara/Source/Niagara/Private/NiagaraSystemInstance.cpp
--- a/Engine/Plugins/FX/Niagara/Source/Niagara/Private/NiagaraSystemInstance.cpp
+++ b/Engine/Plugins/FX/Niagara/Source/Niagara/Private/NiagaraSystemInstance.cpp
@@ -977,16 +977,6 @@ UNiagaraParameterCollectionInstance* FNiagaraSystemInstance::GetParameterCollect
 	return SystemSimulation->GetParameterCollectionInstance(Collection);
 }
 
-void FNiagaraSystemInstance::BindScriptParameterCollections(UNiagaraScript* Script, const FNiagaraScriptRuntimeData* ScriptRuntimeData, FNiagaraParameterStore& DestStore)
-{
-	SystemSimulation->BindScriptParameterCollections(Script, ScriptRuntimeData, DestStore);
-}
-
-void FNiagaraSystemInstance::UnbindScriptParameterCollections(UNiagaraScript* Script, FNiagaraParameterStore& DestStore)
-{
-	SystemSimulation->UnbindScriptParameterCollections(Script, DestStore);
-}
-
 void FNiagaraSystemInstance::AdvanceSimulation(int32 TickCountToSimulate, float TickDeltaSeconds)
 {
 	if (TickCountToSimulate > 0 && !IsPaused())
diff --git a/Engine/Plugins/FX/Niagara/Source/Niagara/Private/NiagaraSystemSimulation.cpp b/Engine/Plugins/FX/Niagara/Source/Niagara/Private/NiagaraSystemSimulation.cpp
--- a/Engine/Plugins/FX/Niagara/Source/Niagara/Private/NiagaraSystemSimulation.cpp
+++ b/Engine/Plugins/FX/Niagara/Source/Niagara/Private/NiagaraSystemSimulation.cpp
@@ -6,7 +6,6 @@
 #include "NiagaraModule.h"
 #include "NiagaraTypes.h"
 #include "NiagaraParameterCollection.h"
-#include "NiagaraScriptRuntimeData.h"
 #include "NiagaraSimCache.h"
 #include "NiagaraSimulationTaskPriority.h"
 #include "NiagaraSystemEmitterState.h"
@@ -653,8 +652,35 @@ bool FNiagaraSystemSimulation::Init(UNiagaraSystem* InSystem, UWorld* InWorld, b
 		{
 			//SCOPE_CYCLE_COUNTER(STAT_NiagaraSystemSim_Init_BindParams);
 
-			BindScriptParameterCollections(SpawnScript, SpawnExecContext->GetScriptRuntimeData(), SpawnExecContext->Parameters);
-			BindScriptParameterCollections(UpdateScript, UpdateExecContext->GetScriptRuntimeData(), UpdateExecContext->Parameters);
+			//Bind parameter collections.
+			auto BindParameterColleciton = [&](UNiagaraParameterCollection* Collection, FNiagaraParameterStore& DestStore)
+			{
+				if (Collection)
+				{
+					if (UNiagaraParameterCollectionInstance* CollectionInst = GetParameterCollectionInstance(Collection))
+					{
+						CollectionInst->GetParameterStore().Bind(&DestStore);
+					}
+					else
+					{
+						UE_LOG(LogNiagara, Error, TEXT("Attempting to bind system simulation to a null parameter collection instance | Collection: %s | System: %s |"), *Collection->GetPathName(), *System->GetPathName());
+					}
+				}
+				else
+				{
+					UE_LOG(LogNiagara, Error, TEXT("Attempting to bind system simulation to a null parameter collection | System: %s |"), *System->GetPathName());
+				}
+			};
+
+			//Bind parameter collections.
+			for (UNiagaraParameterCollection* Collection : SpawnScript->GetCachedParameterCollectionReferences())
+			{
+				BindParameterColleciton(Collection, SpawnExecContext->Parameters);
+			}
+			for (UNiagaraParameterCollection* Collection : UpdateScript->GetCachedParameterCollectionReferences())
+			{
+				BindParameterColleciton(Collection, UpdateExecContext->Parameters);
+			}
 
 			SpawnScript->RapidIterationParameters.Bind(&SpawnExecContext->Parameters);
 			UpdateScript->RapidIterationParameters.Bind(&UpdateExecContext->Parameters);
@@ -2581,85 +2607,3 @@ ENiagaraGPUTickHandlingMode FNiagaraSystemSimulation::GetGPUTickHandlingMode()co
 
 	return ENiagaraGPUTickHandlingMode::None;
 }
-
-void FNiagaraSystemSimulation::BindScriptParameterCollections(UNiagaraScript* Script, const FNiagaraScriptRuntimeData* ScriptRuntimeData, FNiagaraParameterStore& DestStore)
-{
-	for (UNiagaraParameterCollection* Collection : Script->GetCachedParameterCollectionReferences())
-	{
-		if (Collection == nullptr)
-		{
-			UE_LOG(LogNiagara, Warning, TEXT("Null cached parameter collection reference.\nScript: %s"), *Script->GetPathName());
-			continue;
-		}
-
-		UNiagaraParameterCollectionInstance* Inst = GetParameterCollectionInstance(Collection);
-		if (Inst == nullptr)
-		{
-			UE_LOG(LogNiagara, Warning, TEXT("Null parameter collection instance.\nScript: %s\nCollection: %s"),
-				*Script->GetPathName(), *Collection->GetPathName());
-			continue;
-		}
-
-		Inst->GetParameterStore().Bind(&DestStore);
-		ResolveNpcDataInterfaceBindingsForScript(Script, ScriptRuntimeData, Inst, DestStore);
-	}
-}
-
-void FNiagaraSystemSimulation::UnbindScriptParameterCollections(UNiagaraScript* Script, FNiagaraParameterStore& DestStore)
-{
-	for (UNiagaraParameterCollection* Collection : Script->GetCachedParameterCollectionReferences())
-	{
-		if (Collection == nullptr)
-		{
-			continue;
-		}
-
-		UNiagaraParameterCollectionInstance* Inst = GetParameterCollectionInstance(Collection);
-		if (Inst == nullptr)
-		{
-			continue;
-		}
-
-		Inst->GetParameterStore().Unbind(&DestStore);
-	}
-}
-
-void FNiagaraSystemSimulation::ResolveNpcDataInterfaceBindingsForScript(
-	const UNiagaraScript* Script,
-	const FNiagaraScriptRuntimeData* ScriptRuntimeData,
-	UNiagaraParameterCollectionInstance* CollectionInstance,
-	FNiagaraParameterStore& DestStore) const
-{
-	if (DestStore.GetDataInterfaces().Num() == 0 || ScriptRuntimeData == nullptr)
-	{
-		return;
-	}
-
-	for (const FNiagaraResolvedNpcDataInterfaceBinding& ResolvedNpcDataInterfaceBinding : ScriptRuntimeData->GetResolvedNpcDataInterfaceBindings())
-	{
-		UNiagaraDataInterface* CollectionDataInterface = CollectionInstance->GetParameterStore().GetDataInterface(ResolvedNpcDataInterfaceBinding.SourceNpcParameter);
-		if (CollectionDataInterface == nullptr)
-		{
-			UE_LOG(LogNiagara, Warning, TEXT("Missing parameter collection data interface.\nScript:%s\nCollection:%s\nParameter:%s"),
-				*Script->GetPathName(),
-				*CollectionInstance->GetPathName(),
-				*ResolvedNpcDataInterfaceBinding.SourceNpcParameter.GetName().ToString());
-			continue;
-		}
-
-		for (int32 TargetDataInterfaceIndex : ResolvedNpcDataInterfaceBinding.TargetScriptParameterStoreDataInterfaceIndices)
-		{
-			if (TargetDataInterfaceIndex >= 0 && TargetDataInterfaceIndex < DestStore.GetDataInterfaces().Num())
-			{
-				DestStore.SetDataInterface(CollectionDataInterface, TargetDataInterfaceIndex);
-			}
-			else
-			{
-				UE_LOG(LogNiagara, Warning, TEXT("Failed to bind parameter collection data interface.\nScript:%s\nCollection:%s\nParameter:%s"),
-					*Script->GetPathName(),
-					*CollectionInstance->GetPathName(),
-					*ResolvedNpcDataInterfaceBinding.SourceNpcParameter.GetName().ToString());
-			}
-		}
-	}
-}
diff --git a/Engine/Plugins/FX/Niagara/Source/Niagara/Public/NiagaraCommon.h b/Engine/Plugins/FX/Niagara/Source/Niagara/Public/NiagaraCommon.h
--- a/Engine/Plugins/FX/Niagara/Source/Niagara/Public/NiagaraCommon.h
+++ b/Engine/Plugins/FX/Niagara/Source/Niagara/Public/NiagaraCommon.h
@@ -17,9 +17,6 @@
 #include "NiagaraTypeRegistry.h"
 #include "NiagaraCore.h"
 #include "Templates/RefCounting.h"
-#include "NiagaraScriptCompileParameterReference.h"
-#include "NiagaraScriptDataInterfaceCompileId.h"
-#include "NiagaraScriptUsage.h"
 
 #include "NiagaraCommon.generated.h"
 
@@ -796,20 +793,44 @@ struct FNiagaraExternalUObjectInfo
 };
 
 USTRUCT()
-struct FNiagaraScriptDataInterfaceDefaultInfo
+struct FNiagaraScriptDataInterfaceInfo
 {
 	GENERATED_USTRUCT_BODY()
 public:
-	FNiagaraScriptDataInterfaceDefaultInfo()
+	FNiagaraScriptDataInterfaceInfo()
 		: DataInterface(nullptr)
+		, Name(NAME_None)
+		, UserPtrIdx(INDEX_NONE)
 	{
+
 	}
 
 	UPROPERTY()
-	FNiagaraScriptDataInterfaceCompileId Id;
+	TObjectPtr<class UNiagaraDataInterface> DataInterface;
+	
+	UPROPERTY()
+	FName Name;
+
+	UPROPERTY()
+	FName CompileName;
 
+	/** Index of the user pointer for this data interface. */
 	UPROPERTY()
-	TObjectPtr<class UNiagaraDataInterface> DataInterface;
+	int32 UserPtrIdx;
+
+	UPROPERTY()
+	FNiagaraTypeDefinition Type;
+
+	UPROPERTY()
+	FName RegisteredParameterMapRead;
+
+	UPROPERTY()
+	FName RegisteredParameterMapWrite;
+
+	UPROPERTY()
+	FString SourceEmitterName;
+
+	NIAGARA_API bool IsUserDataInterface() const;
 };
 
 USTRUCT()
@@ -836,52 +857,26 @@ struct FNiagaraResolvedUserDataInterfaceBinding
 	}
 };
 
-USTRUCT()
-struct FNiagaraResolvedNpcDataInterfaceBinding
-{
-	GENERATED_BODY()
-
-	UPROPERTY()
-	FNiagaraVariableBase SourceNpcParameter;
-
-	UPROPERTY()
-	TArray<int32> TargetScriptParameterStoreDataInterfaceIndices;
-
-	FNiagaraResolvedNpcDataInterfaceBinding()
-	{
-	}
-
-	FNiagaraResolvedNpcDataInterfaceBinding(const FNiagaraVariableBase& InSourceNpcParameter)
-		: SourceNpcParameter(InSourceNpcParameter)
-	{
-	}
-};
-
 USTRUCT()
 struct FNiagaraScriptResolvedDataInterfaceInfo
 {
 	GENERATED_USTRUCT_BODY();
 
 	FNiagaraScriptResolvedDataInterfaceInfo()
-		: ResolvedDataInterface(nullptr)
+		: bIsInternal(false)
+		, ResolvedDataInterface(nullptr)
 		, UserPtrIdx(-1)
 	{
 	}
 
 	UPROPERTY()
-	FGuid OwningEmitterHandleId;
+	FName Name;
 
 	UPROPERTY()
 	FName CompileName;
 
 	UPROPERTY()
-	FNiagaraTypeDefinition Type;
-
-	UPROPERTY()
-	FNiagaraCompileParameterReference ResolvedParameterReference;
-
-	UPROPERTY()
-	FNiagaraCompileParameterReference TargetParameterReference;
+	FString ResolvedSourceEmitterName;
 
 	UPROPERTY()
 	FNiagaraVariableBase ResolvedVariable;
@@ -889,6 +884,9 @@ struct FNiagaraScriptResolvedDataInterfaceInfo
 	UPROPERTY()
 	FNiagaraVariableBase ParameterStoreVariable;
 
+	UPROPERTY()
+	bool bIsInternal;
+
 	UPROPERTY()
 	TObjectPtr<UNiagaraDataInterface> ResolvedDataInterface;
 
@@ -898,55 +896,21 @@ struct FNiagaraScriptResolvedDataInterfaceInfo
 	NIAGARA_API bool NeedsPerInstanceBinding() const;
 };
 
-USTRUCT()
-struct FNiagaraGraphInputDataInterfaceKey
-{
-	GENERATED_BODY()
-
-	FNiagaraGraphInputDataInterfaceKey()
-	{
-	}
-
-	FNiagaraGraphInputDataInterfaceKey(const FName& InGraphPath, const FNiagaraVariableBase& InInputVariable)
-		: GraphPath(InGraphPath)
-		, InputVariable(InInputVariable)
-	{
-	}
-
-	bool operator==(const FNiagaraGraphInputDataInterfaceKey& Other) const
-	{
-		return GraphPath == Other.GraphPath && InputVariable == Other.InputVariable;
-	}
-
-	friend uint32 GetTypeHash(const FNiagaraGraphInputDataInterfaceKey& Ref)
-	{
-		return HashCombine(GetTypeHash(Ref.GraphPath), GetTypeHash(Ref.InputVariable));
-	}
-
-	bool IsValid() const { return GraphPath != NAME_None && InputVariable.IsValid(); }
-
-	UPROPERTY()
-	FName GraphPath;
-
-	UPROPERTY()
-	FNiagaraVariableBase InputVariable;
-};
-
 USTRUCT()
 struct FNiagaraScriptDataInterfaceCompileInfo
 {
 	GENERATED_USTRUCT_BODY()
 public:
 	FNiagaraScriptDataInterfaceCompileInfo()
-		: UserPtrIdx(INDEX_NONE)
-		, SourceScriptDIIndex(INDEX_NONE)
+		: Name(NAME_None)
+		, UserPtrIdx(INDEX_NONE)
+		, bIsPlaceholder(false)
 	{
+
 	}
 
-	// A unique id for this data interface based on source information when it was first encountered.  This will be unique across the entire owning
-	// object, either the system for system assets, or the script for script assets.
 	UPROPERTY()
-	FNiagaraScriptDataInterfaceCompileId Id;
+	FName Name;
 
 	/** Index of the user pointer for this data interface. */
 	UPROPERTY()
@@ -955,29 +919,32 @@ public:
 	UPROPERTY()
 	FNiagaraTypeDefinition Type;
 
+	// Removed from cooked builds, if we need to add this back the TMap<FName, FName> FunctionSpecifiers should be replaced with an array
+#if WITH_EDITORONLY_DATA
 	UPROPERTY()
 	TArray<FNiagaraFunctionSignature> RegisteredFunctions;
+#endif
 
-	// A generated name for this data interface for compilation purposes.  This name is guaranteed to be unique within the compiled script.
 	UPROPERTY()
-	FName Name;
+	FName RegisteredParameterMapRead;
 
 	UPROPERTY()
-	FNiagaraGraphInputDataInterfaceKey SourceGraphInputDataInterfaceKey;
+	FName RegisteredParameterMapWrite;
 
 	UPROPERTY()
-	int32 SourceScriptDIIndex;
+	bool bIsPlaceholder;
 
 	UPROPERTY()
-	FNiagaraCompileParameterReference SourceExternalParameter;
+	FString SourceEmitterName;
 
-	// The parameter that this data interface was written to if any.
-	UPROPERTY()
-	FNiagaraCompileParameterReference TargetParameter;
+	/** Would this data interface work on the target execution type? Only call this on the game thread.*/
+	NIAGARA_API bool CanExecuteOnTarget(ENiagaraSimTarget SimTarget) const;
 
 	/** Note that this is the CDO for this type of data interface, as we often cannot guarantee that the same instance of the data interface we compiled with is the one the user ultimately executes.  Only call this on the game thread.*/
 	NIAGARA_API UNiagaraDataInterface* GetDefaultDataInterface() const;
 
+	NIAGARA_API bool NeedsPerInstanceBinding() const;
+
 	NIAGARA_API bool MatchesClass(const UClass* InClass) const;
 };
 
@@ -1141,6 +1108,38 @@ private:
 	FCustomWorkDelegate PostWork;
 };
 
+/** Defines different usages for a niagara script. */
+UENUM()
+enum class ENiagaraScriptUsage : uint8
+{
+	/** The script defines a function for use in modules. */
+	Function,
+	/** The script defines a module for use in particle, emitter, or system scripts. */
+	Module,
+	/** The script defines a dynamic input for use in particle, emitter, or system scripts. */
+	DynamicInput,
+	/** The script is called when spawning particles. */
+	ParticleSpawnScript,
+	/** Particle spawn script that handles intra-frame spawning and also pulls in the update script. */
+	ParticleSpawnScriptInterpolated UMETA(Hidden),
+	/** The script is called to update particles every frame. */
+	ParticleUpdateScript,
+	/** The script is called to update particles in response to an event. */
+	ParticleEventScript ,
+	/** The script is called as a particle simulation stage. */
+	ParticleSimulationStageScript,
+	/** The script is called to update particles on the GPU. */
+	ParticleGPUComputeScript UMETA(Hidden),
+	/** The script is called once when the emitter spawns. */
+	EmitterSpawnScript,
+	/** The script is called every frame to tick the emitter. */
+	EmitterUpdateScript ,
+	/** The script is called once when the system spawns. */
+	SystemSpawnScript ,
+	/** The script is called every frame to tick the system. */
+	SystemUpdateScript,
+};
+
 /** Defines common bit masks for script usage */
 namespace ENiagaraScriptUsageMask
 {
diff --git a/Engine/Plugins/FX/Niagara/Source/Niagara/Public/NiagaraCompilationTypes.h b/Engine/Plugins/FX/Niagara/Source/Niagara/Public/NiagaraCompilationTypes.h
--- a/Engine/Plugins/FX/Niagara/Source/Niagara/Public/NiagaraCompilationTypes.h
+++ b/Engine/Plugins/FX/Niagara/Source/Niagara/Public/NiagaraCompilationTypes.h
@@ -58,6 +58,9 @@ struct FNiagaraScriptAsyncCompileData
 
 	UPROPERTY()
 	TArray<FNiagaraVariable> RapidIterationParameters;
+
+	UPROPERTY()
+	TMap<FName, TObjectPtr<UNiagaraDataInterface>> NamedDataInterfaces;
 };
 
 USTRUCT()
diff --git a/Engine/Plugins/FX/Niagara/Source/Niagara/Public/NiagaraDataInterfaceEmitterBinding.h b/Engine/Plugins/FX/Niagara/Source/Niagara/Public/NiagaraDataInterfaceEmitterBinding.h
--- a/Engine/Plugins/FX/Niagara/Source/Niagara/Public/NiagaraDataInterfaceEmitterBinding.h
+++ b/Engine/Plugins/FX/Niagara/Source/Niagara/Public/NiagaraDataInterfaceEmitterBinding.h
@@ -49,7 +49,4 @@ struct FNiagaraDataInterfaceEmitterBinding
 	{
 		return BindingMode == Other.BindingMode && EmitterName == Other.EmitterName;
 	}
-
-private:
-	const FNiagaraEmitterHandle* ResolveHandleOtherBinding(const UNiagaraSystem* OwnerSystem) const;
 };
diff --git a/Engine/Plugins/FX/Niagara/Source/Niagara/Public/NiagaraScriptCompileParameterReference.h b/Engine/Plugins/FX/Niagara/Source/Niagara/Public/NiagaraScriptCompileParameterReference.h
deleted file mode 100644
diff --git a/Engine/Plugins/FX/Niagara/Source/Niagara/Public/NiagaraScriptDataInterfaceCompileId.h b/Engine/Plugins/FX/Niagara/Source/Niagara/Public/NiagaraScriptDataInterfaceCompileId.h
deleted file mode 100644
diff --git a/Engine/Plugins/FX/Niagara/Source/Niagara/Public/NiagaraScriptRuntimeCompiledData.h b/Engine/Plugins/FX/Niagara/Source/Niagara/Public/NiagaraScriptRuntimeCompiledData.h
--- a/Engine/Plugins/FX/Niagara/Source/Niagara/Public/NiagaraScriptRuntimeCompiledData.h
+++ b/Engine/Plugins/FX/Niagara/Source/Niagara/Public/NiagaraScriptRuntimeCompiledData.h
@@ -22,9 +22,6 @@ struct FNiagaraScriptRuntimeCompiledData
 	UPROPERTY()
 	TArray<FNiagaraResolvedUserDataInterfaceBinding> ResolvedUserDataInterfaceBindings;
 
-	UPROPERTY()
-	TArray<FNiagaraResolvedNpcDataInterfaceBinding> ResolvedNpcDataInterfaceBindings;
-
 	UPROPERTY()
 	TArray<FNiagaraResolvedUObjectInfo> ResolvedUObjects;
 
diff --git a/Engine/Plugins/FX/Niagara/Source/Niagara/Public/NiagaraScriptRuntimeData.h b/Engine/Plugins/FX/Niagara/Source/Niagara/Public/NiagaraScriptRuntimeData.h
--- a/Engine/Plugins/FX/Niagara/Source/Niagara/Public/NiagaraScriptRuntimeData.h
+++ b/Engine/Plugins/FX/Niagara/Source/Niagara/Public/NiagaraScriptRuntimeData.h
@@ -22,8 +22,6 @@ public:
 
 	TConstArrayView<FNiagaraResolvedUserDataInterfaceBinding> GetResolvedUserDataInterfaceBindings() const { return MakeConstArrayView(CompiledDataHandle->Instance->ResolvedUserDataInterfaceBindings); }
 
-	TConstArrayView<FNiagaraResolvedNpcDataInterfaceBinding> GetResolvedNpcDataInterfaceBindings() const { return MakeConstArrayView(CompiledDataHandle->Instance->ResolvedNpcDataInterfaceBindings); }
-
 	TConstArrayView<FNiagaraResolvedUObjectInfo> GetResolvedUObjects() const { return MakeConstArrayView(CompiledDataHandle->Instance->ResolvedUObjects); }
 
 	const FNiagaraScriptExecutionParameterStore* GetExecutionReadyParameterStore() const { return TransientData->GetExecutionReadyParameterStore(); } 
@@ -33,4 +31,14 @@ public:
 private:
 	TSharedPtr<const FNiagaraScriptRuntimeCompiledDataHandle> CompiledDataHandle;
 	TSharedPtr<const FNiagaraScriptRuntimeTransientData> TransientData;
+};
+
+class FNiagaraScriptComputeRuntimeData : public FNiagaraScriptRuntimeData
+{
+
+};
+
+class FNiagaraScriptEventRuntimeData : public FNiagaraScriptRuntimeData
+{
+
 };
\ No newline at end of file
diff --git a/Engine/Plugins/FX/Niagara/Source/Niagara/Public/NiagaraScriptUsage.h b/Engine/Plugins/FX/Niagara/Source/Niagara/Public/NiagaraScriptUsage.h
deleted file mode 100644
diff --git a/Engine/Plugins/FX/Niagara/Source/Niagara/Public/NiagaraSystemInstance.h b/Engine/Plugins/FX/Niagara/Source/Niagara/Public/NiagaraSystemInstance.h
--- a/Engine/Plugins/FX/Niagara/Source/Niagara/Public/NiagaraSystemInstance.h
+++ b/Engine/Plugins/FX/Niagara/Source/Niagara/Public/NiagaraSystemInstance.h
@@ -386,11 +386,7 @@ public:
 
 	NIAGARA_API bool IsReadyToRun() const;
 
-	UNiagaraParameterCollectionInstance* GetParameterCollectionInstance(UNiagaraParameterCollection* Collection);
-
-	void BindScriptParameterCollections(UNiagaraScript* Script, const FNiagaraScriptRuntimeData* ScriptRuntimeData, FNiagaraParameterStore& DestStore);
-
-	void UnbindScriptParameterCollections(UNiagaraScript* Script, FNiagaraParameterStore& DestStore);
+	NIAGARA_API UNiagaraParameterCollectionInstance* GetParameterCollectionInstance(UNiagaraParameterCollection* Collection);
 
 	/** 
 	Manually advances this system's simulation by the specified number of ticks and tick delta. 
diff --git a/Engine/Plugins/FX/Niagara/Source/Niagara/Public/NiagaraSystemSimulation.h b/Engine/Plugins/FX/Niagara/Source/Niagara/Public/NiagaraSystemSimulation.h
--- a/Engine/Plugins/FX/Niagara/Source/Niagara/Public/NiagaraSystemSimulation.h
+++ b/Engine/Plugins/FX/Niagara/Source/Niagara/Public/NiagaraSystemSimulation.h
@@ -325,10 +325,6 @@ public:
 
 	const FNiagaraTickInfo& GetTickInfo() { return TickInfo; }
 
-	void BindScriptParameterCollections(UNiagaraScript* Script, const FNiagaraScriptRuntimeData* ScriptRuntimeData, FNiagaraParameterStore& DestStore);
-
-	void UnbindScriptParameterCollections(UNiagaraScript* Script, FNiagaraParameterStore& DestStore);
-
 protected:
 	void DumpStalledInfo();
 
@@ -356,12 +352,6 @@ protected:
 
 	TArray<FNiagaraSystemInstance*>& GetSystemInstances(ENiagaraSystemInstanceState State) { check(State != ENiagaraSystemInstanceState::None); return SystemInstancesPerState[int32(State)]; }
 
-	void ResolveNpcDataInterfaceBindingsForScript(
-		const UNiagaraScript* Script,
-		const FNiagaraScriptRuntimeData* ScriptRuntimeData,
-		UNiagaraParameterCollectionInstance* CollectionInstance,
-		FNiagaraParameterStore& DestStore) const;
-
 	/** System of instances being simulated.  We use a weak object ptr here because once the last referencing object goes away this system may be come invalid at runtime. */
 	UNiagaraSystem* System;
 
diff --git a/Engine/Plugins/FX/Niagara/Source/NiagaraEditor/Private/NiagaraCompilationDigestBridgeImpl.cpp b/Engine/Plugins/FX/Niagara/Source/NiagaraEditor/Private/NiagaraCompilationDigestBridgeImpl.cpp
--- a/Engine/Plugins/FX/Niagara/Source/NiagaraEditor/Private/NiagaraCompilationDigestBridgeImpl.cpp
+++ b/Engine/Plugins/FX/Niagara/Source/NiagaraEditor/Private/NiagaraCompilationDigestBridgeImpl.cpp
@@ -156,11 +156,6 @@ ENiagaraScriptUsage FNiagaraCompilationDigestBridge::GetEmitterUsage(const FEmit
 	return EmitterNode->Usage;
 }
 
-FGuid FNiagaraCompilationDigestBridge::GetEmitterHandleId(const FEmitterNode* EmitterNode)
-{
-	return EmitterNode->EmitterHandleID;
-}
-
 FString FNiagaraCompilationDigestBridge::GetEmitterHandleIdString(const FEmitterNode* EmitterNode)
 {
 	return EmitterNode->EmitterHandleIdString;
@@ -298,11 +293,6 @@ FNiagaraVariable FNiagaraCompilationDigestBridge::GetInputVariable(const FInputN
 	return InputNode->InputVariable;
 }
 
-FNiagaraGraphInputDataInterfaceKey FNiagaraCompilationDigestBridge::GetGraphInputDataInterfaceKey(const FNiagaraCompilationNodeInput* InputNode, const FCompilationCopy* CompilationCopy)
-{
-	return FNiagaraGraphInputDataInterfaceKey(InputNode->GraphPathFName, InputNode->InputVariable);
-}
-
 const TArray<FNiagaraVariable>& FNiagaraCompilationDigestBridge::GetOutputVariables(const FOutputNode* OutputNode)
 {
 	return OutputNode->Outputs;
diff --git a/Engine/Plugins/FX/Niagara/Source/NiagaraEditor/Private/NiagaraCompilationGraphBridgeImpl.cpp b/Engine/Plugins/FX/Niagara/Source/NiagaraEditor/Private/NiagaraCompilationGraphBridgeImpl.cpp
--- a/Engine/Plugins/FX/Niagara/Source/NiagaraEditor/Private/NiagaraCompilationGraphBridgeImpl.cpp
+++ b/Engine/Plugins/FX/Niagara/Source/NiagaraEditor/Private/NiagaraCompilationGraphBridgeImpl.cpp
@@ -276,11 +276,6 @@ FString FNiagaraCompilationGraphBridge::GetEmitterPathName(const FEmitterNode* E
 	return EmitterNode->GetFullName();
 }
 
-FGuid FNiagaraCompilationGraphBridge::GetEmitterHandleId(const FEmitterNode* EmitterNode)
-{
-	return EmitterNode->GetEmitterHandleId();
-}
-
 FString FNiagaraCompilationGraphBridge::GetEmitterHandleIdString(const FEmitterNode* EmitterNode)
 {
 	return EmitterNode->GetEmitterHandleId().ToString(EGuidFormats::Digits);
@@ -388,12 +383,6 @@ FNiagaraVariable FNiagaraCompilationGraphBridge::GetInputVariable(const FInputNo
 	return InputNode->Input;
 }
 
-FNiagaraGraphInputDataInterfaceKey FNiagaraCompilationGraphBridge::GetGraphInputDataInterfaceKey(const UNiagaraNodeInput* InputNode, const FCompilationCopy* CompilationCopy)
-{
-	const UNiagaraDataInterface* InputDataInterface = InputNode->GetDataInterface();
-	return InputDataInterface != nullptr ? CompilationCopy->GetGraphInputDataInterfaceKey(*InputDataInterface) : FNiagaraGraphInputDataInterfaceKey();
-}
-
 const TArray<FNiagaraVariable>& FNiagaraCompilationGraphBridge::GetOutputVariables(const FOutputNode* OutputNode)
 {
 	return OutputNode->GetOutputs();
diff --git a/Engine/Plugins/FX/Niagara/Source/NiagaraEditor/Private/NiagaraCompilationPrivate.h b/Engine/Plugins/FX/Niagara/Source/NiagaraEditor/Private/NiagaraCompilationPrivate.h
--- a/Engine/Plugins/FX/Niagara/Source/NiagaraEditor/Private/NiagaraCompilationPrivate.h
+++ b/Engine/Plugins/FX/Niagara/Source/NiagaraEditor/Private/NiagaraCompilationPrivate.h
@@ -291,10 +291,9 @@ public:
 
 	virtual bool IsDuplicateDataFor(UNiagaraSystem* InSystem, UNiagaraEmitter* InEmitter, UNiagaraScript* InScript) const override;
 	virtual void GetDuplicatedObjects(TArray<UObject*>& Objects) override;
-	virtual const TMap<FNiagaraGraphInputDataInterfaceKey, UNiagaraDataInterface*>& GetGraphInputDataInterfaces() override;
+	virtual const TMap<FName, UNiagaraDataInterface*>& GetObjectNameMap() override;
 	virtual const UNiagaraScriptSourceBase* GetScriptSource() const override;
 	UNiagaraDataInterface* GetDuplicatedDataInterfaceCDOForClass(UClass* Class) const;
-	FNiagaraGraphInputDataInterfaceKey GetGraphInputDataInterfaceKey(const UNiagaraDataInterface& DataInterface) const;
 
 	TArray<FNiagaraParameterMapHistory>& GetPrecomputedHistories() { return PrecompiledHistories; }
 	const TArray<FNiagaraParameterMapHistory>& GetPrecomputedHistories() const { return PrecompiledHistories; }
@@ -336,11 +335,10 @@ public:
 	};
 
 	TSharedPtr<TMap<const UNiagaraGraph*, TArray<FDuplicatedGraphData>>> SharedSourceGraphToDuplicatedGraphsMap;
-	TSharedPtr<TMap<FNiagaraGraphInputDataInterfaceKey, UNiagaraDataInterface*>> SharedDuplicatedGraphInputDataInterfaces;
+	TSharedPtr<TMap<FName, UNiagaraDataInterface*>> SharedNameToDuplicatedDataInterfaceMap;
 	TSharedPtr<TMap<UClass*, UNiagaraDataInterface*>> SharedDataInterfaceClassToDuplicatedCDOMap;
 
 protected:
-	void CollectDuplicatedDataInterfaces(UNiagaraGraph* DuplicateGraph, FName SourceGraphPath);
 	void DuplicateReferencedGraphsRecursive(UNiagaraGraph* InGraph, const FCompileConstantResolver& ConstantResolver, TMap<UNiagaraNodeFunctionCall*, ENiagaraScriptUsage> FunctionsWithUsage);
 private:
 	TArray<TWeakObjectPtr<UNiagaraScriptSource>> TrackedScriptSourceCopies;
diff --git a/Engine/Plugins/FX/Niagara/Source/NiagaraEditor/Private/NiagaraCompilationTasks.cpp b/Engine/Plugins/FX/Niagara/Source/NiagaraEditor/Private/NiagaraCompilationTasks.cpp
--- a/Engine/Plugins/FX/Niagara/Source/NiagaraEditor/Private/NiagaraCompilationTasks.cpp
+++ b/Engine/Plugins/FX/Niagara/Source/NiagaraEditor/Private/NiagaraCompilationTasks.cpp
@@ -33,14 +33,6 @@
 #include "NiagaraNodeOutput.h"
 #include "NiagaraNodeFunctionCall.h"
 
-static int GNiagaraCompilationForceOneThread = 0;
-static FAutoConsoleVariableRef CVarNiagaraCompilationForceOneThread(
-	TEXT("fx.Niagara.Compilation.ForceOneThread"),
-	GNiagaraCompilationForceOneThread,
-	TEXT("Forces the compile tasks to run inline on the main thread."),
-	ECVF_Default
-);
-
 #if ENABLE_COOK_STATS
 namespace NiagaraSystemCookStats
 {
@@ -836,6 +828,84 @@ const FNiagaraSystemCompilationTask::FScriptInfo* FNiagaraSystemCompilationTask:
 	return DigestedScriptInfo.Find(ScriptKey);
 }
 
+void FNiagaraSystemCompilationTask::FCompileTaskInfo::CollectNamedDataInterfaces(FNiagaraSystemCompilationTask* SystemCompileTask, const FCompileGroupInfo& GroupInfo)
+{
+	FString UniqueEmitterName;
+	const FNiagaraCompilationCopyData* CompilationCopyData = nullptr;
+
+	if (const FEmitterInfo* EmitterInfo = SystemCompileTask->SystemInfo.EmitterInfoBySourceEmitter(GroupInfo.SourceEmitterIndex))
+	{
+		CompilationCopyData = static_cast<const FNiagaraCompilationCopyData*>(GroupInfo.CompilationCopy->GetDependentRequest(EmitterInfo->DigestedEmitterIndex).Get());
+		UniqueEmitterName = EmitterInfo->UniqueEmitterName;
+	}
+	else
+	{
+		CompilationCopyData = GroupInfo.CompilationCopy.Get();
+	}
+
+	const FScriptInfo& ScriptInfo = SystemCompileTask->DigestedScriptInfo.FindChecked(ScriptKey);
+
+	auto AccumulateInputDataInterfaces = [this, &UniqueEmitterName](const FNiagaraCompilationNode& Node) -> bool
+	{
+		if (const FNiagaraCompilationNodeInput* InputNode = Node.AsType<FNiagaraCompilationNodeInput>())
+		{
+			if (InputNode->DataInterfaceName != NAME_None && InputNode->DuplicatedDataInterface)
+			{
+				constexpr bool bIsParameterMapDataInterface = false;
+				FName DIName = FNiagaraHlslTranslator::GetDataInterfaceName(InputNode->DataInterfaceName, UniqueEmitterName, bIsParameterMapDataInterface);
+
+				NamedDataInterfaces.Add(DIName, InputNode->DuplicatedDataInterface);
+			}
+		}
+
+		return true;
+	};
+
+	auto DefaultNodeFinder = [&ScriptInfo](const FNiagaraCompilationNodeOutput& OutputNode) -> bool
+	{
+		if (UNiagaraScript::IsEquivalentUsage(OutputNode.Usage, ScriptInfo.Usage) && OutputNode.UsageId == ScriptInfo.UsageId)
+		{
+			return true;
+		}
+
+		if (ScriptInfo.Usage == ENiagaraScriptUsage::ParticleGPUComputeScript)
+		{
+			return UNiagaraScript::IsParticleScript(OutputNode.Usage);
+		}
+		else if (ScriptInfo.Usage == ENiagaraScriptUsage::ParticleSpawnScriptInterpolated)
+		{
+			return UNiagaraScript::IsEquivalentUsage(OutputNode.Usage, ENiagaraScriptUsage::ParticleUpdateScript);
+		}
+
+		return OutputNode.Usage == ScriptInfo.Usage && OutputNode.UsageId == ScriptInfo.UsageId;
+	};
+
+	CompilationCopyData->InstantiatedGraph->NodeTraversal(true, false, DefaultNodeFinder, AccumulateInputDataInterfaces);
+
+	if (UNiagaraScript::IsSystemScript(ScriptInfo.Usage))
+	{
+		check(GroupInfo.SourceEmitterIndex == INDEX_NONE);
+
+		const ENiagaraScriptUsage EmitterUsage = ScriptInfo.Usage == ENiagaraScriptUsage::SystemSpawnScript
+			? ENiagaraScriptUsage::EmitterSpawnScript
+			: ENiagaraScriptUsage::EmitterUpdateScript;
+		const FGuid EmitterUsageId = ScriptInfo.UsageId;
+		auto EmitterNodeFinder = [EmitterUsage, EmitterUsageId](const FNiagaraCompilationNodeOutput& OutputNode) -> bool
+		{
+			return UNiagaraScript::IsEquivalentUsage(OutputNode.Usage, EmitterUsage) && OutputNode.UsageId == EmitterUsageId;
+		};
+
+		for (const FNiagaraCompilationCopyData::FSharedCompilationCopy& EmitterCompilationCopy : CompilationCopyData->EmitterData)
+		{
+			if (const FNiagaraCompilationGraph* EmitterGraph = EmitterCompilationCopy->InstantiatedGraph.Get())
+			{
+				UniqueEmitterName = EmitterCompilationCopy->EmitterUniqueName;
+				EmitterGraph->NodeTraversal(true, false, EmitterNodeFinder, AccumulateInputDataInterfaces);
+			}
+		}
+	}
+}
+
 void FNiagaraSystemCompilationTask::FCompileTaskInfo::Translate(FNiagaraSystemCompilationTask* SystemCompileTask, const FCompileGroupInfo& GroupInfo)
 {
 	TRACE_CPUPROFILER_EVENT_SCOPE(NiagaraAsyncTask_Translate);
@@ -1811,6 +1881,12 @@ bool FNiagaraSystemCompilationTask::Poll(FNiagaraSystemAsyncCompileResults& Resu
 						ResultsInfo.CompilationErrors = ShaderTaskInfo.CompilationErrors;
 					}
 
+					CompileData.NamedDataInterfaces.Reserve(TaskInfo.NamedDataInterfaces.Num());
+					Algo::Transform(TaskInfo.NamedDataInterfaces, CompileData.NamedDataInterfaces, [](const TTuple<FName, UNiagaraDataInterface*>& InElement)
+					{
+						return MakeTuple(InElement.Key, InElement.Value);
+					});
+
 					CompileData.CompileMetrics.TaskWallTime = (float) (FPlatformTime::Seconds() - TaskInfo.TaskStartTime);
 					CompileData.CompileMetrics.DDCFetchTime = 0.0f;
 					CompileData.CompileMetrics.CompilerWallTime = TaskInfo.CompilerWallTime;
@@ -2570,11 +2646,6 @@ void FNiagaraSystemCompilationTask::IssueCompilationTasks()
 	TArray<FTask> IssuedCompilationTasks;
 	TArray<FTaskEvent> CompilationTasks;
 
-	ETaskPriority TaskPriority = ETaskPriority::Normal;
-	EExtendedTaskPriority TaskExtendedPriority = GNiagaraCompilationForceOneThread > 0 
-		? EExtendedTaskPriority::Inline 
-		: EExtendedTaskPriority::None;
-
 	if (HasOutstandingCompileTasks())
 	{
 		FTask PrecompileTask = Launch(UE_SOURCE_LOCATION, [Self = AsShared()]
@@ -2582,7 +2653,7 @@ void FNiagaraSystemCompilationTask::IssueCompilationTasks()
 			COOK_STAT(auto Timer = NiagaraSystemCookStats::UsageStats.TimeSyncWork(); Timer.TrackCyclesOnly(););
 
 			Self->Precompile();
-		}, TaskPriority, TaskExtendedPriority);
+		});
 
 		// async compilation copy tasks
 		for (FCompileGroupInfo& GroupInfo : CompileGroups)
@@ -2594,7 +2665,7 @@ void FNiagaraSystemCompilationTask::IssueCompilationTasks()
 					COOK_STAT(auto Timer = NiagaraSystemCookStats::UsageStats.TimeSyncWork(); Timer.TrackCyclesOnly(););
 
 					GroupInfo.InstantiateCompileGraph(*Self);
-				}, PrecompileTask, TaskPriority, TaskExtendedPriority);
+				}, PrecompileTask);
 
 				for (int32 CompileTaskIndex : GroupInfo.CompileTaskIndices)
 				{
@@ -2610,6 +2681,7 @@ void FNiagaraSystemCompilationTask::IssueCompilationTasks()
 
 							FNiagaraSystemCompilationTask* TaskPtr = &Self.Get();
 
+							CompileTask.CollectNamedDataInterfaces(TaskPtr, GroupInfo);
 							switch (CompileTask.ScriptCompileType)
 							{
 								case EScriptCompileType::CompileForVm:
@@ -2637,14 +2709,14 @@ void FNiagaraSystemCompilationTask::IssueCompilationTasks()
 									CompileTask.Translate(TaskPtr, GroupInfo);
 								break;
 							}
-						}, InstantiateGraphTask, TaskPriority, TaskExtendedPriority));
+						}, InstantiateGraphTask));
 
 						Launch(UE_SOURCE_LOCATION, [Self = AsShared(), &GroupInfo, &CompileTask]
 						{
 							COOK_STAT(auto Timer = NiagaraSystemCookStats::UsageStats.TimeSyncWork(); Timer.TrackCyclesOnly(););
 
 							CompileTask.ProcessCompilation(&Self.Get(), GroupInfo);
-						}, *CompileTask.CompileResultsReadyEvent, TaskPriority, TaskExtendedPriority);
+						}, *CompileTask.CompileResultsReadyEvent);
 
 						CompilationTasks.Add(*CompileTask.CompileResultsProcessedEvent);
 					}
@@ -2656,12 +2728,12 @@ void FNiagaraSystemCompilationTask::IssueCompilationTasks()
 	Launch(UE_SOURCE_LOCATION, [Self = AsShared()]
 	{
 		Self->CompilationState = EState::WaitingForProcessing;
-	}, IssuedCompilationTasks, TaskPriority, TaskExtendedPriority);
+	}, IssuedCompilationTasks);
 
 	Launch(UE_SOURCE_LOCATION, [Self = AsShared()]
 	{
 		Self->IssuePostResultsProcessedTasks();
-	}, CompilationTasks, TaskPriority, TaskExtendedPriority);
+	}, CompilationTasks);
 }
 
 bool FNiagaraSystemCompilationTask::HasOutstandingCompileTasks() const
diff --git a/Engine/Plugins/FX/Niagara/Source/NiagaraEditor/Private/NiagaraCompilationTasks.h b/Engine/Plugins/FX/Niagara/Source/NiagaraEditor/Private/NiagaraCompilationTasks.h
--- a/Engine/Plugins/FX/Niagara/Source/NiagaraEditor/Private/NiagaraCompilationTasks.h
+++ b/Engine/Plugins/FX/Niagara/Source/NiagaraEditor/Private/NiagaraCompilationTasks.h
@@ -204,6 +204,7 @@ struct FNiagaraSystemCompilationTask : public TSharedFromThis<FNiagaraSystemComp
 
 		TUniquePtr<UE::Tasks::FTaskEvent> CompileResultsReadyEvent;
 		TUniquePtr<UE::Tasks::FTaskEvent> CompileResultsProcessedEvent;
+		TMap<FName, UNiagaraDataInterface*> NamedDataInterfaces;
 		TArray<FNiagaraVariable> BakedRapidIterationParameters; // todo - remove and have the data pulled directly from the precompile data
 		double TaskStartTime = 0;
 		float TranslationTime = 0.0f;
@@ -214,6 +215,7 @@ struct FNiagaraSystemCompilationTask : public TSharedFromThis<FNiagaraSystemComp
 		EScriptCompileType ScriptCompileType = EScriptCompileType::Invalid;
 
 		bool IsOutstanding(const FNiagaraSystemCompilationTask& ParentTask) const;
+		void CollectNamedDataInterfaces(FNiagaraSystemCompilationTask* SystemCompileTask, const FCompileGroupInfo& GroupInfo);
 		void Translate(FNiagaraSystemCompilationTask* SystemCompileTask, const FCompileGroupInfo& GroupInfo);
 		void IssueCompileVm(FNiagaraSystemCompilationTask* SystemCompileTask, const FCompileGroupInfo& GroupInfo);
 		void IssueTranslateGpu(FNiagaraSystemCompilationTask* SystemCompileTask, const FCompileGroupInfo& GroupInfo);
diff --git a/Engine/Plugins/FX/Niagara/Source/NiagaraEditor/Private/NiagaraCompiler.cpp b/Engine/Plugins/FX/Niagara/Source/NiagaraEditor/Private/NiagaraCompiler.cpp
--- a/Engine/Plugins/FX/Niagara/Source/NiagaraEditor/Private/NiagaraCompiler.cpp
+++ b/Engine/Plugins/FX/Niagara/Source/NiagaraEditor/Private/NiagaraCompiler.cpp
@@ -154,29 +154,6 @@ void FNiagaraCompileRequestDuplicateData::DuplicateReferencedGraphs(UNiagaraGrap
 	DuplicateReferencedGraphsRecursive(InDupeGraph, ConstantResolver, FunctionsWithUsage);
 }
 
-void FNiagaraCompileRequestDuplicateData::CollectDuplicatedDataInterfaces(UNiagaraGraph* DuplicateGraph, FName SourceGraphPath)
-{
-	TArray<UNiagaraNodeInput*> InputNodes;
-	DuplicateGraph->GetNodesOfClass(InputNodes);
-	const UEdGraphSchema_Niagara* Schema = GetDefault<UEdGraphSchema_Niagara>();
-
-	for (UNiagaraNodeInput* InputNode : InputNodes)
-	{
-		if (InputNode)
-		{
-			if (InputNode->Input.IsDataInterface())
-			{
-				UNiagaraDataInterface* DataInterface = InputNode->GetDataInterface();
-				if (DataInterface != nullptr)
-				{
-					FNiagaraGraphInputDataInterfaceKey GraphInputDataInterfaceKey(SourceGraphPath, InputNode->Input);
-					SharedDuplicatedGraphInputDataInterfaces->Add(GraphInputDataInterfaceKey, DataInterface);
-				}
-			}
-		}
-	}
-}
-
 void FNiagaraCompileRequestDuplicateData::DuplicateReferencedGraphsRecursive(UNiagaraGraph* InGraph, const FCompileConstantResolver& ConstantResolver, TMap<UNiagaraNodeFunctionCall*, ENiagaraScriptUsage> FunctionsWithUsage)
 {
 	if (!InGraph)
@@ -192,6 +169,20 @@ void FNiagaraCompileRequestDuplicateData::DuplicateReferencedGraphsRecursive(UNi
 	{
 		if (UNiagaraNode* InNode = Cast<UNiagaraNode>(Node))
 		{
+			UNiagaraNodeInput* InputNode = Cast<UNiagaraNodeInput>(InNode);
+			if (InputNode)
+			{
+				if (InputNode->Input.IsDataInterface())
+				{
+					UNiagaraDataInterface* DataInterface = InputNode->GetDataInterface();
+					bool bIsParameterMapDataInterface = false;
+					FName DIName = FNiagaraHlslTranslator::GetDataInterfaceName(InputNode->Input.GetName(), EmitterUniqueName, bIsParameterMapDataInterface);
+					UNiagaraDataInterface* Dupe = PrecompileDuplicateObject<UNiagaraDataInterface>(DataInterface, GetTransientPackage());
+					SharedNameToDuplicatedDataInterfaceMap->Add(DIName, Dupe);
+				}
+				continue;
+			}
+
 			UNiagaraNodeFunctionCall* FunctionCallNode = Cast<UNiagaraNodeFunctionCall>(InNode);
 			if (FunctionCallNode)
 			{
@@ -241,7 +232,6 @@ void FNiagaraCompileRequestDuplicateData::DuplicateReferencedGraphsRecursive(UNi
 
 						// Assign the copied script and process any child scripts.
 						FunctionCallNode->FunctionScript = DupeScript;
-						CollectDuplicatedDataInterfaces(DupeGraph, *FunctionGraph->GetPathName());
 						DuplicateReferencedGraphsRecursive(DupeGraph, FunctionConstantResolver, FunctionsWithUsage);
 					}
 					else
@@ -271,9 +261,9 @@ void FNiagaraCompileRequestDuplicateData::DuplicateReferencedGraphsRecursive(UNi
 	}
 }
 
-const TMap<FNiagaraGraphInputDataInterfaceKey, UNiagaraDataInterface*>& FNiagaraCompileRequestDuplicateData::GetGraphInputDataInterfaces()
+const TMap<FName, UNiagaraDataInterface*>& FNiagaraCompileRequestDuplicateData::GetObjectNameMap()
 {
-	return *SharedDuplicatedGraphInputDataInterfaces.Get();
+	return *SharedNameToDuplicatedDataInterfaceMap.Get();
 }
 
 const UNiagaraScriptSourceBase* FNiagaraCompileRequestDuplicateData::GetScriptSource() const
@@ -294,18 +284,6 @@ UNiagaraDataInterface* FNiagaraCompileRequestDuplicateData::GetDuplicatedDataInt
 	return nullptr;
 }
 
-FNiagaraGraphInputDataInterfaceKey FNiagaraCompileRequestDuplicateData::GetGraphInputDataInterfaceKey(const UNiagaraDataInterface& DataInterface) const
-{
-	for (const TPair<FNiagaraGraphInputDataInterfaceKey, UNiagaraDataInterface*>& GraphInputDataInterfacePair : *SharedDuplicatedGraphInputDataInterfaces)
-	{
-		if (GraphInputDataInterfacePair.Value == &DataInterface)
-		{
-			return GraphInputDataInterfacePair.Key;
-		}
-	}
-	return FNiagaraGraphInputDataInterfaceKey();
-}
-
 void FNiagaraCompileRequestData::SortOutputNodesByDependencies(TArray<class UNiagaraNodeOutput*>& NodesToSort, const TArray<class UNiagaraSimulationStageBase*>* SimStages)
 {
 	if (!SimStages)
@@ -377,10 +355,10 @@ void FNiagaraCompileRequestDuplicateData::GetDuplicatedObjects(TArray<UObject*>&
 	Objects.Add(SourceDeepCopy.Get());
 	Objects.Add(NodeGraphDeepCopy.Get());
 
-	if (SharedDuplicatedGraphInputDataInterfaces.IsValid())
+	if (SharedNameToDuplicatedDataInterfaceMap.IsValid())
 	{
 		TArray<UNiagaraDataInterface*> DIs;
-		SharedDuplicatedGraphInputDataInterfaces->GenerateValueArray(DIs);
+		SharedNameToDuplicatedDataInterfaceMap->GenerateValueArray(DIs);
 		for (UNiagaraDataInterface* DI : DIs)
 		{
 			Objects.Add(DI);
@@ -441,7 +419,6 @@ void FNiagaraCompileRequestDuplicateData::DeepCopyGraphs(UNiagaraScriptSource* S
 	SourceDeepCopy = ScriptSource->CreateCompilationCopy(ValidUsages);
 	NodeGraphDeepCopy = SourceDeepCopy->NodeGraph;
 	FEdGraphUtilities::MergeChildrenGraphsIn(NodeGraphDeepCopy.Get(), NodeGraphDeepCopy.Get(), /*bRequireSchemaMatch=*/ true);
-	CollectDuplicatedDataInterfaces(NodeGraphDeepCopy.Get(), *ScriptSource->NodeGraph.GetPathName());
 	DuplicateReferencedGraphs(ScriptSource->NodeGraph, NodeGraphDeepCopy.Get(), InUsage, ConstantResolver);
 }
 
@@ -471,7 +448,6 @@ void FNiagaraCompileRequestDuplicateData::DeepCopyGraphs(const FVersionedNiagara
 	}
 
 	FCompileConstantResolver ConstantResolver(Emitter, ENiagaraScriptUsage::EmitterSpawnScript);
-	CollectDuplicatedDataInterfaces(NodeGraphDeepCopy.Get(), *ScriptSource->NodeGraph.GetPathName());
 	DuplicateReferencedGraphs(ScriptSource->NodeGraph, NodeGraphDeepCopy.Get(), ENiagaraScriptUsage::EmitterSpawnScript, ConstantResolver, FunctionsWithUsage);
 }
 
@@ -1334,7 +1310,7 @@ TSharedPtr<FNiagaraCompileRequestDuplicateDataBase, ESPMode::ThreadSafe> FNiagar
 	FCompileConstantResolver EmptyResolver;
 
 	BasePtr->SharedSourceGraphToDuplicatedGraphsMap = MakeShared<TMap<const UNiagaraGraph*, TArray<FNiagaraCompileRequestDuplicateData::FDuplicatedGraphData>>>();
-	BasePtr->SharedDuplicatedGraphInputDataInterfaces = MakeShared<TMap<FNiagaraGraphInputDataInterfaceKey, UNiagaraDataInterface*>>();
+	BasePtr->SharedNameToDuplicatedDataInterfaceMap = MakeShared<TMap<FName, UNiagaraDataInterface*>>();
 	BasePtr->SharedDataInterfaceClassToDuplicatedCDOMap = MakeShared<TMap<UClass*, UNiagaraDataInterface*>>();
 	BasePtr->OwningSystem = OwningSystem;
 	BasePtr->OwningEmitter = OwningEmitter;
@@ -1380,7 +1356,7 @@ TSharedPtr<FNiagaraCompileRequestDuplicateDataBase, ESPMode::ThreadSafe> FNiagar
 			EmitterPtr->EmitterID = FNiagaraEmitterID(i);
 			EmitterPtr->ValidUsages = BasePtr->ValidUsages;
 			EmitterPtr->SharedSourceGraphToDuplicatedGraphsMap = BasePtr->SharedSourceGraphToDuplicatedGraphsMap;
-			EmitterPtr->SharedDuplicatedGraphInputDataInterfaces = BasePtr->SharedDuplicatedGraphInputDataInterfaces;
+			EmitterPtr->SharedNameToDuplicatedDataInterfaceMap = BasePtr->SharedNameToDuplicatedDataInterfaceMap;
 			EmitterPtr->SharedDataInterfaceClassToDuplicatedCDOMap = BasePtr->SharedDataInterfaceClassToDuplicatedCDOMap;
 			//EmitterPtr->bSimulationStagesEnabled = Handle.GetInstance()->bSimulationStagesEnabled;
 			if (Handle.GetIsEnabled() && Handle.GetInstance().Emitter && (OwningEmitter == nullptr || OwningEmitter == Handle.GetInstance().Emitter)) // Don't need to copy the graph if we aren't going to use it.
diff --git a/Engine/Plugins/FX/Niagara/Source/NiagaraEditor/Private/NiagaraEditorUtilities.cpp b/Engine/Plugins/FX/Niagara/Source/NiagaraEditor/Private/NiagaraEditorUtilities.cpp
--- a/Engine/Plugins/FX/Niagara/Source/NiagaraEditor/Private/NiagaraEditorUtilities.cpp
+++ b/Engine/Plugins/FX/Niagara/Source/NiagaraEditor/Private/NiagaraEditorUtilities.cpp
@@ -41,7 +41,6 @@
 #include "ViewModels/NiagaraScriptGraphViewModel.h"
 #include "NiagaraScriptMergeManager.h"
 #include "NiagaraScriptSource.h"
-#include "NiagaraScriptRuntimeData.h"
 #include "NiagaraSettings.h"
 #include "NiagaraSimulationStageBase.h"
 #include "NiagaraStackEditorData.h"
@@ -4915,34 +4914,29 @@ UNiagaraDataInterface* FNiagaraEditorUtilities::GetResolvedRuntimeInstanceForEdi
 		}
 
 		// If the data interface was owned by an input node, then we need to try to update the compiled version.
-		const FNiagaraEmitterHandle* EmitterHandle = nullptr;
+		const FNiagaraEmitterHandle* EmitterHandle;
 		TArray<const UNiagaraScript*> CompiledScripts;
 		
-		FNiagaraGraphInputDataInterfaceKey GraphInputDataInterfaceKey(*OuterInputNode->GetGraph()->GetPathName(), OuterInputNode->Input);
 		FNiagaraStackGraphUtilities::GetEmitterHandleAndCompiledScriptsForStackNode(OwningSystem, *OuterInputNode, EmitterHandle, CompiledScripts);
 		if (ensureMsgf(CompiledScripts.Num() > 0, TEXT("Could not find compiled scripts for data interface input node.")))
 		{
 			for (const UNiagaraScript* CompiledScript : CompiledScripts)
 			{
-				const FNiagaraScriptDataInterfaceCompileInfo* GraphInputDICompileInfo = CompiledScript->GetVMExecutableData().DataInterfaceInfo.FindByPredicate([&GraphInputDataInterfaceKey]
-					(const FNiagaraScriptDataInterfaceCompileInfo& ScriptDICompileInfo) { return ScriptDICompileInfo.SourceGraphInputDataInterfaceKey == GraphInputDataInterfaceKey; });
-				if (GraphInputDICompileInfo == nullptr)
-				{
-					continue;
-				}
-
-				FGuid EmitterHandleId = (EmitterHandle != nullptr && UNiagaraScript::IsSystemScript(CompiledScript->GetUsage()) == false) ? EmitterHandle->GetId() : FGuid();
-				const FNiagaraScriptRuntimeData* ScriptRuntimeData = OwningSystem.GetScriptRuntimeData(EmitterHandleId, CompiledScript->GetUsage(), CompiledScript->GetUsageId()).Get();
-				if (ScriptRuntimeData == nullptr)
-				{
-					continue;
-				}
-
-				const FNiagaraScriptResolvedDataInterfaceInfo* ResolvedGraphInputDIData = ScriptRuntimeData->GetResolvedDataInterfaces().FindByPredicate([GraphInputDICompileInfo]
-					(const FNiagaraScriptResolvedDataInterfaceInfo& ResolvedDataInterfaceData) { return ResolvedDataInterfaceData.CompileName == GraphInputDICompileInfo->Name; });
-				if (ResolvedGraphInputDIData != nullptr)
+				FString EmitterName = EmitterHandle != nullptr ? EmitterHandle->GetUniqueInstanceName() : FString();
+				bool bIsParameterMapDataInterface = false;
+				FName DataInterfaceName = FNiagaraHlslTranslator::GetDataInterfaceName(OuterInputNode->Input.GetName(), EmitterName, bIsParameterMapDataInterface);
+				
+				const FNiagaraEmitterHandle* EmitterHandleForScript = UNiagaraScript::IsSystemScript(CompiledScript->GetUsage()) ? nullptr : EmitterHandle;
+				const FNiagaraScriptRuntimeData* ScriptRuntimeData = OwningSystem.GetScriptRuntimeData(EmitterHandleForScript, *CompiledScript).Get();
+				if (ScriptRuntimeData != nullptr)
 				{
-					return ResolvedGraphInputDIData->ResolvedDataInterface;
+					for (const FNiagaraScriptResolvedDataInterfaceInfo& ResolvedDI : ScriptRuntimeData->GetResolvedDataInterfaces())
+					{
+						if (ResolvedDI.Name == DataInterfaceName)
+						{
+							return ResolvedDI.ResolvedDataInterface;
+						}
+					}
 				}
 			}
 		}
diff --git a/Engine/Plugins/FX/Niagara/Source/NiagaraEditor/Private/NiagaraGraphDigest.cpp b/Engine/Plugins/FX/Niagara/Source/NiagaraEditor/Private/NiagaraGraphDigest.cpp
--- a/Engine/Plugins/FX/Niagara/Source/NiagaraEditor/Private/NiagaraGraphDigest.cpp
+++ b/Engine/Plugins/FX/Niagara/Source/NiagaraEditor/Private/NiagaraGraphDigest.cpp
@@ -3291,7 +3291,6 @@ FNiagaraCompilationNodeInput::FNiagaraCompilationNodeInput(const UNiagaraNodeInp
 	bExposed = InNode->IsExposed();
 	bRequired = InNode->ExposureOptions.bRequired;
 	bCanAutoBind = InNode->ExposureOptions.bCanAutoBind;
-	GraphPathFName = *InNode->GetGraph()->GetPathName();
 }
 
 
@@ -3306,7 +3305,6 @@ FNiagaraCompilationNodeInput::FNiagaraCompilationNodeInput(const FNiagaraCompila
 	, bRequired(InNode.bRequired)
 	, bExposed(InNode.bExposed)
 	, bCanAutoBind(InNode.bCanAutoBind)
-	, GraphPathFName(InNode.GraphPathFName)
 	, DuplicatedDataInterface(InNode.DuplicatedDataInterface)
 	, ObjectAssetPath(InNode.ObjectAssetPath)
 {
@@ -3383,7 +3381,7 @@ void FNiagaraCompilationNodeInput::Compile(FTranslator* Translator, TArray<int32
 					{
 						check(DuplicatedDataInterface);
 						check(InputVariable.IsDataInterface());
-						Outputs.Add(Translator->RegisterDataInterfaceForInput(InputVariable, Translator->GetGraphInputDataInterfaceKey(this), DuplicatedDataInterface.Get()));
+						Outputs.Add(Translator->RegisterDataInterface(InputVariable, DuplicatedDataInterface, false, false));
 						return;
 					}
 					else if (!ObjectAssetName.IsNone())
@@ -3408,7 +3406,7 @@ void FNiagaraCompilationNodeInput::Compile(FTranslator* Translator, TArray<int32
 		{
 			check(DuplicatedDataInterface);
 			check(InputVariable.IsDataInterface());
-			Outputs.Add(Translator->RegisterDataInterfaceForInput(InputVariable, Translator->GetGraphInputDataInterfaceKey(this), DuplicatedDataInterface.Get()));
+			Outputs.Add(Translator->RegisterDataInterface(InputVariable, DuplicatedDataInterface, false, false));
 			break;
 		}
 		else if (!ObjectAssetName.IsNone())
diff --git a/Engine/Plugins/FX/Niagara/Source/NiagaraEditor/Private/NiagaraGraphDigest.h b/Engine/Plugins/FX/Niagara/Source/NiagaraEditor/Private/NiagaraGraphDigest.h
--- a/Engine/Plugins/FX/Niagara/Source/NiagaraEditor/Private/NiagaraGraphDigest.h
+++ b/Engine/Plugins/FX/Niagara/Source/NiagaraEditor/Private/NiagaraGraphDigest.h
@@ -514,7 +514,6 @@ public:
 	bool bRequired;
 	bool bExposed;
 	bool bCanAutoBind;
-	FName GraphPathFName = NAME_None;
 
 	TObjectPtr<UNiagaraDataInterface> DuplicatedDataInterface;
 	FSoftObjectPath ObjectAssetPath;
diff --git a/Engine/Plugins/FX/Niagara/Source/NiagaraEditor/Private/NiagaraHlslTranslator.cpp b/Engine/Plugins/FX/Niagara/Source/NiagaraEditor/Private/NiagaraHlslTranslator.cpp
--- a/Engine/Plugins/FX/Niagara/Source/NiagaraEditor/Private/NiagaraHlslTranslator.cpp
+++ b/Engine/Plugins/FX/Niagara/Source/NiagaraEditor/Private/NiagaraHlslTranslator.cpp
@@ -1613,32 +1613,6 @@ FNiagaraTranslateResults TNiagaraHlslTranslator<GraphBridge>::Translate(const FN
 		return TranslateResults;
 	}
 
-	auto GetDIOwnerNameByUsage = [](ENiagaraScriptUsage Usage)->FName
-	{
-		switch (Usage)
-		{
-			case ENiagaraScriptUsage::SystemSpawnScript:
-				return "SystemSpawn";
-			case ENiagaraScriptUsage::SystemUpdateScript:
-				return "SystemUpdate";
-			case ENiagaraScriptUsage::EmitterSpawnScript:
-				return "EmitterSpawn";
-			case ENiagaraScriptUsage::EmitterUpdateScript:
-				return "EmitterUpdate";
-			case ENiagaraScriptUsage::ParticleSpawnScript:
-			case ENiagaraScriptUsage::ParticleSpawnScriptInterpolated:
-				return "ParticleSpawn";
-			case ENiagaraScriptUsage::ParticleUpdateScript:
-				return "ParticleUpdate";
-			case ENiagaraScriptUsage::ParticleSimulationStageScript:
-				return "ParticleStage";
-			case ENiagaraScriptUsage::ParticleEventScript:
-				return "ParticleEvent";
-			default:
-				return "Unknown";
-		}
-	};
-
 	switch (CompileOptions.TargetUsage)
 	{
 	case ENiagaraScriptUsage::ParticleSpawnScriptInterpolated:
@@ -1650,7 +1624,6 @@ FNiagaraTranslateResults TNiagaraHlslTranslator<GraphBridge>::Translate(const FN
 		SpawnStage.NumIterations = 1;
 		SpawnStage.bWritesParticles = true;
 		SpawnStage.bShouldUpdateInitialAttributeValues = true;
-		SpawnStage.DIOwnerName = GetDIOwnerNameByUsage(SpawnStage.ScriptUsage);
 		OutputNodes.Add(SourceGraph->FindEquivalentOutputNode(ENiagaraScriptUsage::ParticleSpawnScript, SpawnStage.UsageId));
 
 		FNiagaraHlslTranslationStage& UpdateStage = TranslationStages.Emplace_GetRef(ENiagaraScriptUsage::ParticleUpdateScript, FGuid());
@@ -1660,7 +1633,6 @@ FNiagaraTranslateResults TNiagaraHlslTranslator<GraphBridge>::Translate(const FN
 		UpdateStage.SimulationStageIndex = 0;
 		UpdateStage.NumIterations = 1;
 		UpdateStage.bWritesParticles = true;
-		UpdateStage.DIOwnerName = GetDIOwnerNameByUsage(UpdateStage.ScriptUsage);
 		OutputNodes.Add(SourceGraph->FindEquivalentOutputNode(UpdateStage.ScriptUsage, UpdateStage.UsageId));
 
 		ParamMapHistories.AddDefaulted(2);
@@ -1676,7 +1648,6 @@ FNiagaraTranslateResults TNiagaraHlslTranslator<GraphBridge>::Translate(const FN
 		SpawnStage.NumIterations = 1;
 		SpawnStage.bWritesParticles = true;
 		SpawnStage.bShouldUpdateInitialAttributeValues = true;
-		SpawnStage.DIOwnerName = GetDIOwnerNameByUsage(SpawnStage.ScriptUsage);
 		OutputNodes.Add(SourceGraph->FindEquivalentOutputNode(ENiagaraScriptUsage::ParticleSpawnScript, TranslationStages[0].UsageId));
 
 		FNiagaraHlslTranslationStage& UpdateStage = TranslationStages.Emplace_GetRef(ENiagaraScriptUsage::ParticleUpdateScript, FGuid());
@@ -1690,7 +1661,6 @@ FNiagaraTranslateResults TNiagaraHlslTranslator<GraphBridge>::Translate(const FN
 		UpdateStage.SimulationStageIndex = 0;
 		UpdateStage.NumIterations = 1;
 		UpdateStage.bWritesParticles = true;
-		UpdateStage.DIOwnerName = GetDIOwnerNameByUsage(UpdateStage.ScriptUsage);
 		OutputNodes.Add(SourceGraph->FindEquivalentOutputNode(TranslationStages[1].ScriptUsage, TranslationStages[1].UsageId));
 
 		ParamMapHistories.AddDefaulted(2);
@@ -1780,7 +1750,6 @@ FNiagaraTranslateResults TNiagaraHlslTranslator<GraphBridge>::Translate(const FN
 					SimStage.DirectDispatchElementType = CompileSimStageData.DirectDispatchElementType;
 					SimStage.bOverrideGpuDispatchNumThreads = CompileSimStageData.bOverrideGpuDispatchNumThreads;
 					SimStage.OverrideGpuDispatchNumThreads = CompileSimStageData.OverrideGpuDispatchNumThreads;
-					SimStage.DIOwnerName = *FString::Printf(TEXT("%s%d%s"), *GetDIOwnerNameByUsage(SimStage.ScriptUsage).ToString(), DestSimStageIndex, *StageName);
 
 					if (CompileSimStageData.bOverrideGpuDispatchNumThreads)
 					{
@@ -1857,16 +1826,17 @@ FNiagaraTranslateResults TNiagaraHlslTranslator<GraphBridge>::Translate(const FN
 					
 					// Set up the compile output for the shader stages so that we can properly execute at runtime.
 					FSimulationStageMetaData& SimulationStageMetaData = CompilationOutput.ScriptData.SimulationStageMetaData.AddDefaulted_GetRef();
-					SimulationStageMetaData.SimulationStageName = CompileSimStageData.StageName;
-					SimulationStageMetaData.EnabledBinding = SimStage.EnabledBinding;
-					SimulationStageMetaData.ElementCount = SimStage.ElementCount;
-					SimulationStageMetaData.ElementCountXBinding = SimStage.ElementCountXBinding;
-					SimulationStageMetaData.ElementCountYBinding = SimStage.ElementCountYBinding;
-					SimulationStageMetaData.ElementCountZBinding = SimStage.ElementCountZBinding;
+					SimulationStageMetaData.SimulationStageName = CompileData->CompileSimStageData[SourceSimStageIndex].StageName;
+					SimulationStageMetaData.EnabledBinding = CompileData->CompileSimStageData[SourceSimStageIndex].EnabledBinding;
+					SimulationStageMetaData.ElementCount = CompileData->CompileSimStageData[SourceSimStageIndex].ElementCount;
+					SimulationStageMetaData.ElementCountXBinding = CompileData->CompileSimStageData[SourceSimStageIndex].ElementCountXBinding;
+					SimulationStageMetaData.ElementCountYBinding = CompileData->CompileSimStageData[SourceSimStageIndex].ElementCountYBinding;
+					SimulationStageMetaData.ElementCountZBinding = CompileData->CompileSimStageData[SourceSimStageIndex].ElementCountZBinding;
 					SimulationStageMetaData.ExecuteBehavior = SimStage.ExecuteBehavior;
-					SimulationStageMetaData.IterationSourceType = SimStage.IterationSourceType;
-					SimulationStageMetaData.IterationDirectBinding = SimStage.IterationDirectBinding;
-					SimulationStageMetaData.NumIterationsBinding = SimStage.NumIterationsBinding;
+					SimulationStageMetaData.IterationSourceType = CompileData->CompileSimStageData[SourceSimStageIndex].IterationSourceType;
+					SimulationStageMetaData.IterationDataInterface = CompileData->CompileSimStageData[SourceSimStageIndex].IterationDataInterface;
+					SimulationStageMetaData.IterationDirectBinding = CompileData->CompileSimStageData[SourceSimStageIndex].IterationDirectBinding;
+					SimulationStageMetaData.NumIterationsBinding = CompileData->CompileSimStageData[SourceSimStageIndex].NumIterationsBinding;
 					SimulationStageMetaData.NumIterations = SimStage.NumIterations;
 					SimulationStageMetaData.bWritesParticles = SimStage.bWritesParticles;
 					SimulationStageMetaData.bPartialParticleUpdate = SimStage.bPartialParticleUpdate;
@@ -1880,13 +1850,13 @@ FNiagaraTranslateResults TNiagaraHlslTranslator<GraphBridge>::Translate(const FN
 					SimulationStageMetaData.bGpuIndirectDispatch = false;
 					if ( SimulationStageMetaData.IterationSourceType == ENiagaraIterationSource::DataInterface )
 					{
-						if (SimStage.IterationDataInterface.IsNone())
+						if (SimulationStageMetaData.IterationDataInterface.IsNone())
 						{
 							Error(LOCTEXT("NoIterationDataInterface", "An data interface must be bound to a data interface iteration stage"), nullptr, nullptr);
 							return TranslateResults;
 						}
 
-						if (const FNiagaraVariable* IterationSourceVar = CompileData->EncounteredVariables.FindByPredicate([&](const FNiagaraVariable& VarInfo) { return VarInfo.GetName() == SimStage.IterationDataInterface; }))
+						if (const FNiagaraVariable* IterationSourceVar = CompileData->EncounteredVariables.FindByPredicate([&](const FNiagaraVariable& VarInfo) { return VarInfo.GetName() == SimulationStageMetaData.IterationDataInterface; }))
 						{
 							if ( UNiagaraDataInterface* IterationSourceCDO = CompileDuplicateData->GetDuplicatedDataInterfaceCDOForClass(IterationSourceVar->GetType().GetClass()) )
 							{
@@ -1936,7 +1906,6 @@ FNiagaraTranslateResults TNiagaraHlslTranslator<GraphBridge>::Translate(const FN
 		DefaultStage.NumIterations = 1;
 		DefaultStage.bWritesParticles = true;
 		DefaultStage.bShouldUpdateInitialAttributeValues = TranslationStages[0].ShouldDoSpawnOnlyLogic() || (IsEventSpawnScript() && CompileOptions.AdditionalDefines.Contains(FNiagaraCompileOptions::EventSpawnInitialAttribWritesDefine));
-		DefaultStage.DIOwnerName = GetDIOwnerNameByUsage(DefaultStage.ScriptUsage);
 
 		OutputNodes.Add(SourceGraph->FindEquivalentOutputNode(TranslationStages[0].ScriptUsage, TranslationStages[0].UsageId));
 
@@ -2120,9 +2089,6 @@ FNiagaraTranslateResults TNiagaraHlslTranslator<GraphBridge>::Translate(const FN
 			if (TranslationStages[i].ShouldDoSpawnOnlyLogic())
 				bInitializedDefaults = false;
 
-			TGuardValue<FDIOwnerContext> StageDIOwnerContextScope(CurrentDIOwnerContext, FDIOwnerContext(FGuid(),
-				TranslationStages[i].ScriptUsage, TranslationStages[i].UsageId, TranslationStages[i].DIOwnerName));
-
 			if (UNiagaraScript::IsParticleSpawnScript(TranslationStages[i].ScriptUsage))
 			{
 				AddBodyComment(bHasInterpolatedSpawn ? TEXT("//Begin Interpolated Spawn Script!") : TEXT("//Begin Spawn Script!"));
@@ -2181,9 +2147,6 @@ FNiagaraTranslateResults TNiagaraHlslTranslator<GraphBridge>::Translate(const FN
 	}
 	else if (TranslationStages.Num() == 1)
 	{
-		TGuardValue<FDIOwnerContext> StageDIOwnerContextScope(CurrentDIOwnerContext, FDIOwnerContext(FGuid(),
-			TranslationStages[0].ScriptUsage, TranslationStages[0].UsageId, TranslationStages[0].DIOwnerName));
-
 		CurrentBodyChunkMode = TranslationStages[0].ChunkModeIndex;
 		ActiveStageIdx = 0;
 		check(CompileOptions.TargetUsage == TranslationStages[0].ScriptUsage);
@@ -2248,17 +2211,17 @@ FNiagaraTranslateResults TNiagaraHlslTranslator<GraphBridge>::Translate(const FN
 			{
 				int32 OutputIdx = 0;
 				//ensure the interpolated spawn constants are part of the parameter set.
-				ParameterMapRegisterExternalConstantNamespaceVariable(SYS_PARAM_ENGINE_TIME, FNiagaraVariableBase(), nullptr, 0, OutputIdx, nullptr);
-				ParameterMapRegisterExternalConstantNamespaceVariable(SYS_PARAM_ENGINE_WORLD_DELTA_TIME, FNiagaraVariableBase(), nullptr, 0, OutputIdx, nullptr);
-				ParameterMapRegisterExternalConstantNamespaceVariable(SYS_PARAM_ENGINE_DELTA_TIME, FNiagaraVariableBase(), nullptr, 0, OutputIdx, nullptr);
-				ParameterMapRegisterExternalConstantNamespaceVariable(SYS_PARAM_ENGINE_INV_DELTA_TIME, FNiagaraVariableBase(), nullptr, 0, OutputIdx, nullptr);
-				ParameterMapRegisterExternalConstantNamespaceVariable(SYS_PARAM_ENGINE_EXEC_COUNT, FNiagaraVariableBase(), nullptr, 0, OutputIdx, nullptr);
-				ParameterMapRegisterExternalConstantNamespaceVariable(SYS_PARAM_EMITTER_SPAWNRATE, FNiagaraVariableBase(), nullptr, 0, OutputIdx, nullptr);
+				ParameterMapRegisterExternalConstantNamespaceVariable(SYS_PARAM_ENGINE_TIME, nullptr, 0, OutputIdx, nullptr);
+				ParameterMapRegisterExternalConstantNamespaceVariable(SYS_PARAM_ENGINE_WORLD_DELTA_TIME, nullptr, 0, OutputIdx, nullptr);
+				ParameterMapRegisterExternalConstantNamespaceVariable(SYS_PARAM_ENGINE_DELTA_TIME, nullptr, 0, OutputIdx, nullptr);
+				ParameterMapRegisterExternalConstantNamespaceVariable(SYS_PARAM_ENGINE_INV_DELTA_TIME, nullptr, 0, OutputIdx, nullptr);
+				ParameterMapRegisterExternalConstantNamespaceVariable(SYS_PARAM_ENGINE_EXEC_COUNT, nullptr, 0, OutputIdx, nullptr);
+				ParameterMapRegisterExternalConstantNamespaceVariable(SYS_PARAM_EMITTER_SPAWNRATE, nullptr, 0, OutputIdx, nullptr);
 				if (CompilationTarget != ENiagaraSimTarget::GPUComputeSim)
 				{
-					ParameterMapRegisterExternalConstantNamespaceVariable(SYS_PARAM_EMITTER_SPAWN_INTERVAL, FNiagaraVariableBase(), nullptr, 0, OutputIdx, nullptr);
-					ParameterMapRegisterExternalConstantNamespaceVariable(SYS_PARAM_EMITTER_INTERP_SPAWN_START_DT, FNiagaraVariableBase(), nullptr, 0, OutputIdx, nullptr);
-					ParameterMapRegisterExternalConstantNamespaceVariable(SYS_PARAM_EMITTER_SPAWN_GROUP, FNiagaraVariableBase(), nullptr, 0, OutputIdx, nullptr);
+					ParameterMapRegisterExternalConstantNamespaceVariable(SYS_PARAM_EMITTER_SPAWN_INTERVAL, nullptr, 0, OutputIdx, nullptr);
+					ParameterMapRegisterExternalConstantNamespaceVariable(SYS_PARAM_EMITTER_INTERP_SPAWN_START_DT, nullptr, 0, OutputIdx, nullptr);
+					ParameterMapRegisterExternalConstantNamespaceVariable(SYS_PARAM_EMITTER_SPAWN_GROUP, nullptr, 0, OutputIdx, nullptr);
 				}
 			}
 
@@ -2266,10 +2229,10 @@ FNiagaraTranslateResults TNiagaraHlslTranslator<GraphBridge>::Translate(const FN
 			{
 				int32 OutputIdx = 0;
 				// NOTE(mv): This will explicitly expose Engine.Emitter.TotalSpawnedParticles to the HLSL code regardless of whether it is exposed in a script or not. 
-				ParameterMapRegisterExternalConstantNamespaceVariable(SYS_PARAM_ENGINE_EMITTER_TOTAL_SPAWNED_PARTICLES, FNiagaraVariableBase(), nullptr, 0, OutputIdx, nullptr);
-				ParameterMapRegisterExternalConstantNamespaceVariable(SYS_PARAM_EMITTER_RANDOM_SEED, FNiagaraVariableBase(), nullptr, 0, OutputIdx, nullptr);
-				ParameterMapRegisterExternalConstantNamespaceVariable(SYS_PARAM_ENGINE_EMITTER_INSTANCE_SEED, FNiagaraVariableBase(), nullptr, 0, OutputIdx, nullptr);
-				ParameterMapRegisterExternalConstantNamespaceVariable(SYS_PARAM_ENGINE_SYSTEM_RANDOM_SEED, FNiagaraVariableBase(), nullptr, 0, OutputIdx, nullptr);
+				ParameterMapRegisterExternalConstantNamespaceVariable(SYS_PARAM_ENGINE_EMITTER_TOTAL_SPAWNED_PARTICLES, nullptr, 0, OutputIdx, nullptr);
+				ParameterMapRegisterExternalConstantNamespaceVariable(SYS_PARAM_EMITTER_RANDOM_SEED, nullptr, 0, OutputIdx, nullptr);
+				ParameterMapRegisterExternalConstantNamespaceVariable(SYS_PARAM_ENGINE_EMITTER_INSTANCE_SEED, nullptr, 0, OutputIdx, nullptr);
+				ParameterMapRegisterExternalConstantNamespaceVariable(SYS_PARAM_ENGINE_SYSTEM_RANDOM_SEED, nullptr, 0, OutputIdx, nullptr);
 			}
 		}
 
@@ -2794,30 +2757,26 @@ void TNiagaraHlslTranslator<GraphBridge>::HandleSimStageSetupAndTeardown(int32 I
 		}
 	}
 
-	// Try to find the owner by written data interface parameters.
-	FNiagaraCompileParameterReference IterationSourceReference(IterationSourceVar.GetName(), FGuid(), ExtractFlagsFromParameterName(IterationSourceVar.GetName()));
-	int32 IterationSourceDataInterfaceInfoIndex = CompilationOutput.ScriptData.DataInterfaceInfo.IndexOfByPredicate([IterationSourceReference]
-		(const FNiagaraScriptDataInterfaceCompileInfo& CompileInfo) { return CompileInfo.TargetParameter == IterationSourceReference; });
-
-	if (IterationSourceDataInterfaceInfoIndex == INDEX_NONE)
+	// Find the data interface in the table. Note that this may not be found because we don't actually call any functions on the data interface yet.
+	int32 DataInterfaceOwnerIndex = INDEX_NONE;
+	for (int32 i = 0; i < CompilationOutput.ScriptData.DataInterfaceInfo.Num(); i++)
 	{
-		// If a written data interface wasn't found, try to find a read parameter.
-		IterationSourceDataInterfaceInfoIndex = CompilationOutput.ScriptData.DataInterfaceInfo.IndexOfByPredicate([IterationSourceReference]
-			(const FNiagaraScriptDataInterfaceCompileInfo& CompileInfo) { return CompileInfo.SourceExternalParameter == IterationSourceReference; });
+		FNiagaraScriptDataInterfaceCompileInfo& Info = CompilationOutput.ScriptData.DataInterfaceInfo[i];
+		if (IterationDataInterface == Info.Name)
+		{
+			DataInterfaceOwnerIndex = i;
+			break;
+		}
 	}
 
 	// RIght now we need to know if anyone wrote to the IterationSource this stage. That can be one of two ways:
 	// 1) Someone wrote to StackContext.XXXX
-	// 2) Someone called a function that was marked to write
-	bool bFunctionReadFromIterationSource = false;
-	bool bFunctionWroteToIterationSource = false;
+	// 2) Someone called a function that was marked to write 
 	const int32 SourceSimStage = TranslationStage.SimulationStageIndex;
-	if (IterationSourceDataInterfaceInfoIndex != INDEX_NONE)
-	{
-		FName IterationSourceDataInterfaceName = CompilationOutput.ScriptData.DataInterfaceInfo[IterationSourceDataInterfaceInfoIndex].Name;
-		bFunctionReadFromIterationSource = CompilationOutput.ScriptData.SimulationStageMetaData[SourceSimStage].InputDataInterfaces.Contains(IterationSourceDataInterfaceName);
-		bFunctionWroteToIterationSource = CompilationOutput.ScriptData.SimulationStageMetaData[SourceSimStage].OutputDestinations.Contains(IterationSourceDataInterfaceName);
-	}
+	ensure(CompilationOutput.ScriptData.SimulationStageMetaData.IsValidIndex(SourceSimStage));
+	bool bWroteToIterationSource = CompilationOutput.ScriptData.SimulationStageMetaData[SourceSimStage].OutputDestinations.Contains(TranslationStage.GetIterationDataInterface());
+	if (WriteVars.Num() > 0)
+		bWroteToIterationSource = true;
 
 	// Now decide if we need to put in the pre/post
 	if (CDO && CDO->CanExecuteOnTarget(ENiagaraSimTarget::GPUComputeSim))
@@ -2827,14 +2786,8 @@ void TNiagaraHlslTranslator<GraphBridge>::HandleSimStageSetupAndTeardown(int32 I
 		bool bNeedsAttributeWrite = false;
 		bool bNeedsAttributeRead = false;
 
-		// If a function read from or wrote to the iteration source data interface, then it needs an owner.
-		if (bFunctionReadFromIterationSource || bFunctionWroteToIterationSource)
-		{
-			bNeedsDIOwner = true;
-		}
-
 		// Put in the general pre/post if we wrote to the IterationSoruce at all
-		if (CDO->SupportsSetupAndTeardownHLSL() && (bFunctionWroteToIterationSource || WriteVars.Num() > 0))
+		if (CDO->SupportsSetupAndTeardownHLSL() && bWroteToIterationSource)
 		{
 			bNeedsDIOwner = true;
 			bNeedsSetupAndTeardown = true;
@@ -2856,27 +2809,18 @@ void TNiagaraHlslTranslator<GraphBridge>::HandleSimStageSetupAndTeardown(int32 I
 		}
 		
 		// If it wasn't previously added, let's go ahead and do so. Maybe they are solely using the StackContext namespace.
-		if (IterationSourceDataInterfaceInfoIndex == INDEX_NONE && bNeedsDIOwner)
+		if (DataInterfaceOwnerIndex == INDEX_NONE && bNeedsDIOwner)
 		{
-			IterationSourceDataInterfaceInfoIndex = RegisterDataInterfaceForParameterRead(IterationSourceVar, INDEX_NONE);
+			DataInterfaceOwnerIndex = RegisterDataInterface(IterationSourceVar, CDO, true, true);
 		}
 
 		// If we haven't created it by now, bail out. 
-		if (bNeedsDIOwner && IterationSourceDataInterfaceInfoIndex == INDEX_NONE)
+		if (DataInterfaceOwnerIndex == INDEX_NONE && bNeedsDIOwner)
 		{
 			Error(FText::Format(LOCTEXT("CannotRegisterDataInterface", "Variable {0}'s cannot register as a data interface!"), FText::FromName(IterationDataInterface)), nullptr, nullptr);
 			return;
 		}
 
-		// If we need to bind the iteration source data interface, specify the unique owner name here on the stage metadata for binding at runtime.
-		// We use the unique owner name here rather than the iteration source parameter name to avoid issues with emitter aliasing or other ambiguous
-		// resolving issues.
-		if (bNeedsDIOwner)
-		{
-			CompilationOutput.ScriptData.SimulationStageMetaData[TranslationStage.SimulationStageIndex].IterationDataInterface =
-				CompilationOutput.ScriptData.DataInterfaceInfo[IterationSourceDataInterfaceInfoIndex].Name;
-		}
-
 		// It is an invalid state to use the IterationSource and StackContext namespace without implementing SupportsIterationSourceNamespaceAttributesHLSL
 		if (ReadVars.Num() > 0 && !bNeedsAttributeRead)
 		{
@@ -2898,8 +2842,7 @@ void TNiagaraHlslTranslator<GraphBridge>::HandleSimStageSetupAndTeardown(int32 I
 		// Convert to a FNiagaraDataInterfaceGPUParamInfo to keep the API simple and consistent
 		TArray<FNiagaraFunctionSignature> GeneratedFunctionSignatures;
 		FNiagaraDataInterfaceGPUParamInfo DIInstanceInfo;
-		const FNiagaraScriptDataInterfaceCompileInfo& Info = CompilationOutput.ScriptData.DataInterfaceInfo[IterationSourceDataInterfaceInfoIndex];
-		ConvertCompileInfoToParamInfo(Info, DIInstanceInfo, GeneratedFunctionSignatures);
+		ConvertCompileInfoToParamInfo(CompilationOutput.ScriptData.DataInterfaceInfo[DataInterfaceOwnerIndex], DIInstanceInfo, GeneratedFunctionSignatures);
 
 		// This next part might be a big confusing, but because DataInterfaces are in non-editor code, it makes it impossible for them to return graphs or other
 		// structures. We want them to feel free to invoke their own fucntions and not have to do a lot of extra wranging, so we treat them like a custom hlsl node.
@@ -2917,7 +2860,7 @@ void TNiagaraHlslTranslator<GraphBridge>::HandleSimStageSetupAndTeardown(int32 I
 
 		TArray<int32> Inputs;
 		Inputs.Emplace(InWhichStage); // The parameter map input index above
-		Inputs.Emplace(IterationSourceDataInterfaceInfoIndex); // the data interface index above
+		Inputs.Emplace(DataInterfaceOwnerIndex); // the data interface index above
 		FString SetupBody;
 		FString TeardownBody;
 
@@ -3423,7 +3366,10 @@ void FNiagaraHlslTranslator::DefineDataSetWriteFunction(FString &HlslOutputStrin
 
 void FNiagaraHlslTranslator::ConvertCompileInfoToParamInfo(const FNiagaraScriptDataInterfaceCompileInfo& Info, FNiagaraDataInterfaceGPUParamInfo& DIInstanceInfo, TArray<FNiagaraFunctionSignature>& GeneratedFunctionSignatures)
 {
-	DIInstanceInfo.DataInterfaceHLSLSymbol = Info.Name.ToString();
+	FString OwnerIDString = Info.Name.ToString();
+	FString SanitizedOwnerIDString = BuildSanitizedSymbolString(OwnerIDString.Replace(TEXT("."), TEXT("_")));
+
+	DIInstanceInfo.DataInterfaceHLSLSymbol = SanitizedOwnerIDString;
 	DIInstanceInfo.DIClassName = Info.Type.GetClass()->GetName();
 
 	// Build a list of function instances that will be generated for this DI.
@@ -5301,7 +5247,7 @@ int32 TNiagaraHlslTranslator<GraphBridge>::GetRapidIterationParameter(const FNia
 
 		// go ahead and make it into a constant variable..
 		int32 OutputChunkId = INDEX_NONE;
-		if (ParameterMapRegisterExternalConstantNamespaceVariable(Parameter, FNiagaraVariableBase(), nullptr, INDEX_NONE, OutputChunkId, nullptr))
+		if (ParameterMapRegisterExternalConstantNamespaceVariable(Parameter, nullptr, INDEX_NONE, OutputChunkId, nullptr))
 		{
 			return OutputChunkId;
 		}
@@ -5427,7 +5373,7 @@ int32 TNiagaraHlslTranslator<GraphBridge>::GetParameter(const FNiagaraVariable&
 	FNiagaraVariable OutputVariable = Parameter;
 	if (FNiagaraParameterUtilities::IsInNamespace(OutputVariable, PARAM_MAP_ATTRIBUTE_STR) || FNiagaraParameterUtilities::IsExternalConstantNamespace(OutputVariable, CompileOptions.TargetUsage, CompileOptions.GetTargetUsageBitmask()))
 	{
-		if (!ParameterMapRegisterExternalConstantNamespaceVariable(OutputVariable, FNiagaraVariableBase(), nullptr, 0, OutputChunkIdx, nullptr))
+		if (!ParameterMapRegisterExternalConstantNamespaceVariable(OutputVariable, nullptr, 0, OutputChunkIdx, nullptr))
 		{
 			OutputChunkIdx = INDEX_NONE;
 		}
@@ -5435,7 +5381,7 @@ int32 TNiagaraHlslTranslator<GraphBridge>::GetParameter(const FNiagaraVariable&
 	else
 	{
 		OutputVariable = FNiagaraParameterUtilities::MoveToExternalConstantNamespaceVariable(OutputVariable, CompileOptions.TargetUsage);
-		if (!ParameterMapRegisterExternalConstantNamespaceVariable(OutputVariable, FNiagaraVariableBase(), nullptr, 0, OutputChunkIdx, nullptr))
+		if (!ParameterMapRegisterExternalConstantNamespaceVariable(OutputVariable, nullptr, 0, OutputChunkIdx, nullptr))
 		{
 			OutputChunkIdx = INDEX_NONE;
 		}
@@ -6125,7 +6071,6 @@ void TNiagaraHlslTranslator<GraphBridge>::ParameterMapSet(const FParamMapSetNode
 				continue;
 			}
 
-			FNiagaraVariableBase OriginalVar = Var;
 			Var = ActiveHistoryForFunctionCalls.ResolveAliases(Var);
 			
 			const FNiagaraKnownConstantInfo ConstantInfo = FNiagaraConstants::GetKnownConstantInfo(Var.GetName(), false);
@@ -6212,6 +6157,14 @@ void TNiagaraHlslTranslator<GraphBridge>::ParameterMapSet(const FParamMapSetNode
 
 			if (Var.IsDataInterface())
 			{
+				if (CompileOptions.TargetUsage == ENiagaraScriptUsage::ParticleSpawnScriptInterpolated && TranslationStages[ActiveStageIdx].ScriptUsage == ENiagaraScriptUsage::ParticleUpdateScript)
+				{
+					// We don't want to add writes for particle update data interface parameters in both interpolated spawn and update, so skip them when processing the update stage of the 
+					// interpolated spawn script.  We don't skip the writes when compiling the particle update script because it's not recompiled when the interpolated spawn flag is changed
+					// and this would result in missing data interfaces if interpolated spawn was turned off.
+					continue;
+				}
+
 				bool bAllowDataInterfaces = true;
 				if (ParamMapHistoryIdx < ParamMapHistories.Num() && ParamMapHistories[ParamMapHistoryIdx].IsPrimaryDataSetOutput(Var, CompileOptions.TargetUsage, bAllowDataInterfaces))
 				{
@@ -6221,7 +6174,37 @@ void TNiagaraHlslTranslator<GraphBridge>::ParameterMapSet(const FParamMapSetNode
 						continue;
 					}
 
-					RegisterDataInterfaceForParameterWrite(OriginalVar, Input);
+					FName UsageName;
+					if (FNiagaraParameterUtilities::IsAliasedEmitterParameter(Var))
+					{
+						FNiagaraVariable AliasedVar = ActiveHistoryForFunctionCalls.ResolveAliases(Var);
+						UsageName = AliasedVar.GetName();
+					}
+					else
+					{
+						UsageName = Var.GetName();
+					}
+
+					FNiagaraScriptDataInterfaceCompileInfo& Info = CompilationOutput.ScriptData.DataInterfaceInfo[Input];
+					if (Info.RegisteredParameterMapWrite == NAME_None)
+					{
+						Info.RegisteredParameterMapWrite = UsageName;
+					}
+					else
+					{
+						// This is a special case where a single data interface is written to multiple parameters.  In order to handle this correctly the data interface
+						// needs it's own registration slot so that it ends up in the parameter store under the correct name so that it can be looked up correctly in various places
+						// and can be bound properly to the renderer parameter stores.
+						UNiagaraDataInterface* CDO = CompileDuplicateData->GetDuplicatedDataInterfaceCDOForClass(const_cast<UClass*>(Var.GetType().GetClass()));
+						if (CDO)
+						{
+							FString PlaceholderName = FString::Printf(TEXT("Placeholder.Placeholder%i"), CompilationOutput.ScriptData.DataInterfaceInfo.Num());
+							FNiagaraVariable PlaceholderVar(Var.GetType(), *PlaceholderName);
+							int32 PlaceholderIndex = RegisterDataInterface(PlaceholderVar, CDO, true, false);
+							CompilationOutput.ScriptData.DataInterfaceInfo[PlaceholderIndex].RegisteredParameterMapRead = CompilationOutput.ScriptData.DataInterfaceInfo[Input].RegisteredParameterMapWrite;
+							CompilationOutput.ScriptData.DataInterfaceInfo[PlaceholderIndex].RegisteredParameterMapWrite = UsageName;
+						}
+					}
 				}
 			}
 			else if ( Var.IsUObject() )
@@ -6438,14 +6421,14 @@ bool TNiagaraHlslTranslator<GraphBridge>::HandleBoundConstantVariableToDataSetRe
 
 		if (bIsEmitterLocalSpaceCompileOptionSet == false)
 		{
-			return ParameterMapRegisterExternalConstantNamespaceVariable(SYS_PARAM_ENGINE_POSITION, FNiagaraVariableBase(), InNode, InParamMapHistoryIdx, Output, InDefaultPin);
+			return ParameterMapRegisterExternalConstantNamespaceVariable(SYS_PARAM_ENGINE_POSITION, InNode, InParamMapHistoryIdx, Output, InDefaultPin);
 		}
 	}
 	return false;
 }
 
 template<typename GraphBridge>
-bool TNiagaraHlslTranslator<GraphBridge>::ParameterMapRegisterExternalConstantNamespaceVariable(FNiagaraVariable InVariable, const FNiagaraVariableBase& InUnresolvedVariable, const FNode* InNodeForErrorReporting, int32 InParamMapHistoryIdx, int32& Output, const FPin* InDefaultPin)
+bool TNiagaraHlslTranslator<GraphBridge>::ParameterMapRegisterExternalConstantNamespaceVariable(FNiagaraVariable InVariable, const FNode* InNodeForErrorReporting, int32 InParamMapHistoryIdx, int32& Output, const FPin* InDefaultPin)
 {
 	Output = INDEX_NONE;
 
@@ -6498,8 +6481,31 @@ bool TNiagaraHlslTranslator<GraphBridge>::ParameterMapRegisterExternalConstantNa
 			{
 				if (bIsDataInterface)
 				{
-					Output = RegisterDataInterfaceForParameterRead(InUnresolvedVariable.IsValid() ? InUnresolvedVariable : InVariable, INDEX_NONE);
-					return true;
+					UNiagaraDataInterface* DataInterface = nullptr;
+					if (GraphBridge::IsCollectionValid(Collection))
+					{
+						DataInterface = GraphBridge::GetCollectionDataInterface(Collection, InVariable);
+						if (DataInterface == nullptr)
+						{
+							Error(FText::Format(LOCTEXT("ParameterCollectionDataInterfaceNotFoundErrorFormat", "Data interface named {0} of type {1} was not found in Parameter Collection {2}"),
+								FText::FromName(InVariable.GetName()), InVariable.GetType().GetNameText(), FText::FromString(GraphBridge::GetCollectionFullName(Collection))), InNodeForErrorReporting, InDefaultPin);
+							return false;
+						}
+					}
+					else
+					{
+						DataInterface = CompileDuplicateData->GetDuplicatedDataInterfaceCDOForClass(const_cast<UClass*>(InVariable.GetType().GetClass()));
+						if (DataInterface == nullptr)
+						{
+							Error(FText::Format(LOCTEXT("GetDuplicatedDataInterfaceCDOForClassFailed", "GetDuplicatedDataInterfaceCDOForClass failed for Variable({0}) Class({1})"), FText::FromName(InVariable.GetName()), InVariable.GetType().GetNameText()), InNodeForErrorReporting, InDefaultPin);
+							return false;
+						}
+					}
+					if (ensure(DataInterface))
+					{
+						Output = RegisterDataInterface(InVariable, DataInterface, true, true);
+						return true;
+					}
 				}
 				if (bIsUObject)
 				{
@@ -6984,8 +6990,6 @@ void TNiagaraHlslTranslator<GraphBridge>::Emitter(const FEmitterNode* EmitterNod
 
 	{
 		FScopedBuilderEmitter<GraphBridge> EmitterBuidlerScope(ActiveHistoryForFunctionCalls, EmitterNode);
-		TGuardValue<FDIOwnerContext> EmitterDIOwnerContextScope(CurrentDIOwnerContext, FDIOwnerContext(GraphBridge::GetEmitterHandleId(EmitterNode),
-			GraphBridge::GetEmitterUsage(EmitterNode), FGuid(), *GraphBridge::GetEmitterUniqueName(EmitterNode)));
 
 		// Clear out the parameter map writes to emitter module parameters as they should not be shared across emitters.
 		if (ParamMapHistoryIdx != -1 && ParamMapHistoryIdx < ParamMapHistories.Num())
@@ -7212,7 +7216,7 @@ void TNiagaraHlslTranslator<GraphBridge>::HandleParameterRead(int32 ParamMapHist
 			OutputChunkId = MakeStaticVariableDirect(DefaultPin);
 			return;
 		}
-		else if (ParameterMapRegisterExternalConstantNamespaceVariable(Var, InVar, ErrorNode, ParamMapHistoryIdx, OutputChunkId, DefaultPin))
+		else if (ParameterMapRegisterExternalConstantNamespaceVariable(Var, ErrorNode, ParamMapHistoryIdx, OutputChunkId, DefaultPin))
 		{
 			return;
 		}			
@@ -7319,7 +7323,7 @@ void TNiagaraHlslTranslator<GraphBridge>::HandleParameterRead(int32 ParamMapHist
 				}
 
 				// If it isn't found yet.. go ahead and make it into a constant variable..
-				if (LastSetChunkIdx == INDEX_NONE && ParameterMapRegisterExternalConstantNamespaceVariable(Var, InVar, ErrorNode, ParamMapHistoryIdx, OutputChunkId, InputPin))
+				if (LastSetChunkIdx == INDEX_NONE && ParameterMapRegisterExternalConstantNamespaceVariable(Var, ErrorNode, ParamMapHistoryIdx, OutputChunkId, InputPin))
 				{
 					LastSetChunkIdx = OutputChunkId;
 					if (VarIdx != INDEX_NONE && VarIdx < ParamMapSetVariablesToChunks[ParamMapHistoryIdx].Num())
@@ -7579,7 +7583,7 @@ void TNiagaraHlslTranslator<GraphBridge>::HandleParameterRead(int32 ParamMapHist
 				if (Var.IsDataInterface() || Var.IsUObject())
 				{
 					// If this variable is a data interface and it's in the parameter map, but hasn't been set yet, then is is an external data interface so try to register it.
-					if (ParameterMapRegisterExternalConstantNamespaceVariable(Var, InVar, ErrorNode, ParamMapHistoryIdx, OutputChunkId, DefaultPin))
+					if (ParameterMapRegisterExternalConstantNamespaceVariable(Var, ErrorNode, ParamMapHistoryIdx, OutputChunkId, DefaultPin))
 					{
 						return;
 					}
@@ -7596,7 +7600,12 @@ void TNiagaraHlslTranslator<GraphBridge>::HandleParameterRead(int32 ParamMapHist
 		// for variables encountered without ever being set. We do this by creating an instance of the CDO.
 		if (UNiagaraScript::IsStandaloneScript(CompileOptions.TargetUsage) && LastSetChunkIdx == INDEX_NONE)
 		{
-			LastSetChunkIdx = RegisterDataInterfaceForParameterRead(Var, INDEX_NONE);
+			UNiagaraDataInterface* DataInterface = CompileDuplicateData->GetDuplicatedDataInterfaceCDOForClass(const_cast<UClass*>(Var.GetType().GetClass()));
+			check(DataInterface != nullptr);
+			if (DataInterface)
+			{
+				LastSetChunkIdx = RegisterDataInterface(Var, DataInterface, true, false);
+			}
 		}
 
 		OutputChunkId = LastSetChunkIdx;
@@ -7846,6 +7855,82 @@ int32 TNiagaraHlslTranslator<GraphBridge>::RegisterUObjectPath(const FNiagaraVar
 	return ObjectIndex;
 }
 
+template<typename GraphBridge>
+int32 TNiagaraHlslTranslator<GraphBridge>::RegisterDataInterface(const FNiagaraVariable& Var, const UNiagaraDataInterface* DataInterface, bool bPlaceholder, bool bAddParameterMapRead)
+{
+	//FString PathName = DataInterface ? *DataInterface->GetPathName() : TEXT("XX");
+	//UE_LOG(LogNiagaraEditor, Log, TEXT("RegisterDataInterface %s %s %s '%s' %s"), *Var.GetName().ToString(), *Var.GetType().GetName(), *PathName,
+	//	bPlaceholder ? TEXT("bPlaceholder=true") : TEXT("bPlaceholder=false"));
+
+	int32 FuncParam;
+	if (GetFunctionParameter(Var, FuncParam))
+	{
+		if (FuncParam != INDEX_NONE)
+		{
+			//This data interface param has been overridden by the function call so use that index.	
+			return FuncParam;
+		}
+	}
+
+	//If we get here then this is a new data interface.
+	const FString* EmitterAlias = ActiveHistoryForFunctionCalls.GetEmitterAlias();
+	FName DataInterfaceName = GetDataInterfaceName(Var.GetName(), EmitterAlias != nullptr ? *EmitterAlias : FString(), bAddParameterMapRead);
+
+	if (DataInterface && DataInterface->NeedsGPUContextInit() && CompileOptions.IsGpuScript())
+	{
+		CompilationOutput.ScriptData.bNeedsGPUContextInit = true;
+	}
+
+	FNiagaraVariable ReadVariable;
+	if (bAddParameterMapRead)
+	{
+		ReadVariable = FNiagaraParameterUtilities::IsAliasedEmitterParameter(Var.GetName().ToString())
+			? ActiveHistoryForFunctionCalls.ResolveAliases(Var) 
+			: Var;
+	}
+
+	int32 Idx = CompilationOutput.ScriptData.DataInterfaceInfo.IndexOfByPredicate([&](const FNiagaraScriptDataInterfaceCompileInfo& OtherInfo)
+	{
+		return OtherInfo.Name == DataInterfaceName || (ReadVariable.IsValid() && OtherInfo.RegisteredParameterMapWrite == ReadVariable.GetName());
+	});
+
+	if (Idx == INDEX_NONE)
+	{
+		Idx = CompilationOutput.ScriptData.DataInterfaceInfo.AddDefaulted();
+		CompilationOutput.ScriptData.DataInterfaceInfo[Idx].Name = DataInterfaceName;
+		CompilationOutput.ScriptData.DataInterfaceInfo[Idx].Type = Var.GetType();
+		CompilationOutput.ScriptData.DataInterfaceInfo[Idx].bIsPlaceholder = bPlaceholder;
+
+		FString SourceEmitterName = FString();
+		if (CompileData->GetUniqueEmitterName().IsEmpty() == false)
+		{
+			SourceEmitterName = CompileData->GetUniqueEmitterName();
+		}
+		else if (EmitterAlias != nullptr)
+		{
+			SourceEmitterName = *EmitterAlias;
+		}
+		CompilationOutput.ScriptData.DataInterfaceInfo[Idx].SourceEmitterName = SourceEmitterName;
+
+		//Interface requires per instance data so add a user pointer table entry.
+		if (DataInterface != nullptr && DataInterface->PerInstanceDataSize() > 0)
+		{
+			CompilationOutput.ScriptData.DataInterfaceInfo[Idx].UserPtrIdx = CompilationOutput.ScriptData.NumUserPtrs++;
+		}
+
+		if (bAddParameterMapRead)
+		{
+			CompilationOutput.ScriptData.DataInterfaceInfo[Idx].RegisteredParameterMapRead = ReadVariable.GetName();
+		}
+	}
+	else
+	{
+		check(CompilationOutput.ScriptData.DataInterfaceInfo[Idx].Type == Var.GetType());
+	}
+
+	return Idx;
+}
+
 template<typename GraphBridge>
 void TNiagaraHlslTranslator<GraphBridge>::Operation(const FOpNode* Operation, TArray<int32>& Inputs, TArray<int32>& Outputs)
 {
@@ -8724,7 +8809,7 @@ void TNiagaraHlslTranslator<GraphBridge>::ProcessCustomHlsl(const FString& InCus
 
 			if (bPermuteSignatureByDataInterface)
 			{
-				OutSignature.Name = FName(OutSignature.Name.ToString() + Info.Name.ToString());
+				OutSignature.Name = FName(OutSignature.Name.ToString() + BuildSanitizedSymbolName(Info.Name, true));
 			}
 
 			SigInputs.Add(Input);
@@ -8798,7 +8883,7 @@ void TNiagaraHlslTranslator<GraphBridge>::ProcessCustomHlsl(const FString& InCus
 				FNiagaraParameterUtilities::IsExternalConstantNamespace(VarActual, CompileOptions.TargetUsage, CompileOptions.GetTargetUsageBitmask()))
 			{
 				int32 TempOutput;
-				if (ParameterMapRegisterExternalConstantNamespaceVariable(VarActual, VarAliased, InNodeForErrorReporting, ParamMapHistoryIdx, TempOutput, nullptr))
+				if (ParameterMapRegisterExternalConstantNamespaceVariable(VarActual, InNodeForErrorReporting, ParamMapHistoryIdx, TempOutput, nullptr))
 				{
 					continue;
 				}
@@ -9091,7 +9176,7 @@ void TNiagaraHlslTranslator<GraphBridge>::RegisterFunctionCall(ENiagaraScriptUsa
 								// for aliased values we check to see if they're used in the current context by resolving the alias and checking against the current resolved variable
 								// name since aliased values can only be resolved for reading in the correct context.
 								bool bIsValidForCurrentCallingContext = (bIsInTopLevelFunction && bIsAliased == false) || (bIsAliased && ActiveHistoryForFunctionCalls.ResolveAliases(AliasedVar).GetName() == Var.GetName());
-								if (bIsValidForCurrentCallingContext && !Var.GetType().IsStatic() && Var.IsDataInterface() == false)
+								if (bIsValidForCurrentCallingContext && !Var.GetType().IsStatic() )
 								{
 									int32 LastSetChunkIdx = ParamMapSetVariablesToChunks[ActiveStageIdx][VarIdx];
 									if (LastSetChunkIdx == INDEX_NONE)
@@ -9590,6 +9675,24 @@ FString FNiagaraHlslTranslator::GenerateFunctionParametersHlsl(const FNiagaraFun
 	return StringBuilder.ToString();
 }
 
+FName FNiagaraHlslTranslator::GetDataInterfaceName(FName BaseName, const FString& UniqueEmitterName, bool bIsParameterMapDataInterface)
+{
+	if (UniqueEmitterName.IsEmpty() == false)
+	{
+		if (FNiagaraParameterUtilities::IsAliasedEmitterParameter(BaseName.ToString()))
+		{
+			return FNiagaraParameterUtilities::ResolveEmitterAlias(BaseName, UniqueEmitterName);
+		}
+		else if(bIsParameterMapDataInterface == false)
+		{
+			// Don't mangle the parameter map reads for emitter scripts because they are from the system or user parameter stores and
+			// they won't bind correctly.
+			return *(UniqueEmitterName + TEXT(".") + BaseName.ToString());
+		}
+	}
+	return BaseName;
+}
+
 FString FNiagaraHlslTranslator::GetFunctionIncludeStatement(const FNiagaraCustomHlslInclude& Include) const
 {
 	TStringBuilder<128> IncludeStatement;
@@ -9820,247 +9923,6 @@ FString FNiagaraHlslTranslator::GenerateAssignment(const FNiagaraTypeDefinition&
 	return DestinationDefinition + " = " + SourceDefinition;
 }
 
-int32 FNiagaraHlslTranslator::RegisterDataInterfaceForInput(const FNiagaraVariableBase& InputVariable, const FNiagaraGraphInputDataInterfaceKey& GraphInputDataInterfaceKey, const UNiagaraDataInterface* InputDataInterface)
-{
-	int32 FuncParam;
-	if (GetFunctionParameter(InputVariable, FuncParam))
-	{
-		if (FuncParam != INDEX_NONE)
-		{
-			//This data interface param has been overridden by the function call so use that index.	
-			return FuncParam;
-		}
-	}
-
-	if (ensureMsgf(GraphInputDataInterfaceKey.IsValid(), TEXT("GraphInputDataInterfaceKey is valid for input variable %s"), *InputVariable.GetName().ToString()) == false)
-	{
-		return INDEX_NONE;
-	}
-
-	int32 ExistingIndex = CompilationOutput.ScriptData.DataInterfaceInfo.IndexOfByPredicate([&GraphInputDataInterfaceKey]
-	(const FNiagaraScriptDataInterfaceCompileInfo& OtherInfo) { return OtherInfo.SourceGraphInputDataInterfaceKey == GraphInputDataInterfaceKey; });
-	if (ExistingIndex != INDEX_NONE)
-	{
-		return ExistingIndex;
-	}
-
-	const UNiagaraDataInterface* InputDataInterfaceOrCDO = InputDataInterface != nullptr
-		? InputDataInterface
-		: GetDataInterfaceCDO(InputVariable.GetType().GetClass());
-
-	FString InternalCompileName = FString::Printf(TEXT("INTERNAL.%s"), *InputVariable.GetName().ToString());
-	int32 AddedIndex = AddDataInterfaceInfo(InputVariable.GetType(), *InternalCompileName, *InputDataInterfaceOrCDO);
-	FNiagaraScriptDataInterfaceCompileInfo& AddedInfo = CompilationOutput.ScriptData.DataInterfaceInfo[AddedIndex];
-	AddedInfo.SourceGraphInputDataInterfaceKey = GraphInputDataInterfaceKey;
-
-	return AddedIndex;
-}
-
-ENiagaraCompileParameterReferenceFlags FNiagaraHlslTranslator::ExtractFlagsFromParameterName(FName ParameterFullName)
-{
-	FNiagaraParameterHandle ReadHandle(ParameterFullName);
-	ENiagaraCompileParameterReferenceFlags Flags = ENiagaraCompileParameterReferenceFlags::None;
-	if (ReadHandle.IsUserHandle())
-	{
-		Flags = ENiagaraCompileParameterReferenceFlags::User;
-	}
-	else if (ReadHandle.IsSystemHandle())
-	{
-		Flags = ENiagaraCompileParameterReferenceFlags::System;
-	}
-	else if (ReadHandle.IsEmitterHandle())
-	{
-		Flags = ENiagaraCompileParameterReferenceFlags::Emitter;
-	}
-	else if (ReadHandle.IsParticleAttributeHandle())
-	{
-		Flags = ENiagaraCompileParameterReferenceFlags::Particle;
-	}
-	else if (ReadHandle.IsTransientHandle())
-	{
-		Flags = ENiagaraCompileParameterReferenceFlags::Transient;
-	}
-	else if (ReadHandle.IsLocalHandle())
-	{
-		Flags = ENiagaraCompileParameterReferenceFlags::Local;
-	}
-	else if (ReadHandle.IsParameterCollectionHandle())
-	{
-		Flags = ENiagaraCompileParameterReferenceFlags::NPC;
-	}
-	return Flags;
-}
-
-int32 FNiagaraHlslTranslator::RegisterDataInterfaceForParameterRead(const FNiagaraVariableBase& ReadParameter, int32 ConnectedInfoIndex)
-{
-	FNiagaraCompileParameterReference ReadParameterReference(ReadParameter.GetName(), FGuid(), ExtractFlagsFromParameterName(ReadParameter.GetName()));
-	if (ReadParameterReference.HasFlags(ENiagaraCompileParameterReferenceFlags::Emitter))
-	{
-		ReadParameterReference.EmitterHandleId = CurrentDIOwnerContext.EmitterHandleId;
-	}
-
-	// If there is already a data interface that's a registered target for this parameter, we can use that.
-	int32 ExistingWriteIndex = CompilationOutput.ScriptData.DataInterfaceInfo.IndexOfByPredicate([&](const FNiagaraScriptDataInterfaceCompileInfo& OtherInfo)
-		{ return OtherInfo.TargetParameter == ReadParameterReference; });
-	if (ExistingWriteIndex != INDEX_NONE)
-	{
-		return ExistingWriteIndex;
-	}
-
-	// If there was no previous write, then this is an external read, so check if we have a previous external read.
-	int32 ExistingReadIndex = CompilationOutput.ScriptData.DataInterfaceInfo.IndexOfByPredicate([&](const FNiagaraScriptDataInterfaceCompileInfo& OtherInfo)
-		{ return OtherInfo.SourceExternalParameter == ReadParameterReference; });
-	if (ExistingReadIndex != INDEX_NONE)
-	{
-		return ExistingReadIndex;
-	}
-
-	// If we have no writes, and no previous external reads, check if we can use a connected default.
-	FNiagaraScriptDataInterfaceCompileInfo* ConnectedInfo = ConnectedInfoIndex != INDEX_NONE && ConnectedInfoIndex < CompilationOutput.ScriptData.DataInterfaceInfo.Num()
-		? &CompilationOutput.ScriptData.DataInterfaceInfo[ConnectedInfoIndex] : nullptr;
-	if (ConnectedInfo != nullptr)
-	{
-		if (ConnectedInfo->SourceExternalParameter.IsValid() == false)
-		{
-			ConnectedInfo->SourceExternalParameter = ReadParameterReference;
-			return ConnectedInfoIndex;
-		}
-		else
-		{
-			// If the connected default was already used with another external read, we register a duplicate placeholder so we can register this read.
-			int32 PlaceholderIndex = AddPlaceholderDuplicateDataInterface(ConnectedInfoIndex);
-			FNiagaraScriptDataInterfaceCompileInfo& PlaceholderDuplicate = CompilationOutput.ScriptData.DataInterfaceInfo[PlaceholderIndex];
-			PlaceholderDuplicate.SourceExternalParameter = ReadParameterReference;
-			PlaceholderDuplicate.SourceScriptDIIndex = ConnectedInfoIndex;
-		}
-	}
-
-
-	// No registered write, existing external read, or default, so we register a new external read.
-	UNiagaraDataInterface* CDO = GetDataInterfaceCDO(ReadParameter.GetType().GetClass());
-	int32 AddedIndex = AddDataInterfaceInfo(ReadParameter.GetType(), ReadParameter.GetName(), *CDO);
-	FNiagaraScriptDataInterfaceCompileInfo& AddedInfo = CompilationOutput.ScriptData.DataInterfaceInfo[AddedIndex];
-	AddedInfo.SourceExternalParameter = ReadParameterReference;
-	return AddedIndex;
-}
-
-void FNiagaraHlslTranslator::RegisterDataInterfaceForParameterWrite(const FNiagaraVariableBase& WriteParameter, int32 ConnectedInfoIndex)
-{
-	FNiagaraCompileParameterReference WriteParameterReference(WriteParameter.GetName(), FGuid(), ExtractFlagsFromParameterName(WriteParameter.GetName()));
-	if (WriteParameterReference.HasFlags(ENiagaraCompileParameterReferenceFlags::Emitter))
-	{
-		WriteParameterReference.EmitterHandleId = CurrentDIOwnerContext.EmitterHandleId;
-	}
-
-	// First check if this data interfaces was already written to, if it was already written to, validate that it's not being written to
-	// a second time with a different data interface.
-	int32 ExistingWriteIndex = CompilationOutput.ScriptData.DataInterfaceInfo.IndexOfByPredicate([&](const FNiagaraScriptDataInterfaceCompileInfo& OtherInfo)
-		{ return OtherInfo.TargetParameter == WriteParameterReference; });
-	if (ExistingWriteIndex != INDEX_NONE)
-	{
-		if (ConnectedInfoIndex != INDEX_NONE &&
-			ExistingWriteIndex != ConnectedInfoIndex &&
-			CompilationOutput.ScriptData.DataInterfaceInfo[ExistingWriteIndex].SourceScriptDIIndex != ConnectedInfoIndex)
-		{
-			Warning(FText::Format(LOCTEXT("MultipleDIWriteWarningFormat", "Data interface parameter {0} written to multiple times, only using the first write."),
-				FText::FromName(WriteParameterReference.ParameterFullName)));
-		}
-		return;
-	}
-
-	// Get the data interface info currently connected to this parameter if available.
-	FNiagaraScriptDataInterfaceCompileInfo* ConnectedInfo = ConnectedInfoIndex != INDEX_NONE && ConnectedInfoIndex < CompilationOutput.ScriptData.DataInterfaceInfo.Num()
-		? &CompilationOutput.ScriptData.DataInterfaceInfo[ConnectedInfoIndex] : nullptr;
-	if (ConnectedInfo != nullptr)
-	{
-		if (ConnectedInfo->TargetParameter.IsValid() == false)
-		{
-			// The parameter is connected to a registered data interface with no target, so we specify the write here.
-			ConnectedInfo->TargetParameter = WriteParameterReference;
-		}
-		else
-		{
-			// If the connected default was already written to another target, we register a duplicate placeholder so we can register this target write.
-			int32 PlaceholderIndex = AddPlaceholderDuplicateDataInterface(ConnectedInfoIndex);
-			FNiagaraScriptDataInterfaceCompileInfo& PlaceholderDuplicate = CompilationOutput.ScriptData.DataInterfaceInfo[PlaceholderIndex];
-			PlaceholderDuplicate.TargetParameter = WriteParameterReference;
-			PlaceholderDuplicate.SourceScriptDIIndex = ConnectedInfoIndex;
-		}
-		return;
-	}
-
-	// If the write was not connected and it's not already registered then we need to use a placeholder that will get populated
-	// with a default in the resolve step.
-	UNiagaraDataInterface* CDO = GetDataInterfaceCDO(WriteParameter.GetType().GetClass());
-	int32 AddedIndex = AddDataInterfaceInfo(WriteParameter.GetType(), WriteParameter.GetName(), *CDO);
-	FNiagaraScriptDataInterfaceCompileInfo& AddedInfo = CompilationOutput.ScriptData.DataInterfaceInfo[AddedIndex];
-	AddedInfo.TargetParameter = WriteParameterReference;
-}
-
-FNiagaraScriptDataInterfaceCompileId FNiagaraHlslTranslator::MakeUniqueDataInterfaceCompileId(FName BaseName) const
-{
-	int32 IdIndex = 0;
-	FNiagaraScriptDataInterfaceCompileId CompileId = FNiagaraScriptDataInterfaceCompileId(CurrentDIOwnerContext.EmitterHandleId,
-		CurrentDIOwnerContext.ScriptUsage, CurrentDIOwnerContext.ScriptUsageId, CurrentDIOwnerContext.OwnerName, BaseName, IdIndex);
-
-	while (CompilationOutput.ScriptData.DataInterfaceInfo.ContainsByPredicate([&CompileId](const FNiagaraScriptDataInterfaceCompileInfo& CompileInfo)
-		{ return CompileInfo.Id == CompileId; }))
-	{
-		IdIndex++;
-		CompileId = FNiagaraScriptDataInterfaceCompileId(CurrentDIOwnerContext.EmitterHandleId,
-			CurrentDIOwnerContext.ScriptUsage, CurrentDIOwnerContext.ScriptUsageId, CurrentDIOwnerContext.OwnerName, BaseName, IdIndex);
-	}
-
-	return CompileId;
-}
-
-FName FNiagaraHlslTranslator::MakeUniqueDIName(FName OwnerName, FName BaseName) const
-{
-	TSet<FName> CurrentNames;
-	for (const FNiagaraScriptDataInterfaceCompileInfo& DataInterfaceCompileInfo : CompilationOutput.ScriptData.DataInterfaceInfo)
-	{
-		CurrentNames.Add(DataInterfaceCompileInfo.Name);
-	}
-
-	FName CandidateName = *BuildSanitizedSymbolString(BaseName.ToString().Replace(TEXT("."), TEXT("_")).Replace(TEXT(" "), TEXT("_")));
-
-	return FNiagaraUtilities::GetUniqueName(CandidateName, CurrentNames);
-}
-
-int32 FNiagaraHlslTranslator::AddDataInterfaceInfo(const FNiagaraTypeDefinition& Type, FName BaseName, const UNiagaraDataInterface& DataInterface)
-{
-	FNiagaraScriptDataInterfaceCompileId UniqueId = MakeUniqueDataInterfaceCompileId(BaseName);
-	FName UniqueName = MakeUniqueDIName(UniqueId.GetOwnerName(), BaseName);
-
-	int32 Idx = CompilationOutput.ScriptData.DataInterfaceInfo.AddDefaulted();
-	FNiagaraScriptDataInterfaceCompileInfo& Info = CompilationOutput.ScriptData.DataInterfaceInfo[Idx];
-	Info.Id = UniqueId;
-	Info.Name = UniqueName;
-	Info.Type = Type;
-	Info.UserPtrIdx = DataInterface.PerInstanceDataSize() > 0 ? CompilationOutput.ScriptData.NumUserPtrs++ : INDEX_NONE;
-
-	if (DataInterface.NeedsGPUContextInit() && CompileOptions.IsGpuScript())
-	{
-		CompilationOutput.ScriptData.bNeedsGPUContextInit = true;
-	}
-
-	return Idx;
-}
-
-int32 FNiagaraHlslTranslator::AddPlaceholderDuplicateDataInterface(int32 IndexToDuplicate)
-{
-	FName BaseName = CompilationOutput.ScriptData.DataInterfaceInfo[IndexToDuplicate].Id.GetBaseName();
-	FNiagaraScriptDataInterfaceCompileId UniqueId = MakeUniqueDataInterfaceCompileId(BaseName);
-	FName UniqueName = MakeUniqueDIName(UniqueId.GetOwnerName(), BaseName);
-
-	int32 DuplicateIndex = CompilationOutput.ScriptData.DataInterfaceInfo.AddDefaulted();
-	FNiagaraScriptDataInterfaceCompileInfo& DuplicateInfo = CompilationOutput.ScriptData.DataInterfaceInfo[DuplicateIndex];
-	DuplicateInfo = CompilationOutput.ScriptData.DataInterfaceInfo[IndexToDuplicate];
-	DuplicateInfo.Id = UniqueId;
-	DuplicateInfo.Name = UniqueName;
-	return DuplicateIndex;
-}
-
 template<typename GraphBridge>
 void TNiagaraHlslTranslator<GraphBridge>::Convert(const FConvertNode* Convert, TArrayView<const int32> Inputs, TArray<int32>& Outputs)
 {
@@ -10910,12 +10772,6 @@ void TNiagaraHlslTranslator<GraphBridge>::TrimAttributes(const FNiagaraCompileOp
 	}
 }
 
-template<typename GraphBridge>
-FNiagaraGraphInputDataInterfaceKey TNiagaraHlslTranslator<GraphBridge>::GetGraphInputDataInterfaceKey(const FInputNode* InputNode) const
-{
-	return GraphBridge::GetGraphInputDataInterfaceKey(InputNode, CompileDuplicateData);
-}
-
 TUniquePtr<INiagaraHlslTranslator> INiagaraHlslTranslator::CreateTranslator(const FNiagaraCompileRequestDataBase* InCompileData, const FNiagaraCompileRequestDuplicateDataBase* InDuplicateData)
 {
 	// translate everything
diff --git a/Engine/Plugins/FX/Niagara/Source/NiagaraEditor/Private/NiagaraHlslTranslator.h b/Engine/Plugins/FX/Niagara/Source/NiagaraEditor/Private/NiagaraHlslTranslator.h
--- a/Engine/Plugins/FX/Niagara/Source/NiagaraEditor/Private/NiagaraHlslTranslator.h
+++ b/Engine/Plugins/FX/Niagara/Source/NiagaraEditor/Private/NiagaraHlslTranslator.h
@@ -203,7 +203,6 @@ public:
 	FString CustomReadFunction;
 	FString CustomWriteFunction;
 	int32 ParamMapHistoryIndex = INDEX_NONE;
-	FName DIOwnerName = NAME_None;
 
 	bool ShouldDoSpawnOnlyLogic() const;
 	template<typename GraphBridge>
@@ -407,6 +406,7 @@ public:
 
 	static NIAGARAEDITOR_API FString GetFunctionSignatureSymbol(const FNiagaraFunctionSignature& Sig);
 
+	static NIAGARAEDITOR_API FName GetDataInterfaceName(FName BaseName, const FString& UniqueEmitterName, bool bIsParameterMapDataInterface);
 	static NIAGARAEDITOR_API FString GenerateFunctionHlslPrototype(FStringView InVariableName, const FNiagaraFunctionSignature& FunctionSignature);
 	static NIAGARAEDITOR_API FString GenerateFunctionParametersHlsl(const FNiagaraFunctionSignature& FunctionSignature);
 
@@ -415,10 +415,6 @@ public:
 
 	NIAGARAEDITOR_API bool IsEventSpawnScript()const;
 
-	/** Registers an input node reference to a data interface and returns the data interface info index.  The info for this data interfaces may be newly created for this usage, or may be an existing
-		entry if this inputs input variable has been encountered before. */
-	int32 RegisterDataInterfaceForInput(const FNiagaraVariableBase& InputVariable, const FNiagaraGraphInputDataInterfaceKey& GraphInputDataInterfaceKey, const UNiagaraDataInterface* InputDataInterface);
-
 	// collection of functions that rely on the graph implementation and so must be relegated to the TNiagaraHlslTranslator
 	virtual void DefineDataSetVariableWrites(FString& HlslOutput, const FNiagaraDataSetID& Id, int32 DataSetIndex, const TArray<FNiagaraVariable>& WriteVars) = 0;
 	virtual void Message(FNiagaraCompileEventSeverity Severity, FText MessageText, FStringView ShortDescription = FStringView()) = 0;
@@ -485,23 +481,6 @@ protected:
 	void AppendFunctionHlslPrototypeString(FStringView InVariableName, const FNiagaraFunctionSignature& FunctionSignature, FStringBuilderBase& Builder) const;
 	FString BuildFunctionHlslPrototypeString(FStringView InVariableName, const FNiagaraFunctionSignature& FunctionSignature) const;
 
-	/** Extracts parameter reference flags from a parameter's full name. */
-	static ENiagaraCompileParameterReferenceFlags ExtractFlagsFromParameterName(FName ParameterFullName);
-
-	/** Registers a parameter read reference to a data interface and returns the data interface info index.  The info for this data interface may be newly created
-		for this usage, or may be an existing entry if the parameter was read from or written to previously in this translation. */
-	int32 RegisterDataInterfaceForParameterRead(const FNiagaraVariableBase& ReadParameter, int32 BoundIndex);
-
-	/** Registers a target parameter write reference for a data interface. */
-	void RegisterDataInterfaceForParameterWrite(const FNiagaraVariableBase& WriteParameter, int32 BoundIndex);
-
-	/** Creates a unique data interface compile id using the current context and the base name. */
-	FNiagaraScriptDataInterfaceCompileId MakeUniqueDataInterfaceCompileId(FName BaseName) const;
-
-	FName MakeUniqueDIName(FName OwnerName, FName BaseName) const;
-	int32 AddDataInterfaceInfo(const FNiagaraTypeDefinition& Type, FName BaseName, const UNiagaraDataInterface& DataInterface);
-	int32 AddPlaceholderDuplicateDataInterface(int32 IndexToDuplicate);
-
 	template<typename T>
 	void BuildConstantBuffer(ENiagaraCodeChunkMode ChunkMode);
 
@@ -611,26 +590,6 @@ protected:
 
 	// local copy of UNiagaraEditorSettings::GetHLSLKeywordReplacmentsMap() for the duration of the translation process
 	TMap<FString, FString> KeywordReplacementMap;
-
-	struct FDIOwnerContext
-	{
-		FDIOwnerContext() { }
-
-		FDIOwnerContext(const FGuid& InEmitterHandleId, ENiagaraScriptUsage InScriptUsage, const FGuid& InScriptUsageId, FName InOwnerName)
-			: EmitterHandleId(InEmitterHandleId)
-		    , ScriptUsage(InScriptUsage)
-			, ScriptUsageId(InScriptUsageId)
-			, OwnerName(InOwnerName)
-		{
-		}
-
-		FGuid EmitterHandleId;
-		ENiagaraScriptUsage ScriptUsage;
-		FGuid ScriptUsageId;
-		FName OwnerName;
-	};
-
-	FDIOwnerContext CurrentDIOwnerContext;
 };
 
 // all data and function implementations that rely on specifics of the graph implementation
@@ -702,6 +661,7 @@ public:
 
 	int32 RegisterUObject(const FNiagaraVariable& Variable, UObject* Object, bool bAddParameterMapRead);
 	int32 RegisterUObjectPath(const FNiagaraVariable& Variable, const FSoftObjectPath& ObjectPath, bool bAddParmeterMapRead);
+	int32 RegisterDataInterface(const FNiagaraVariable& Var, const UNiagaraDataInterface* DataInterface, bool bPlaceholder, bool bAddParameterMapRead);
 
 	void Operation(const FOpNode* Operation, TArray<int32>& Inputs, TArray<int32>& Outputs);
 	void Output(const FOutputNode* OutputNode, const TArray<int32>& ComputedInputs);
@@ -757,8 +717,6 @@ public:
 
 	void TrimAttributes(const FNiagaraCompileOptions& InCompileOptions, TArray<FNiagaraVariable>& Attributes);
 
-	FNiagaraGraphInputDataInterfaceKey GetGraphInputDataInterfaceKey(const FInputNode* InputNode) const;
-
 protected:
 	void InitializeParameterMapDefaults(int32 ParamMapHistoryIdx);
 	void HandleParameterRead(int32 ParamMapHistoryIdx, const FNiagaraVariable& Var, const FInputPin* DefaultPin, const FNode* ErrorNode, int32& OutputChunkId, TOptional<ENiagaraDefaultMode> DefaultMode = TOptional<ENiagaraDefaultMode>(), TOptional<FNiagaraScriptVariableBinding> DefaultBinding = TOptional<FNiagaraScriptVariableBinding>(), bool bTreatAsUnknownParameterMap = false, bool bIgnoreDefaultSetFirst = false);
@@ -773,7 +731,7 @@ protected:
 	void ProcessCustomHlsl(const FString& InCustomHlsl, ENiagaraScriptUsage InUsage, const FNiagaraFunctionSignature& InSignature, const TArray<int32>& Inputs, const FNode* NodeForErrorReporting, FString& OutCustomHlsl, FNiagaraFunctionSignature& OutSignature);
 	void HandleSimStageSetupAndTeardown(int32 InWhichStage, FString& OutHlsl);
 	// Register a System/Engine/read-only variable in its namespaced form
-	bool ParameterMapRegisterExternalConstantNamespaceVariable(FNiagaraVariable InVariable, const FNiagaraVariableBase& InUnresolvedVariable, const FNode* InNode, int32 InParamMapHistoryIdx, int32& Output, const FPin* InDefaultPin);
+	bool ParameterMapRegisterExternalConstantNamespaceVariable(FNiagaraVariable InVariable, const FNode* InNode, int32 InParamMapHistoryIdx, int32& Output, const FPin* InDefaultPin);
 
 	// Register an attribute in its non namespaced form
 	bool ParameterMapRegisterNamespaceAttributeVariable(const FNiagaraVariable& InVariable, const FNode* InNode, int32 InParamMapHistoryIdx, int32& Output);
diff --git a/Engine/Plugins/FX/Niagara/Source/NiagaraEditor/Private/NiagaraNodeInput.cpp b/Engine/Plugins/FX/Niagara/Source/NiagaraEditor/Private/NiagaraNodeInput.cpp
--- a/Engine/Plugins/FX/Niagara/Source/NiagaraEditor/Private/NiagaraNodeInput.cpp
+++ b/Engine/Plugins/FX/Niagara/Source/NiagaraEditor/Private/NiagaraNodeInput.cpp
@@ -415,8 +415,11 @@ void UNiagaraNodeInput::Compile(FTranslator* Translator, TArray<int32>& Outputs)
 				{
 					if ( IsDataInterface() )
 					{
-						Outputs.Add(Translator->RegisterDataInterfaceForInput(Input, Translator->GetGraphInputDataInterfaceKey(this), DataInterface));
-						return;
+						if (ensure(DataInterface != nullptr))
+						{
+							Outputs.Add(Translator->RegisterDataInterface(Input, DataInterface, false, false));
+							return;
+						}
 					}
 					else if ( IsObjectAsset() )
 					{
@@ -436,7 +439,7 @@ void UNiagaraNodeInput::Compile(FTranslator* Translator, TArray<int32>& Outputs)
 	case ENiagaraInputNodeUsage::Parameter:
 		if (IsDataInterface())
 		{
-			Outputs.Add(Translator->RegisterDataInterfaceForInput(Input, Translator->GetGraphInputDataInterfaceKey(this), DataInterface));
+			Outputs.Add(Translator->RegisterDataInterface(Input, DataInterface, false, false));
 			break;
 		}
 		else if ( IsObjectAsset() )
diff --git a/Engine/Plugins/FX/Niagara/Source/NiagaraEditor/Private/NiagaraScriptSource.cpp b/Engine/Plugins/FX/Niagara/Source/NiagaraEditor/Private/NiagaraScriptSource.cpp
--- a/Engine/Plugins/FX/Niagara/Source/NiagaraEditor/Private/NiagaraScriptSource.cpp
+++ b/Engine/Plugins/FX/Niagara/Source/NiagaraEditor/Private/NiagaraScriptSource.cpp
@@ -259,7 +259,7 @@ void UNiagaraScriptSource::PostLoadFromEmitter(FVersionedNiagaraEmitter OwningEm
 	}
 }
 
-void CollectGraphInputDataInterfacesRecursive(UNiagaraGraph* InGraph, ENiagaraScriptUsage Usage, FGuid UsageId, FString EmitterUniqueName, TSet<UNiagaraGraph*>& VisitedGraphs, TMap<FNiagaraGraphInputDataInterfaceKey, UNiagaraDataInterface*>& Result)
+void FindObjectNamesRecursive(UNiagaraGraph* InGraph, ENiagaraScriptUsage Usage, FGuid UsageId, FString EmitterUniqueName, TSet<UNiagaraGraph*>& VisitedGraphs, TMap<FName, UNiagaraDataInterface*>& Result)
 {
 	if (!InGraph)
 	{
@@ -278,12 +278,13 @@ void CollectGraphInputDataInterfacesRecursive(UNiagaraGraph* InGraph, ENiagaraSc
 				if (InputNode->Input.IsDataInterface())
 				{
 					UNiagaraDataInterface* DataInterface = InputNode->GetDataInterface();
-					FNiagaraGraphInputDataInterfaceKey GraphInputDataInterfaceKey(*InGraph->GetPathName(), InputNode->Input);
-					if (Result.Contains(GraphInputDataInterfaceKey) && Result[GraphInputDataInterfaceKey] != DataInterface)
+					bool bIsParameterMapDataInterface = false;
+					FName DIName = FNiagaraHlslTranslator::GetDataInterfaceName(InputNode->Input.GetName(), EmitterUniqueName, bIsParameterMapDataInterface);
+					if (Result.Contains(DIName) && Result[DIName] != DataInterface)
 					{
-						UE_LOG(LogNiagaraEditor, Verbose, TEXT("Duplicate data interface name %s in graph %s. One of the data interfaces will override the other when resolving the name."), *InputNode->Input.ToString(), *InGraph->GetPathName());
+						UE_LOG(LogNiagaraEditor, Verbose, TEXT("Duplicate data interface name %s in graph %s. One of the data interfaces will override the other when resolving the name."), *DIName.ToString(), *InGraph->GetPathName());
 					}
-					Result.Add(GraphInputDataInterfaceKey, DataInterface);
+					Result.Add(DIName, DataInterface);
 				}
 				continue;
 			}
@@ -294,16 +295,16 @@ void CollectGraphInputDataInterfacesRecursive(UNiagaraGraph* InGraph, ENiagaraSc
 				if(VisitedGraphs.Contains(FunctionGraph) == false)
 				{
 					VisitedGraphs.Add(FunctionGraph);
-					CollectGraphInputDataInterfacesRecursive(FunctionGraph, FunctionCallNode->GetCalledUsage(), FGuid(), EmitterUniqueName, VisitedGraphs, Result);
+					FindObjectNamesRecursive(FunctionGraph, FunctionCallNode->GetCalledUsage(), FGuid(), EmitterUniqueName, VisitedGraphs, Result);
 				}
 			}
 		}
 	}
 }
 
-TMap<FNiagaraGraphInputDataInterfaceKey, UNiagaraDataInterface*> UNiagaraScriptSource::CollectGraphInputDataInterfaces(UNiagaraSystem& System, ENiagaraScriptUsage Usage, FGuid UsageId, FString EmitterUniqueName) const
+TMap<FName, UNiagaraDataInterface*> UNiagaraScriptSource::ComputeObjectNameMap(UNiagaraSystem& System, ENiagaraScriptUsage Usage, FGuid UsageId, FString EmitterUniqueName) const
 {
-	TMap<FNiagaraGraphInputDataInterfaceKey, UNiagaraDataInterface*> Result;
+	TMap<FName, UNiagaraDataInterface*> Result;
 	if (!NodeGraph)
 	{
 		return Result;
@@ -322,19 +323,19 @@ TMap<FNiagaraGraphInputDataInterfaceKey, UNiagaraDataInterface*> UNiagaraScriptS
 				OutputNode->GetUsage() == ENiagaraScriptUsage::ParticleUpdateScript ||
 				OutputNode->GetUsage() == ENiagaraScriptUsage::ParticleSimulationStageScript)
 			{
-				CollectGraphInputDataInterfacesRecursive(NodeGraph, OutputNode->GetUsage(), OutputNode->GetUsageId(), EmitterUniqueName, VisitedGraphs, Result);
+				FindObjectNamesRecursive(NodeGraph, OutputNode->GetUsage(), OutputNode->GetUsageId(), EmitterUniqueName, VisitedGraphs, Result);
 			}
 		}
 	}
 	else
 	{
 		// Collect the data interfaces for the current script.
-		CollectGraphInputDataInterfacesRecursive(NodeGraph, Usage, UsageId, EmitterUniqueName, VisitedGraphs, Result);
+		FindObjectNamesRecursive(NodeGraph, Usage, UsageId, EmitterUniqueName, VisitedGraphs, Result);
 
 		if (Usage == ENiagaraScriptUsage::ParticleSpawnScriptInterpolated)
 		{
 			// The interpolated spawn script needs to include the particle update script's data interfaces as well.
-			CollectGraphInputDataInterfacesRecursive(NodeGraph, ENiagaraScriptUsage::ParticleUpdateScript, FGuid(), EmitterUniqueName, VisitedGraphs, Result);
+			FindObjectNamesRecursive(NodeGraph, ENiagaraScriptUsage::ParticleUpdateScript, FGuid(), EmitterUniqueName, VisitedGraphs, Result);
 		}
 		else if (Usage == ENiagaraScriptUsage::SystemSpawnScript || Usage == ENiagaraScriptUsage::SystemUpdateScript)
 		{
@@ -350,7 +351,7 @@ TMap<FNiagaraGraphInputDataInterfaceKey, UNiagaraDataInterface*> UNiagaraScriptS
 				{
 					if (UNiagaraScriptSource* EmitterSource = Cast<UNiagaraScriptSource>(EmitterData->GraphSource); EmitterSource != nullptr)
 					{
-						CollectGraphInputDataInterfacesRecursive(EmitterSource->NodeGraph, EmitterUsage, FGuid(), EmitterHandle.GetUniqueInstanceName(), VisitedGraphs, Result);
+						FindObjectNamesRecursive(EmitterSource->NodeGraph, EmitterUsage, FGuid(), EmitterHandle.GetUniqueInstanceName(), VisitedGraphs, Result);
 					}
 				}
 			}
diff --git a/Engine/Plugins/FX/Niagara/Source/NiagaraEditor/Private/NiagaraValidationRules.cpp b/Engine/Plugins/FX/Niagara/Source/NiagaraEditor/Private/NiagaraValidationRules.cpp
--- a/Engine/Plugins/FX/Niagara/Source/NiagaraEditor/Private/NiagaraValidationRules.cpp
+++ b/Engine/Plugins/FX/Niagara/Source/NiagaraEditor/Private/NiagaraValidationRules.cpp
@@ -1452,11 +1452,11 @@ void UNiagaraValidationRule_TickDependencyCheck::CheckValidity(const FNiagaraVal
 
 				// Generate issue
 				UObject* StackObject = nullptr;
-				if (ResolvedDI.ResolvedParameterReference.EmitterHandleId.IsValid())
+				if (ResolvedDI.ResolvedSourceEmitterName.Len() > 0)
 				{
 					for (const TSharedRef<FNiagaraEmitterHandleViewModel>& EmitterViewModel : Context.ViewModel->GetEmitterHandleViewModels())
 					{
-						if (EmitterViewModel->GetId() == ResolvedDI.ResolvedParameterReference.EmitterHandleId)
+						if (EmitterViewModel->GetName() == FName(ResolvedDI.ResolvedSourceEmitterName))
 						{
 							StackObject = NiagaraValidation::GetStackEntry<UNiagaraStackEmitterPropertiesItem>(EmitterViewModel->GetEmitterStackViewModel());
 							break;
@@ -1474,7 +1474,7 @@ void UNiagaraValidationRule_TickDependencyCheck::CheckValidity(const FNiagaraVal
 				}
 
 				const FText DIClassText = FText::FromName(RuntimeDI->GetClass()->GetFName());
-				const FText DIVariableText = FText::FromName(ResolvedDI.ParameterStoreVariable.GetName());
+				const FText DIVariableText = FText::FromName(ResolvedDI.Name);
 				FNiagaraValidationResult& ValiationResult = OutResults.Emplace_GetRef(
 					Severity,
 					LOCTEXT("TickDependencyCheckFormat", "Performance issue due to late ticking which may cause waits on the game thread."),
diff --git a/Engine/Plugins/FX/Niagara/Source/NiagaraEditor/Private/ViewModels/NiagaraScriptInputCollectionViewModel.cpp b/Engine/Plugins/FX/Niagara/Source/NiagaraEditor/Private/ViewModels/NiagaraScriptInputCollectionViewModel.cpp
--- a/Engine/Plugins/FX/Niagara/Source/NiagaraEditor/Private/ViewModels/NiagaraScriptInputCollectionViewModel.cpp
+++ b/Engine/Plugins/FX/Niagara/Source/NiagaraEditor/Private/ViewModels/NiagaraScriptInputCollectionViewModel.cpp
@@ -280,6 +280,22 @@ void FNiagaraScriptInputCollectionViewModel::RefreshParameterViewModels()
 			else if (InputNode->Input.IsDataInterface())
 			{
 				UNiagaraDataInterface* EmitterDataInterface = InputNode->GetDataInterface();
+				for (FVersionedNiagaraScriptWeakPtr& ScriptWeakPtr : Scripts)
+				{
+					if (EmitterDataInterface != nullptr)
+					{
+						break;
+					}
+
+					for (FNiagaraScriptDataInterfaceInfo& DataInterfaceInfoItem : ScriptWeakPtr.Script->GetCachedDefaultDataInterfaces())
+					{
+						if (DataInterfaceInfoItem.Name == InputNode->Input.GetName())
+						{
+							EmitterDataInterface = DataInterfaceInfoItem.DataInterface;
+							break;
+						}
+					}
+				}
 				ParameterViewModel = MakeShareable(new FNiagaraScriptParameterViewModel(GraphVariable, *InputNode, EmitterDataInterface, ParameterEditMode));
 			}
 			else if (InputNode->Input.IsUObject())
@@ -512,6 +528,23 @@ void FNiagaraScriptInputCollectionViewModel::OnParameterValueChangedInternal(TSh
 						DataInterface->CopyTo(DataInterfaceToUpdate);
 					}
 				}
+
+				// Synchronize script variables...
+				for (FVersionedNiagaraScriptWeakPtr& VersionedScript : Scripts)
+				{
+					if (!VersionedScript.Script.IsValid() || !VersionedScript.Script->GetVMExecutableData().IsValid())
+					{
+						continue;
+					}
+					for (FNiagaraScriptDataInterfaceInfo& Info : VersionedScript.Script->GetCachedDefaultDataInterfaces())
+					{
+						if (Info.Name == ChangedParameter->GetName() && DataInterface != Info.DataInterface)
+						{
+							DataInterface->CopyTo(Info.DataInterface);
+							break;
+						}
+					}
+				}
 			}
 		}
 		else if (ChangedParameter->GetDefaultValueType() == INiagaraParameterViewModel::EDefaultValueType::ObjectAsset)
diff --git a/Engine/Plugins/FX/Niagara/Source/NiagaraEditor/Public/NiagaraCompilationBridge.h b/Engine/Plugins/FX/Niagara/Source/NiagaraEditor/Public/NiagaraCompilationBridge.h
--- a/Engine/Plugins/FX/Niagara/Source/NiagaraEditor/Public/NiagaraCompilationBridge.h
+++ b/Engine/Plugins/FX/Niagara/Source/NiagaraEditor/Public/NiagaraCompilationBridge.h
@@ -167,7 +167,6 @@ struct FNiagaraCompilationGraphBridge
 	static ENiagaraScriptUsage GetEmitterUsage(const FEmitterNode* EmitterNode);
 	static FString GetEmitterName(const FEmitterNode* EmitterNode);
 	static FString GetEmitterPathName(const FEmitterNode* EmitterNode);
-	static FGuid GetEmitterHandleId(const FEmitterNode* EmitterNode);
 	static FString GetEmitterHandleIdString(const FEmitterNode* EmitterNode);
 	static const FGraph* GetFunctionNodeGraph(const FFunctionCallNode* FunctionCall);
 	static FString GetFunctionFullName(const FFunctionCallNode* FunctionCall);
@@ -189,7 +188,6 @@ struct FNiagaraCompilationGraphBridge
 	static FNiagaraVariable GetPinVariable(const FPin* Pin, bool bNeedsValue, ENiagaraStructConversion Conversion);
 	static const FInputPin* GetPinAsInput(const FPin* Pin);
 	static FNiagaraVariable GetInputVariable(const FInputNode* InputNode);
-	static FNiagaraGraphInputDataInterfaceKey GetGraphInputDataInterfaceKey(const FInputNode* InputNode, const FCompilationCopy* CompilationCopy);
 	static const TArray<FNiagaraVariable>& GetOutputVariables(const FOutputNode* OutputNode);
 	static TArray<FNiagaraVariable> GetGraphOutputNodeVariables(const FGraph* Graph, ENiagaraScriptUsage Usage);
 	static TArray<const FInputNode*> GetGraphInputNodes(const FGraph* Graph, const FNiagaraFindInputNodeOptions& Options);
@@ -301,7 +299,6 @@ struct FNiagaraCompilationDigestBridge
 	static ENiagaraScriptUsage GetEmitterUsage(const FEmitterNode* EmitterNode);
 	static FString GetEmitterName(const FEmitterNode* EmitterNode);
 	static FString GetEmitterPathName(const FEmitterNode* EmitterNode);
-	static FGuid GetEmitterHandleId(const FEmitterNode* EmitterNode);
 	static FString GetEmitterHandleIdString(const FEmitterNode* EmitterNode);
 	static const FGraph* GetFunctionNodeGraph(const FFunctionCallNode* FunctionCall);
 	static FString GetFunctionFullName(const FFunctionCallNode* FunctionCall);
@@ -323,7 +320,6 @@ struct FNiagaraCompilationDigestBridge
 	static FNiagaraVariable GetPinVariable(const FPin* Pin, bool bNeedsValue, ENiagaraStructConversion Conversion);
 	static const FInputPin* GetPinAsInput(const FPin* Pin);
 	static FNiagaraVariable GetInputVariable(const FInputNode* InputNode);
-	static FNiagaraGraphInputDataInterfaceKey GetGraphInputDataInterfaceKey(const FInputNode* InputNode, const FCompilationCopy* CompilationCopy);
 	static const TArray<FNiagaraVariable>& GetOutputVariables(const FOutputNode* OutputNode);
 	static TArray<FNiagaraVariable> GetGraphOutputNodeVariables(const FGraph* Graph, ENiagaraScriptUsage Usage);
 	static TArray<const FInputNode*> GetGraphInputNodes(const FGraph* Graph, const FNiagaraFindInputNodeOptions& Options);
diff --git a/Engine/Plugins/FX/Niagara/Source/NiagaraEditor/Public/NiagaraScriptSource.h b/Engine/Plugins/FX/Niagara/Source/NiagaraEditor/Public/NiagaraScriptSource.h
--- a/Engine/Plugins/FX/Niagara/Source/NiagaraEditor/Public/NiagaraScriptSource.h
+++ b/Engine/Plugins/FX/Niagara/Source/NiagaraEditor/Public/NiagaraScriptSource.h
@@ -56,7 +56,7 @@ public:
 
 	virtual void PostLoadFromEmitter(FVersionedNiagaraEmitter OwningEmitter) override;
 
-	virtual TMap<FNiagaraGraphInputDataInterfaceKey, UNiagaraDataInterface*> CollectGraphInputDataInterfaces(UNiagaraSystem& System, ENiagaraScriptUsage Usage, FGuid UsageId, FString EmitterUniqueName) const override;
+	virtual TMap<FName, UNiagaraDataInterface*> ComputeObjectNameMap(UNiagaraSystem& System, ENiagaraScriptUsage Usage, FGuid UsageId, FString EmitterUniqueName) const override;
 
 	NIAGARAEDITOR_API virtual bool AddModuleIfMissing(FString ModulePath, ENiagaraScriptUsage Usage, bool& bOutFoundModule)override;
 
diff --git a/Engine/Source/Runtime/Core/Private/UObject/FortniteMainBranchObjectVersion.cpp b/Engine/Source/Runtime/Core/Private/UObject/FortniteMainBranchObjectVersion.cpp
--- a/Engine/Source/Runtime/Core/Private/UObject/FortniteMainBranchObjectVersion.cpp
+++ b/Engine/Source/Runtime/Core/Private/UObject/FortniteMainBranchObjectVersion.cpp
@@ -12,7 +12,7 @@ TMap<FGuid, FGuid> FFortniteMainBranchObjectVersion::GetSystemGuids()
 	SystemGuids.Add(DevGuids.LANDSCAPE_MOBILE_COOK_VERSION, FGuid("32D02EF867C74B71A0D4E0FA41392732"));
 	SystemGuids.Add(DevGuids.MATERIALSHADERMAP_DERIVEDDATA_VER, FGuid("4DCE94C9E9994B3DA256F6A015B87CA9"));
 	SystemGuids.Add(DevGuids.NIAGARASHADERMAP_DERIVEDDATA_VER, FGuid("CC229D919A214E59B79D01AC2E5BF981"));
-	SystemGuids.Add(DevGuids.Niagara_LatestScriptCompileVersion, FGuid("C19F5C3F8A704C4A83F3077F3236AE2B"));
+	SystemGuids.Add(DevGuids.Niagara_LatestScriptCompileVersion, FGuid("FB15C2870A4E4AE8B5F58566F0E3825E"));
 	SystemGuids.Add(DevGuids.MaterialTranslationDDCVersion, FGuid("FD43494F1D8443D79A06E6E28F503A42"));
 
 	// These GUIDs are in AppendRestrictedFortniteMainStreamObjectVersionGuids as they should rarely be changed.
