- Remove check for mismatched icon / tooltips on renderers and replace with log

#rnx

diff --git a/Engine/Plugins/FX/Niagara/Source/NiagaraEditorWidgets/Private/SNiagaraEmitterRendererThumbnails.cpp b/Engine/Plugins/FX/Niagara/Source/NiagaraEditorWidgets/Private/SNiagaraEmitterRendererThumbnails.cpp
--- a/Engine/Plugins/FX/Niagara/Source/NiagaraEditorWidgets/Private/SNiagaraEmitterRendererThumbnails.cpp
+++ b/Engine/Plugins/FX/Niagara/Source/NiagaraEditorWidgets/Private/SNiagaraEmitterRendererThumbnails.cpp
@@ -145,19 +145,23 @@ void SNiagaraEmitterRendererThumbnails::RefreshThumbnails()
 		RendererProperties->GetRendererWidgets(InInstance, Widgets, UThumbnailManager::Get().GetSharedThumbnailPool());
 		RendererProperties->GetRendererTooltipWidgets(InInstance, TooltipWidgets, UThumbnailManager::Get().GetSharedThumbnailPool());
 
+		if (TooltipWidgets.Num() < Widgets.Num())
+		{
+			UE_LOG(LogNiagaraEditor, Warning, TEXT("Renderer(%s) has a mismatch of widgets(%d) to tooltips(%d)"), *GetFullNameSafe(RendererProperties), Widgets.Num(), TooltipWidgets.Num());
+		}
+
 		// Fallback widget & tooltip
 		if(Widgets.Num() == 0)
 		{
 			RendererProperties->CreateDefaultRendererWidget(Widgets);
 		}
 
-		if(TooltipWidgets.Num() == 0)
+		while ( TooltipWidgets.Num() < Widgets.Num() )
 		{
 			FText BaseText = LOCTEXT("DefaultRendererTooltip", "{0}");
 			TooltipWidgets.Add(SNew(STextBlock).Text(FText::FormatOrdered(BaseText, RendererProperties->GetClass()->GetDisplayNameText())));
 		}
 			
-		check(Widgets.Num() == TooltipWidgets.Num());
 		for (int32 WidgetIndex = 0; WidgetIndex < Widgets.Num(); WidgetIndex++)
 		{
 			SetupRendererThumbnailTooltip(Widgets[WidgetIndex], TooltipWidgets[WidgetIndex]);
