- Initialize properties on FNiagaraUserParameterInfo

#rnx

diff --git a/Engine/Plugins/FX/Niagara/Source/Niagara/Public/NiagaraFunctionLibrary.h b/Engine/Plugins/FX/Niagara/Source/Niagara/Public/NiagaraFunctionLibrary.h
--- a/Engine/Plugins/FX/Niagara/Source/Niagara/Public/NiagaraFunctionLibrary.h
+++ b/Engine/Plugins/FX/Niagara/Source/Niagara/Public/NiagaraFunctionLibrary.h
@@ -36,10 +36,10 @@ struct FNiagaraUserParameterInfo
 	FName TypeName;
 
 	UPROPERTY(EditAnywhere, BlueprintReadOnly, Category=Niagara)
-	int32 ParameterSizeBytes;
+	int32 ParameterSizeBytes = 0;
 
 	UPROPERTY(EditAnywhere, BlueprintReadOnly, Category=Niagara)
-	bool bIsEnumType;
+	bool bIsEnumType = false;
 };
 
 USTRUCT(MinimalAPI, BlueprintType)
