Fix warning V591: Non-void function should return a value.
#rnx

diff --git a/Engine/Plugins/FX/Niagara/Source/NiagaraEditor/Internal/NiagaraStackEntryEnumerable.h b/Engine/Plugins/FX/Niagara/Source/NiagaraEditor/Internal/NiagaraStackEntryEnumerable.h
--- a/Engine/Plugins/FX/Niagara/Source/NiagaraEditor/Internal/NiagaraStackEntryEnumerable.h
+++ b/Engine/Plugins/FX/Niagara/Source/NiagaraEditor/Internal/NiagaraStackEntryEnumerable.h
@@ -286,6 +286,7 @@ public:
 	{
 		TArray<TNiagaraStackEntry*> Array;
 		Array.Append(*this);
+		return Array;
 	}
 
 private:
