Fixed pseudo asset package name for asset filtering

diff --git a/Engine/Plugins/FX/Niagara/Source/NiagaraEditor/Private/NiagaraEmitterFactoryNew.cpp b/Engine/Plugins/FX/Niagara/Source/NiagaraEditor/Private/NiagaraEmitterFactoryNew.cpp
--- a/Engine/Plugins/FX/Niagara/Source/NiagaraEditor/Private/NiagaraEmitterFactoryNew.cpp
+++ b/Engine/Plugins/FX/Niagara/Source/NiagaraEditor/Private/NiagaraEmitterFactoryNew.cpp
@@ -61,6 +61,16 @@ bool UNiagaraEmitterFactoryNew::ConfigureProperties()
 		OverrideProfiles.AssetViewOptionsProfileName = FName("TaggedAssetBrowser");
 		OverrideProfiles.DefaultFilterMenuExpansion = EAssetTypeCategories::FX;
 		OverrideProfiles.FilterBarSaveName = FName("TaggedAssetBrowser.NiagaraEmitterFactory");
+
+		// We create a pseudo asset based on the path we currently want to create an asset at.
+		// This is then used for asset reference filtering to only show assets that this pseudo asset/path would be compatible with
+		
+		// @TODO (ME): This is not ideal. Ideally the factory would have access to creation path and/or property handle (if created from property asset picker -> create menu)
+		// then we wouldn't need to use GetLastDirectory which is relatively unsafe
+		FString PackagePath = FEditorDirectories::Get().GetLastDirectory(ELastDirectory::NEW_ASSET);
+		FString AssetName = "TmpAsset";
+		FName PackageName = FName(PackagePath + "/" + AssetName);
+		FAssetData PseudoReferencingAsset(PackageName, FName(PackagePath), FName(AssetName), UObject::StaticClass()->GetClassPathName());
 		
 		STaggedAssetBrowser::FArguments AssetBrowserArgs;
 		AssetBrowserArgs
@@ -71,9 +81,7 @@ bool UNiagaraEmitterFactoryNew::ConfigureProperties()
 			.OnAdditionalShouldFilterAsset_UObject(this, &UNiagaraEmitterFactoryNew::OnAdditionalShouldFilterAsset)
 			.OnGetExtraFrontendFilters_UObject(this, &UNiagaraEmitterFactoryNew::OnGetExtraFrontendFilters)
 			.OnExtendAddFilterMenu_UObject(this, &UNiagaraEmitterFactoryNew::OnExtendAddFilterMenu)
-			// @TODO (ME): This is not ideal. Ideally the factory would have access to creation path and/or property handle (if created from property asset picker -> create menu)
-			// then we wouldn't need to use GetLastDirectory which is relatively unsafe
-			.AdditionalReferencingAssets({ FAssetData("TmpAsset", FName(FEditorDirectories::Get().GetLastDirectory(ELastDirectory::NEW_ASSET)), "TmpAsset", UObject::StaticClass()->GetClassPathName())});
+			.AdditionalReferencingAssets({ PseudoReferencingAsset });
 
 		SWindow::FArguments WindowArgs;
 		WindowArgs.Title(LOCTEXT("EmitterAssetBrowserWindowTitle", "Create Niagara Emitter - Select an emitter as a base"));
diff --git a/Engine/Plugins/FX/Niagara/Source/NiagaraEditor/Private/NiagaraSystemFactoryNew.cpp b/Engine/Plugins/FX/Niagara/Source/NiagaraEditor/Private/NiagaraSystemFactoryNew.cpp
--- a/Engine/Plugins/FX/Niagara/Source/NiagaraEditor/Private/NiagaraSystemFactoryNew.cpp
+++ b/Engine/Plugins/FX/Niagara/Source/NiagaraEditor/Private/NiagaraSystemFactoryNew.cpp
@@ -57,6 +57,16 @@ bool UNiagaraSystemFactoryNew::ConfigureProperties()
 		FDefaultDetailsTabConfiguration DefaultDetailsTabConfiguration;
 		DefaultDetailsTabConfiguration.bUseDefaultDetailsTab = true;
 		DefaultDetailsTabConfiguration.EmptySelectionMessage = LOCTEXT("EmptySystemFactorySelectionMessage", "Select an emitter or system as a starting point for your new system.\nA system consists of one or more emitters.");
+
+		// We create a pseudo asset based on the path we currently want to create an asset at.
+		// This is then used for asset reference filtering to only show assets that this pseudo asset/path would be compatible with
+		
+		// @TODO (ME): This is not ideal. Ideally the factory would have access to creation path and/or property handle (if created from property asset picker -> create menu)
+		// then we wouldn't need to use GetLastDirectory which is relatively unsafe
+		FString PackagePath = FEditorDirectories::Get().GetLastDirectory(ELastDirectory::NEW_ASSET);
+		FString AssetName = "TmpAsset";
+		FName PackageName = FName(PackagePath + "/" + AssetName);
+		FAssetData PseudoReferencingAsset(PackageName, FName(PackagePath), FName(AssetName), UObject::StaticClass()->GetClassPathName());
 		
 		STaggedAssetBrowser::FArguments AssetBrowserArgs;
 		AssetBrowserArgs
@@ -68,7 +78,7 @@ bool UNiagaraSystemFactoryNew::ConfigureProperties()
 			.OnExtendAddFilterMenu_UObject(this, &UNiagaraSystemFactoryNew::OnExtendAddFilterMenu)
 			// @TODO (ME): This is not ideal. Ideally the factory would have access to creation path and/or property handle (if created from property asset picker -> create menu)
 			// then we wouldn't need to use GetLastDirectory which is relatively unsafe
-			.AdditionalReferencingAssets({ FAssetData("TmpAsset", FName(FEditorDirectories::Get().GetLastDirectory(ELastDirectory::NEW_ASSET)), "TmpAsset", UObject::StaticClass()->GetClassPathName())});
+			.AdditionalReferencingAssets({ PseudoReferencingAsset });
 
 		SWindow::FArguments WindowArgs;
 		WindowArgs.Title(LOCTEXT("SystemAssetBrowserWindowTitle", "Create Niagara System - Select an emitter or system as a base"));
