- Add debug draw icon to stack row for LW to mirror renderers / regular modules

#tests Use debug draw icon on stack row

diff --git a/Engine/Plugins/FX/Niagara/Source/NiagaraEditor/Private/ViewModels/Stack/NiagaraStackStatelessEmitterSimulateGroup.cpp b/Engine/Plugins/FX/Niagara/Source/NiagaraEditor/Private/ViewModels/Stack/NiagaraStackStatelessEmitterSimulateGroup.cpp
--- a/Engine/Plugins/FX/Niagara/Source/NiagaraEditor/Private/ViewModels/Stack/NiagaraStackStatelessEmitterSimulateGroup.cpp
+++ b/Engine/Plugins/FX/Niagara/Source/NiagaraEditor/Private/ViewModels/Stack/NiagaraStackStatelessEmitterSimulateGroup.cpp
@@ -426,16 +426,7 @@ void UNiagaraStackStatelessModuleItem::RefreshChildrenInternal(const TArray<UNia
 			if (StatelessModule->CanDebugDraw())
 			{
 				// Debug draw is handled separately here since it's visibility needs to be determined by the results of the CanDebugDraw function, rather than by property metadata.
-				FBoolProperty* DebugDrawProperty = nullptr;
-				for (TFieldIterator<FProperty> PropertyIt(StatelessModule->GetClass(), EFieldIteratorFlags::SuperClassFlags::IncludeSuper, EFieldIteratorFlags::DeprecatedPropertyFlags::ExcludeDeprecated); PropertyIt; ++PropertyIt)
-				{
-					if (PropertyIt->GetFName() == UNiagaraStatelessModule::PrivateMemberNames::bDebugDrawEnabled)
-					{
-						DebugDrawProperty = CastField<FBoolProperty>(*PropertyIt);
-						break;
-					}
-				}
-				if (DebugDrawProperty != nullptr)
+				if ( FBoolProperty* DebugDrawProperty = GetDebugDrawProperty() )
 				{
 					HeaderValueHandlers.Add(MakeShared<FNiagaraStackItemPropertyHeaderValue>(
 						*StatelessModule, nullptr, *DebugDrawProperty, 
@@ -483,6 +474,46 @@ void UNiagaraStackStatelessModuleItem::GetHeaderValueHandlers(TArray<TSharedRef<
 	}
 }
 
+bool UNiagaraStackStatelessModuleItem::SupportsDebugDraw() const
+{
+	UNiagaraStatelessModule* StatelessModule = StatelessModuleWeak.Get();
+	return StatelessModule ? StatelessModule->CanDebugDraw() : false;
+}
+
+TOptional<FText> UNiagaraStackStatelessModuleItem::GetDebugDrawTooltip() const
+{
+	TOptional<FText> Tooltip;
+	if (FBoolProperty* DebugDrawProperty = GetDebugDrawProperty())
+	{
+		Tooltip = DebugDrawProperty->GetToolTipText();
+	}
+	return Tooltip;
+}
+
+bool UNiagaraStackStatelessModuleItem::IsDebugDrawEnabled() const
+{
+	UNiagaraStatelessModule* StatelessModule = StatelessModuleWeak.Get();
+	return StatelessModule ? StatelessModule->IsDebugDrawEnabled() : false;
+}
+
+void UNiagaraStackStatelessModuleItem::SetDebugDrawEnabled(bool bInEnabled)
+{
+	UNiagaraStatelessModule* StatelessModule = StatelessModuleWeak.Get();
+	FBoolProperty* DebugDrawProperty = GetDebugDrawProperty();
+	if (StatelessModule && DebugDrawProperty)
+	{
+		FScopedTransaction ScopedTransaction(FText::Format(LOCTEXT("SetDebugDrawEnabled", "Change {0}"), DebugDrawProperty->GetDisplayNameText()));
+		StatelessModule->Modify();
+
+		DebugDrawProperty->SetPropertyValue(DebugDrawProperty->ContainerPtrToValuePtr<void>(StatelessModule), bInEnabled);
+		
+		FPropertyChangedEvent PropertyChangedEvent(DebugDrawProperty);
+		StatelessModule->PostEditChangeProperty(PropertyChangedEvent);
+
+		OnHeaderValueChanged();
+	}
+}
+
 const UNiagaraStackEntry::FCollectedUsageData& UNiagaraStackStatelessModuleItem::GetCollectedUsageData() const
 {
 	if (CachedCollectedUsageData.IsSet() == false)
@@ -540,4 +571,16 @@ void UNiagaraStackStatelessModuleItem::OnHeaderValueChanged()
 	}
 }
 
+FBoolProperty* UNiagaraStackStatelessModuleItem::GetDebugDrawProperty() const
+{
+	UNiagaraStatelessModule* StatelessModule = StatelessModuleWeak.Get();
+	if (StatelessModule == nullptr)
+	{
+		return nullptr;
+	}
+
+	UClass* ModuleClass = StatelessModule->GetClass();
+	return CastField<FBoolProperty>(ModuleClass->FindPropertyByName(UNiagaraStatelessModule::PrivateMemberNames::bDebugDrawEnabled));
+}
+
 #undef LOCTEXT_NAMESPACE
diff --git a/Engine/Plugins/FX/Niagara/Source/NiagaraEditor/Public/ViewModels/Stack/NiagaraStackStatelessEmitterSimulateGroup.h b/Engine/Plugins/FX/Niagara/Source/NiagaraEditor/Public/ViewModels/Stack/NiagaraStackStatelessEmitterSimulateGroup.h
--- a/Engine/Plugins/FX/Niagara/Source/NiagaraEditor/Public/ViewModels/Stack/NiagaraStackStatelessEmitterSimulateGroup.h
+++ b/Engine/Plugins/FX/Niagara/Source/NiagaraEditor/Public/ViewModels/Stack/NiagaraStackStatelessEmitterSimulateGroup.h
@@ -80,6 +80,11 @@ public:
 	virtual bool SupportsHeaderValues() const override { return true; }
 	virtual void GetHeaderValueHandlers(TArray<TSharedRef<INiagaraStackItemHeaderValueHandler>>& OutHeaderValueHandlers) const override;
 
+	virtual bool SupportsDebugDraw() const override;
+	virtual TOptional<FText> GetDebugDrawTooltip() const override;
+	virtual bool IsDebugDrawEnabled() const override;
+	virtual void SetDebugDrawEnabled(bool bInEnabled) override;
+
 	virtual const FCollectedUsageData& GetCollectedUsageData() const override;
 
 	UNiagaraStatelessModule* GetStatelessModule() const { return StatelessModuleWeak.Get(); }
@@ -95,6 +100,8 @@ private:
 
 	void OnHeaderValueChanged();
 
+	FBoolProperty* GetDebugDrawProperty() const;
+
 private:
 	TWeakObjectPtr<UNiagaraStatelessModule> StatelessModuleWeak;
 	FText DisplayName;
