- Fix for spawn rate not working with infinite loop duration with LWE
- Around 1.5 days we still run into precision issues, but up until that point everything spawns as expected

#rnx
#tests Make a spawn rate LWE set loop to once + infinite notice nothing spawns before this change

diff --git a/Engine/Plugins/FX/Niagara/Source/Niagara/Private/Stateless/NiagaraStatelessEmitterInstance.cpp b/Engine/Plugins/FX/Niagara/Source/Niagara/Private/Stateless/NiagaraStatelessEmitterInstance.cpp
--- a/Engine/Plugins/FX/Niagara/Source/Niagara/Private/Stateless/NiagaraStatelessEmitterInstance.cpp
+++ b/Engine/Plugins/FX/Niagara/Source/Niagara/Private/Stateless/NiagaraStatelessEmitterInstance.cpp
@@ -362,8 +362,11 @@ void FNiagaraStatelessEmitterInstance::InitEmitterState()
 
 	if (EmitterState.LoopBehavior == ENiagaraLoopBehavior::Once && EmitterState.LoopDurationMode == ENiagaraLoopDurationMode::Infinite)
 	{
-		CurrentLoopDuration	= FLT_MAX;
-		CurrentLoopAgeEnd	= FLT_MAX;
+		// Infinite will loop after this time which works as expected, however we do start to run into numerical precision issues after running for around 1.5 days
+		// We also can not use TNumericLimits<float>::Max() / FLT_MAX as it will introduce numerical precision issues resulting in nothing appearing to spawn for rate based spawning
+		constexpr float LongDuration = 60.0f * 60.0f * 24.0f;
+		CurrentLoopDuration	= LongDuration;
+		CurrentLoopAgeEnd	= LongDuration;
 	}
 	else
 	{
