Replace GNiagaraPrecachePSOAtAssetLoadingTime with GNiagaraPSOPrecachingTime to have more control when Niagara based PSOs should be precached. Options are: asset loading time, component time or proxy creation time (or not at all)

diff --git a/Engine/Plugins/FX/Niagara/Source/Niagara/Private/NiagaraComponent.cpp b/Engine/Plugins/FX/Niagara/Source/Niagara/Private/NiagaraComponent.cpp
--- a/Engine/Plugins/FX/Niagara/Source/Niagara/Private/NiagaraComponent.cpp
+++ b/Engine/Plugins/FX/Niagara/Source/Niagara/Private/NiagaraComponent.cpp
@@ -54,6 +54,8 @@ DEFINE_LOG_CATEGORY(LogNiagara);
 
 DEFINE_RENDER_COMMAND_PIPE(NiagaraDynamicData, ERenderCommandPipeFlags::None);
 
+extern int32 GNiagaraPSOPrecachingTime;
+
 namespace NiagaraComponentPrivate
 {
 	static bool GParallelGDMEEnabled = true;
@@ -1234,7 +1236,10 @@ bool UNiagaraComponent::InitializeSystem()
 			}
 		}
 
-		PrecachePSOs();
+		if (GNiagaraPSOPrecachingTime > 1)
+		{
+			PrecachePSOs();
+		}
 
 #if WITH_EDITORONLY_DATA
 		OnSystemInstanceChangedDelegate.Broadcast();
@@ -1746,7 +1751,10 @@ void UNiagaraComponent::PostSystemTick_GameThread()
 
 	if (bRecachePSOs)
 	{
-		PrecachePSOs();
+		if (GNiagaraPSOPrecachingTime > 1)
+		{
+			PrecachePSOs();
+		}
 		bRecachePSOs = false;
 	}
 }
@@ -2425,12 +2433,18 @@ FPrimitiveSceneProxy* UNiagaraComponent::CreateSceneProxy()
 #if UE_WITH_PSO_PRECACHING
 	if (Asset != nullptr)
 	{
+		// If we are relying on proxy creation PSOs only then call precache PSOs here because asset doesn't have collected PSO precache data
+		if (GNiagaraPSOPrecachingTime == 1)
+		{
+			PrecachePSOs();
+		}
+
 		// If PSOs not precached yet then rely on PSOs reuqested by the shared asset itself - this will make sure
 		// those pending requests are boosted if still pending
 		if (!bPSOPrecacheCalled)
 		{
 			PrecacheAssetPSOs(Asset);
-		}
+		} 		
 
 		EPSOPrecachePriority PSOPrecachePriority = NiagaraComponentPrivate::GBoostPSOPrecacheHighest ? EPSOPrecachePriority::Highest : EPSOPrecachePriority::High;
 		if (CheckPSOPrecachingAndBoostPriority(PSOPrecachePriority) && GetPSOPrecacheProxyCreationStrategy() != EPSOPrecacheProxyCreationStrategy::AlwaysCreate)
@@ -2451,6 +2465,14 @@ void UNiagaraComponent::CollectPSOPrecacheData(const FPSOPrecacheParams& BasePre
 	if (SystemInstanceController.IsValid())
 	{
 		SystemInstanceController->CollectPSOPrecacheData(BasePrecachePSOParams, OutParams);
+		
+		if (bRenderCustomDepth)
+		{
+			for (FMaterialInterfacePSOPrecacheParams& PrecacheParams : OutParams)
+			{
+				PrecacheParams.PSOPrecacheParams.bRenderCustomDepth = bRenderCustomDepth;
+			}
+		}
 	}
 }
 
diff --git a/Engine/Plugins/FX/Niagara/Source/Niagara/Private/NiagaraSystem.cpp b/Engine/Plugins/FX/Niagara/Source/Niagara/Private/NiagaraSystem.cpp
--- a/Engine/Plugins/FX/Niagara/Source/Niagara/Private/NiagaraSystem.cpp
+++ b/Engine/Plugins/FX/Niagara/Source/Niagara/Private/NiagaraSystem.cpp
@@ -193,11 +193,15 @@ static FAutoConsoleVariableRef CVarNiagaraCompileWaitLoggingTerminationCap(
 	ECVF_Default
 );
 
-int32 GNiagaraPrecachePSOAtAssetLoadingTime = 1;
-FAutoConsoleVariableRef CVarNiagaraPrecachePSOAtAssetLoadingTime(
-	TEXT("r.PSOPrecache.NiagaraPrecachePSOAtAssetLoadingTime"),
-	GNiagaraPrecachePSOAtAssetLoadingTime,
-	TEXT("Controls whether Niagara systems start PSO precaching at asset loading time (1 = default) or only when activated (0)."),
+int32 GNiagaraPSOPrecachingTime = 3;
+FAutoConsoleVariableRef CVarNiagaraPSOPrecachingTime(
+	TEXT("r.PSOPrecache.NiagaraPrecachingTime"),
+	GNiagaraPSOPrecachingTime,
+	TEXT("Controls when PSO precaching happens for Niagara systems:\n")
+	TEXT("	0: no precaching\n")
+	TEXT("	1: precaching at component proxy creation time only\n")
+	TEXT("	2: precaching at component loading time\n")
+	TEXT("	3: precaching at asset loading time and component time (default)"),
 	ECVF_Default);
 
 #if WITH_EDITORONLY_DATA
@@ -1573,7 +1577,7 @@ void UNiagaraSystem::PostLoad()
 	}
 #endif // WITH_EDITORONLY_DATA
 
-	if (GNiagaraPrecachePSOAtAssetLoadingTime)
+	if (GNiagaraPSOPrecachingTime >= 3)
 	{
 		PrecachePSOs();
 	}
