- Add a stack warning for missing sampling regions

diff --git a/Engine/Plugins/FX/Niagara/Source/Niagara/Private/NiagaraDataInterfaceSkeletalMesh.cpp b/Engine/Plugins/FX/Niagara/Source/Niagara/Private/NiagaraDataInterfaceSkeletalMesh.cpp
--- a/Engine/Plugins/FX/Niagara/Source/Niagara/Private/NiagaraDataInterfaceSkeletalMesh.cpp
+++ b/Engine/Plugins/FX/Niagara/Source/Niagara/Private/NiagaraDataInterfaceSkeletalMesh.cpp
@@ -2743,6 +2743,36 @@ void UNiagaraDataInterfaceSkeletalMesh::GetFeedback(UNiagaraSystem* Asset, UNiag
 			);
 		}
 	}
+
+	// Do we have any sampling regions enabled, if so check they are valid.
+	if (SkelMesh && SamplingRegions.Num() > 0)
+	{
+		FString MissingRegions;
+
+		const FSkeletalMeshSamplingInfo& SamplingInfo = SkelMesh->GetSamplingInfo();
+		for (FName RegionName : SamplingRegions)
+		{
+			const int32 RegionIdx = SamplingInfo.IndexOfRegion(RegionName);
+			if (RegionIdx == INDEX_NONE)
+			{
+				if (!MissingRegions.IsEmpty())
+				{
+					MissingRegions.Append(TEXT(", "));
+				}
+				RegionName.AppendString(MissingRegions);
+			}
+		}
+
+		if (!MissingRegions.IsEmpty())
+		{
+			OutWarnings.Emplace(
+				FText::Format(LOCTEXT("MissingSamplingRegions", "The following regions were not found '{0}'"), FText::FromString(MissingRegions)),
+				LOCTEXT("MissingSamplingRegionsSummary", "Invalid Sampling Regions"),
+				FNiagaraDataInterfaceFix(
+				)
+			);
+		}
+	}
 #endif
 
 	if (SoftSourceActor.Get() == nullptr && bHasNoMeshAssignedWarning)
