- Fix potential divide by zero
- In practise this should not be possible WidthSegmentationCount will never be 0

#rnx

diff --git a/Engine/Plugins/FX/Niagara/Source/Niagara/Private/NiagaraRendererRibbons.cpp b/Engine/Plugins/FX/Niagara/Source/Niagara/Private/NiagaraRendererRibbons.cpp
--- a/Engine/Plugins/FX/Niagara/Source/Niagara/Private/NiagaraRendererRibbons.cpp
+++ b/Engine/Plugins/FX/Niagara/Source/Niagara/Private/NiagaraRendererRibbons.cpp
@@ -1235,7 +1235,6 @@ void FNiagaraRendererRibbons::GetDynamicRayTracingInstances(FRayTracingInstanceC
 }
 #endif
 
-//-V::609 UE-348905
 void FNiagaraRendererRibbons::GenerateShapeStateMultiPlane(FNiagaraRibbonShapeGeometryData& State, int32 MultiPlaneCount, int32 WidthSegmentationCount, bool bEnableAccurateGeometry, bool bUseMaterialBackfaceCulling)
 {
 	using namespace NiagaraRendererRibbonsPrivate;
@@ -1248,15 +1247,16 @@ void FNiagaraRendererRibbons::GenerateShapeStateMultiPlane(FNiagaraRibbonShapeGe
 	State.BitsNeededForShape = CalculateBitsForRange(State.NumVerticesInSlice);
 	State.BitMaskForShape = CalculateBitMask(State.BitsNeededForShape);
 	
+	const float NormalizeVertexId = float(WidthSegmentationCount > 0 ? WidthSegmentationCount : 1);
 	for (int32 PlaneIndex = 0; PlaneIndex < MultiPlaneCount; PlaneIndex++)
 	{
 		const float RotationAngle = (static_cast<float>(PlaneIndex) / MultiPlaneCount) * 180.0f;
 
 		for (int32 VertexId = 0; VertexId <= WidthSegmentationCount; VertexId++)
 		{
-			const FVector2f Position = FVector2f((static_cast<float>(VertexId) / WidthSegmentationCount) - 0.5f, 0).GetRotated(RotationAngle);
+			const FVector2f Position = FVector2f((static_cast<float>(VertexId) / NormalizeVertexId) - 0.5f, 0).GetRotated(RotationAngle);
 			const FVector2f Normal = FVector2f(0, 1).GetRotated(RotationAngle);
-			const float TextureV = static_cast<float>(VertexId) / WidthSegmentationCount;
+			const float TextureV = static_cast<float>(VertexId) / NormalizeVertexId;
 
 			State.SliceVertexData.Emplace(Position, Normal, TextureV);
 		}
@@ -1270,9 +1270,9 @@ void FNiagaraRendererRibbons::GenerateShapeStateMultiPlane(FNiagaraRibbonShapeGe
 
 			for (int32 VertexId = 0; VertexId <= WidthSegmentationCount; VertexId++)
 			{
-				const FVector2f Position = FVector2f((static_cast<float>(VertexId) / WidthSegmentationCount) - 0.5f, 0).GetRotated(RotationAngle);
+				const FVector2f Position = FVector2f((static_cast<float>(VertexId) / NormalizeVertexId) - 0.5f, 0).GetRotated(RotationAngle);
 				const FVector2f Normal = FVector2f(0, -1).GetRotated(RotationAngle);
-				const float TextureV = static_cast<float>(VertexId) / WidthSegmentationCount;
+				const float TextureV = static_cast<float>(VertexId) / NormalizeVertexId;
 
 				State.SliceVertexData.Emplace(Position, Normal, TextureV);
 			}
@@ -1389,11 +1389,13 @@ void FNiagaraRendererRibbons::GenerateShapeStatePlane(FNiagaraRibbonShapeGeometr
 	State.BitMaskForShape = CalculateBitMask(State.BitsNeededForShape);
 	
 	State.SliceVertexData.Reserve(WidthSegmentationCount + 1);
+
+	const float NormalizeVertexId = float(WidthSegmentationCount > 0 ? WidthSegmentationCount : 1);
 	for (int32 VertexId = 0; VertexId <= WidthSegmentationCount; VertexId++)
 	{
-		const FVector2f Position = FVector2f((static_cast<float>(VertexId) / WidthSegmentationCount) - 0.5f, 0);
+		const FVector2f Position = FVector2f((static_cast<float>(VertexId) / NormalizeVertexId) - 0.5f, 0);
 		const FVector2f Normal = FVector2f(0, 1);
-		const float TextureV = static_cast<float>(VertexId) / WidthSegmentationCount;
+		const float TextureV = static_cast<float>(VertexId) / NormalizeVertexId;
 
 		State.SliceVertexData.Emplace(Position, Normal, TextureV);
 	}
