- Fix potential divide by zero
- In practise this should not be possible TubeSubdivisions will never be 0

#rnx

diff --git a/Engine/Plugins/FX/Niagara/Source/Niagara/Private/NiagaraRendererRibbons.cpp b/Engine/Plugins/FX/Niagara/Source/Niagara/Private/NiagaraRendererRibbons.cpp
--- a/Engine/Plugins/FX/Niagara/Source/Niagara/Private/NiagaraRendererRibbons.cpp
+++ b/Engine/Plugins/FX/Niagara/Source/Niagara/Private/NiagaraRendererRibbons.cpp
@@ -1317,12 +1317,14 @@ void FNiagaraRendererRibbons::GenerateShapeStateTube(FNiagaraRibbonShapeGeometry
 	State.BitMaskForShape = CalculateBitMask(State.BitsNeededForShape);
 	
 	State.SliceVertexData.Reserve(TubeSubdivisions + 1);
+
+	const float NormalizeVertexId = float(TubeSubdivisions > 0 ? TubeSubdivisions : 1);
 	for (int32 VertexId = 0; VertexId <= TubeSubdivisions; VertexId++)
 	{
-		const float RotationAngle = (static_cast<float>(VertexId) / TubeSubdivisions) * -360.0f;
+		const float RotationAngle = (static_cast<float>(VertexId) / NormalizeVertexId) * -360.0f;
 		const FVector2f Position = FVector2f(-0.5f, 0.0f).GetRotated(RotationAngle);
 		const FVector2f Normal = FVector2f(-1, 0).GetRotated(RotationAngle);
-		const float TextureV = static_cast<float>(VertexId) / TubeSubdivisions;
+		const float TextureV = static_cast<float>(VertexId) / NormalizeVertexId;
 
 		State.SliceVertexData.Emplace(Position, Normal, TextureV);
 	}
