Fixed missing delegate unregister

diff --git a/Engine/Plugins/FX/Niagara/Source/NiagaraEditor/Private/ViewModels/NiagaraSystemViewModel.cpp b/Engine/Plugins/FX/Niagara/Source/NiagaraEditor/Private/ViewModels/NiagaraSystemViewModel.cpp
--- a/Engine/Plugins/FX/Niagara/Source/NiagaraEditor/Private/ViewModels/NiagaraSystemViewModel.cpp
+++ b/Engine/Plugins/FX/Niagara/Source/NiagaraEditor/Private/ViewModels/NiagaraSystemViewModel.cpp
@@ -275,6 +275,8 @@ void FNiagaraSystemViewModel::Cleanup()
 		System->OnSystemPostEditChange().Remove(SystemChangedDelegateHandle);
 	}
 
+	UNiagaraEditorSettings::OnSettingsChanged().RemoveAll(this);
+	
 	// Make sure that we clear out all of our event handlers
 	UnregisterViewModelWithMap(RegisteredHandle);
 
