- Apply Owner Scale to Attributes module for LW emitters
- Mesh scale options do not exist as local space emitters already match with LW
- Drag options have not been added due to complexities around handling the stable solve velocities path

#virtualized

diff --git a/Engine/Plugins/FX/Niagara/Shaders/Private/Stateless/Modules/NiagaraStatelessModule_CameraOffset.ush b/Engine/Plugins/FX/Niagara/Shaders/Private/Stateless/Modules/NiagaraStatelessModule_CameraOffset.ush
--- a/Engine/Plugins/FX/Niagara/Shaders/Private/Stateless/Modules/NiagaraStatelessModule_CameraOffset.ush
+++ b/Engine/Plugins/FX/Niagara/Shaders/Private/Stateless/Modules/NiagaraStatelessModule_CameraOffset.ush
@@ -1,13 +1,14 @@
 // Copyright Epic Games, Inc. All Rights Reserved.
 
 FNiagaraStatelessBuiltDistributionType CameraOffset_Distribution;
+float CameraOffset_Scale;
 
 void CameraOffset_Simulate(inout FStatelessParticle Particle)
 {
 	if (FNiagaraStatelessBuiltDistribution::IsValid(CameraOffset_Distribution))
 	{
 		const float2 CameraOffsets = SampleDistributionValuesFloat(CameraOffset_Distribution, 0, Particle.NormalizedAge, Particle.PreviousNormalizedAge);
-		Particle.CameraOffset			= CameraOffsets[0];
-		Particle.PreviousCameraOffset	= CameraOffsets[1];
+		Particle.CameraOffset			= CameraOffsets[0] * CameraOffset_Scale;
+		Particle.PreviousCameraOffset	= CameraOffsets[1] * CameraOffset_Scale;
 	}
 }
diff --git a/Engine/Plugins/FX/Niagara/Shaders/Private/Stateless/Modules/NiagaraStatelessModule_SolveVelocitiesAndForces.ush b/Engine/Plugins/FX/Niagara/Shaders/Private/Stateless/Modules/NiagaraStatelessModule_SolveVelocitiesAndForces.ush
--- a/Engine/Plugins/FX/Niagara/Shaders/Private/Stateless/Modules/NiagaraStatelessModule_SolveVelocitiesAndForces.ush
+++ b/Engine/Plugins/FX/Niagara/Shaders/Private/Stateless/Modules/NiagaraStatelessModule_SolveVelocitiesAndForces.ush
@@ -1,5 +1,8 @@
 // Copyright Epic Games, Inc. All Rights Reserved.
 
+float3	SolveVelocitiesAndForces_OwnerAccelerationScale;
+//float	SolveVelocitiesAndForces_OwnerDragScale;
+float3	SolveVelocitiesAndForces_OwnerVelocityScale;
 float	SolveVelocitiesAndForces_MassScale;
 float	SolveVelocitiesAndForces_MassBias;
 float	SolveVelocitiesAndForces_DragScale;
@@ -134,7 +137,13 @@ void SolveVelocitiesAndForces_Simulate(inout FStatelessParticle Particle)
 		SolveVelocitiesAndForces_EvaluateNoise(ModuleData, Particle.Lifetime, Particle.NormalizedAge, Particle.Age, NosieFieldBias, Particle.Position);
 		SolveVelocitiesAndForces_EvaluateNoise(ModuleData, Particle.Lifetime, Particle.PreviousNormalizedAge, Particle.PreviousAge, NosieFieldBias, Particle.PreviousPosition);
 	}
-	
+
+	// Apply Owner Scaling to data
+	//ModuleData.Drag			*= SolveVelocitiesAndForces_OwnerDragScale;
+	ModuleData.Velocity		*= SolveVelocitiesAndForces_OwnerVelocityScale;
+	ModuleData.Wind			*= SolveVelocitiesAndForces_OwnerAccelerationScale;
+	ModuleData.Acceleration *= SolveVelocitiesAndForces_OwnerAccelerationScale;
+
 	// Simulate
 	// Note we need to ensure velocity aligned sprites work we need to make sure Age & PreviousAge != 0.0
 	SolveVelocitiesAndForces_IntegratePosition(ModuleData, max(Particle.Age, UE_KINDA_SMALL_NUMBER), Particle.Position);
diff --git a/Engine/Plugins/FX/Niagara/Shaders/Private/Stateless/NiagaraStatelessSimulationDefault.usf b/Engine/Plugins/FX/Niagara/Shaders/Private/Stateless/NiagaraStatelessSimulationDefault.usf
--- a/Engine/Plugins/FX/Niagara/Shaders/Private/Stateless/NiagaraStatelessSimulationDefault.usf
+++ b/Engine/Plugins/FX/Niagara/Shaders/Private/Stateless/NiagaraStatelessSimulationDefault.usf
@@ -50,6 +50,7 @@
 	PARTICLE_MODULE(Null) /* CurlNoiseForce */ \
 	PARTICLE_MODULE(Null) /* Drag */ \
 	PARTICLE_MODULE(Null) /* GravityForce */ \
+	PARTICLE_MODULE(Null) /* ApplyOwnerScaleToAttributes */ \
 	PARTICLE_MODULE(SolveVelocitiesAndForces) \
 	/* Post Solve */ \
 	PARTICLE_MODULE(CameraOffset) \
diff --git a/Engine/Plugins/FX/Niagara/Source/Niagara/Internal/Stateless/Modules/NiagaraStatelessModuleCommon.h b/Engine/Plugins/FX/Niagara/Source/Niagara/Internal/Stateless/Modules/NiagaraStatelessModuleCommon.h
--- a/Engine/Plugins/FX/Niagara/Source/Niagara/Internal/Stateless/Modules/NiagaraStatelessModuleCommon.h
+++ b/Engine/Plugins/FX/Niagara/Source/Niagara/Internal/Stateless/Modules/NiagaraStatelessModuleCommon.h
@@ -4,12 +4,107 @@
 
 #include "Stateless/NiagaraStatelessModule.h"
 
+#include "NiagaraStatelessModuleCommon.generated.h"
+
+USTRUCT()
+struct FNiagaraStatelessSystemScaleBuildData
+{
+	GENERATED_BODY()
+
+	static FName GetName();
+
+	////////////////////////////////////////////////////////////////////////////////////////////////////////
+	// Enable bits
+	UPROPERTY(EditAnywhere, Category = "Parameters|Camera", meta = (DisplayPriority = 600))
+	uint32 bScaleCameraOffset : 1 = false;
+
+	// Not including drag scale right now as we have separate solver passes which is harder to handle
+	//UPROPERTY(EditAnywhere, Category = "Parameters|Drag", meta = (DisplayPriority = 500))
+	//uint32 bScaleDrag : 1 = false;
+
+	UPROPERTY(EditAnywhere, Category = "Parameters|Forces", meta = (DisplayPriority = 200))
+	uint32 bScaleForces : 1 = false;
+
+	// Mesh scale in a no-op for local space with standard emitters
+	//UPROPERTY(EditAnywhere, Category = "Parameters", meta = (DisplayPriority = 400))
+	//uint32 bScaleInitialMeshSize : 1 = false;
+
+	UPROPERTY(EditAnywhere, Category = "Parameters|Ribbon", meta = (DisplayPriority = 700))
+	uint32 bScaleInitialRibbonWidth : 1 = false;
+
+	UPROPERTY(EditAnywhere, Category = "Parameters|Sprite", meta = (DisplayPriority = 300))
+	uint32 bScaleInitialSpriteSize : 1 = false;
+
+	UPROPERTY(EditAnywhere, Category = "Parameters|Veclocity", meta = (DisplayPriority = 100))
+	uint32 bScaleInitialVelocity : 1 = false;
+
+	////////////////////////////////////////////////////////////////////////////////////////////////////////
+	// Bitfields
+	UPROPERTY(EditAnywhere, Category = "Parameters|Sprite", meta = (DisplayPriority = 302, EditCondition = "bScaleInitialSpriteSize", EditConditionHides))
+	uint32 bScaleInitialSpriteSizeBySingleAxis : 1 = false;
+
+	//UPROPERTY(EditAnywhere, Category = "Parameters", meta = (DisplayPriority = 403, EditCondition = "bScaleInitialMeshSize", EditConditionHides))
+	//uint32 bScaleInitialMeshSizeAbsoluteScale : 1 = false;
+
+	UPROPERTY(EditAnywhere, Category = "Parameters|Ribbon", meta = (DisplayPriority = 702, EditCondition = "bScaleInitialRibbonWidth", EditConditionHides))
+	uint32 bScaleInitialRibbonWidthBySingleAxis : 1 = false;
+
+	//UPROPERTY(EditAnywhere, Category = "Parameters|Drag", meta = (DisplayPriority = 502))
+	//uint32 bInverseDragScaling : 1 = false;
+
+	////////////////////////////////////////////////////////////////////////////////////////////////////////
+	// Velocity
+	UPROPERTY(EditAnywhere, Category = "Parameters|Veclocity", meta = (DisplayPriority = 101, EditCondition = "bScaleInitialVelocity", EditConditionHides))
+	FVector3f ScaleInitialVelocityScaleAmount = FVector3f::OneVector;
+
+	////////////////////////////////////////////////////////////////////////////////////////////////////////
+	// Forces
+	UPROPERTY(EditAnywhere, Category = "Parameters|Forces", meta = (DisplayPriority = 201, EditCondition = "bScaleForces", EditConditionHides))
+	FVector3f ScaleForcesScaleAmount = FVector3f::OneVector;
+
+	UPROPERTY(EditAnywhere, Category = "Parameters|Forces", meta = (DisplayPriority = 202, EditCondition = "bScaleForces", EditConditionHides))
+	uint32 bScaleForcesInverseScaling : 1 = false;
+
+	////////////////////////////////////////////////////////////////////////////////////////////////////////
+	// Sprites
+	UPROPERTY(EditAnywhere, Category = "Parameters|Sprite", meta = (DisplayPriority = 301, EditCondition = "bScaleInitialSpriteSize", EditConditionHides))
+	float ScaleInitialSpriteSizeScaleAmount = 1.0f;
+
+	UPROPERTY(EditAnywhere, Category = "Parameters|Sprite", meta = (DisplayPriority = 303, EditCondition = "bScaleInitialSpriteSize && bScaleInitialSpriteSizeBySingleAxis", EditConditionHides))
+	ENiagaraOrientationAxis ScaleInitialSpriteSizeAxis = ENiagaraOrientationAxis::XAxis;
+
+	////////////////////////////////////////////////////////////////////////////////////////////////////////
+	// Meshes
+	//UPROPERTY(EditAnywhere, Category = "Parameters", meta = (DisplayPriority = 401, EditCondition = "bScaleInitialMeshSize", EditConditionHides))
+	//FVector3f ScaleInitialMeshSizeScaleAmount = FVector3f::OneVector;
+
+	//UPROPERTY(EditAnywhere, Category = "Parameters", meta = (DisplayPriority = 402, EditCondition = "bScaleInitialMeshSize", EditConditionHides))
+	//ENiagaraCoordinateSpace ScaleInitialMeshSizeScaleSpace = ENiagaraCoordinateSpace::Simulation;
+
+	////////////////////////////////////////////////////////////////////////////////////////////////////////
+	// Drag
+	//UPROPERTY(EditAnywhere, Category = "Parameters|Drag", meta = (DisplayPriority = 501, EditCondition = "bScaleDrag", EditConditionHides))
+	//float ScaleDragScaleAmount = 1.0f;
+
+	////////////////////////////////////////////////////////////////////////////////////////////////////////
+	// Camera
+	UPROPERTY(EditAnywhere, Category = "Parameters|Camera", meta = (DisplayPriority = 601, EditCondition = "bScaleCameraOffset", EditConditionHides))
+	float ScaleCameraOffsetAmount = 1.0f;
+
+	////////////////////////////////////////////////////////////////////////////////////////////////////////
+	// Ribbons
+	UPROPERTY(EditAnywhere, Category = "Parameters|Ribbon", meta = (DisplayPriority = 701, EditCondition = "bScaleInitialRibbonWidth", EditConditionHides))
+	float ScaleInitialRibbonWidthScaleAmount = 1.0f;
+	UPROPERTY(EditAnywhere, Category = "Parameters|Ribbon", meta = (DisplayPriority = 703, EditCondition = "bScaleInitialRibbonWidth && bScaleInitialRibbonWidthBySingleAxis", EditConditionHides))
+	ENiagaraOrientationAxis ScaleInitialRibbonWidthScalingAxis = ENiagaraOrientationAxis::XAxis;
+};
+
 namespace NiagaraStateless
 {
 	//-TODO: Examine structure padding, FQuat4f / UObject*
 	struct FPhysicsBuildData
 	{
-		static FName GetName() { return FName("FPhysicsBuildData"); }
+		static FName GetName();
 
 		FNiagaraStatelessRangeFloat		MassRange = FNiagaraStatelessRangeFloat(1.0f);
 		FNiagaraStatelessRangeFloat		DragRange = FNiagaraStatelessRangeFloat(0.0f);
@@ -48,21 +143,5 @@ namespace NiagaraStateless
 		FNiagaraStatelessRangeVector3	NoiseFieldOffset = FNiagaraStatelessRangeVector3(FVector3f::ZeroVector);
 	};
 
-	inline FQuat4f DirectionToQuat(FVector3f Direction)
-	{
-		Direction = Direction.GetSafeNormal(UE_KINDA_SMALL_NUMBER, FVector3f::ZAxisVector);
-
-		const FVector3f Forward = FVector3f::ZAxisVector;
-		const float Dot = FVector3f::DotProduct(Forward, Direction);
-		if (Dot < -0.99999f)
-		{
-			return FQuat4f(FVector3f::YAxisVector, PI);
-		}
-
-		const FVector3f Axis = FVector3f::CrossProduct(Forward, Direction);
-		const float Angle = FMath::Acos(Dot);
-		FQuat4f Quat(Axis, Angle);
-		Quat.Normalize();
-		return Quat;
-	}
+	extern FQuat4f DirectionToQuat(FVector3f Direction);
 }
diff --git a/Engine/Plugins/FX/Niagara/Source/Niagara/Internal/Stateless/Modules/NiagaraStatelessModule_ApplyOwnerScaleToAttributes.h b/Engine/Plugins/FX/Niagara/Source/Niagara/Internal/Stateless/Modules/NiagaraStatelessModule_ApplyOwnerScaleToAttributes.h
new file mode 100644
--- /dev/null
+++ b/Engine/Plugins/FX/Niagara/Source/Niagara/Internal/Stateless/Modules/NiagaraStatelessModule_ApplyOwnerScaleToAttributes.h
@@ -0,0 +1,26 @@
+// Copyright Epic Games, Inc. All Rights Reserved.
+
+#pragma once
+
+#include "Stateless/NiagaraStatelessModule.h"
+#include "Stateless/Modules/NiagaraStatelessModuleCommon.h"
+
+#include "NiagaraStatelessModule_ApplyOwnerScaleToAttributes.generated.h"
+
+// Allows you to control how Engine.Owner.Scale is applied to attributes
+UCLASS(MinimalAPI, EditInlineNew, meta = (DisplayName = "Apply Owner Scale To Attributes"))
+class UNiagaraStatelessModule_ApplyOwnerScaleToAttributes : public UNiagaraStatelessModule
+{
+	GENERATED_BODY()
+
+public:
+
+	UPROPERTY(EditAnywhere, Category = "Parameters", meta = (ShowOnlyInnerProperties, FullyExpand = true))
+	FNiagaraStatelessSystemScaleBuildData SystemScaleData;
+
+	virtual void BuildEmitterData(const FNiagaraStatelessEmitterDataBuildContext& BuildContext) const override;
+
+#if WITH_EDITOR
+	virtual bool CanDisableModule() const override { return true; }
+#endif
+};
diff --git a/Engine/Plugins/FX/Niagara/Source/Niagara/Internal/Stateless/NiagaraStatelessEmitterData.h b/Engine/Plugins/FX/Niagara/Source/Niagara/Internal/Stateless/NiagaraStatelessEmitterData.h
--- a/Engine/Plugins/FX/Niagara/Source/Niagara/Internal/Stateless/NiagaraStatelessEmitterData.h
+++ b/Engine/Plugins/FX/Niagara/Source/Niagara/Internal/Stateless/NiagaraStatelessEmitterData.h
@@ -23,6 +23,8 @@ struct FInstancedStruct;
 
 struct FNiagaraStatelessEmitterData
 {
+	friend class FNiagaraStatelessEmitterDataBuildContext;
+
 	UE_NONCOPYABLE(FNiagaraStatelessEmitterData);
 
 	FNiagaraStatelessEmitterData() = default;
@@ -58,6 +60,8 @@ struct FNiagaraStatelessEmitterData
 	const UNiagaraStatelessEmitterTemplate*			EmitterTemplate = nullptr;	// Used to access shader information
 
 	TArray<uint8>									BuiltData;					// Built data, generally allocated by modules if any
+	TArray<uint8>									SharedBuiltData;			// Shared built data, allocated by modules generally used during parameter setting
+	TArray<TPair<FName, uint32>>					SharedBuiltDataOffsets;		// Shared build data offsets, used to find the data
 	TArray<uint8>									StaticDataBufferCpu;		// Used with CPU generation, must be valid if ParticleSimExecData is also valid
 	FByteAddressBuffer								StaticDataBufferGpu;		// Used with GPU generation, guaranteed to be minimum safe size
 
diff --git a/Engine/Plugins/FX/Niagara/Source/Niagara/Internal/Stateless/NiagaraStatelessEmitterDataBuildContext.h b/Engine/Plugins/FX/Niagara/Source/Niagara/Internal/Stateless/NiagaraStatelessEmitterDataBuildContext.h
--- a/Engine/Plugins/FX/Niagara/Source/Niagara/Internal/Stateless/NiagaraStatelessEmitterDataBuildContext.h
+++ b/Engine/Plugins/FX/Niagara/Source/Niagara/Internal/Stateless/NiagaraStatelessEmitterDataBuildContext.h
@@ -10,6 +10,7 @@ struct FNiagaraDataSetCompiledData;
 struct FNiagaraParameterBinding;
 struct FNiagaraParameterBindingWithValue;
 struct FNiagaraParameterStore;
+struct FNiagaraStatelessEmitterData;
 namespace NiagaraStateless
 {
 	class FParticleSimulationContext;
@@ -22,15 +23,8 @@ class FNiagaraStatelessEmitterDataBuildContext
 public:
 	UE_NONCOPYABLE(FNiagaraStatelessEmitterDataBuildContext);
 
-	FNiagaraStatelessEmitterDataBuildContext(FNiagaraDataSetCompiledData& InParticleDataSet, FNiagaraParameterStore& InRendererBindings, TArray<TPair<int32, FInstancedStruct>>& InExpressions, TArray<uint8>& InBuiltData, TArray<uint8>& InStaticDataBuffer, NiagaraStateless::FParticleSimulationExecData* InParticleExecData)
-		: ParticleDataSet(InParticleDataSet)
-		, RendererBindings(InRendererBindings)
-		, Expressions(InExpressions)
-		, BuiltData(InBuiltData)
-		, StaticDataBuffer(InStaticDataBuffer)
-		, ParticleExecData(InParticleExecData)
-	{
-	}
+	explicit FNiagaraStatelessEmitterDataBuildContext(FNiagaraStatelessEmitterData& InEmitterData);
+	~FNiagaraStatelessEmitterDataBuildContext();
 
 	void PreModuleBuild(int32 InShaderParameterOffset);
 
@@ -42,14 +36,21 @@ public:
 	uint32 AddStaticData(TConstArrayView<FVector4f> StaticData) const;
 	uint32 AddStaticData(TConstArrayView<FLinearColor> StaticData) const;
 
+	NIAGARA_API void* AllocateBuiltData(uint32 Size, uint32 Alingment) const;
+	NIAGARA_API void* AllocateSharedBuiltData(FName DataName, uint32 Size, uint32 Alingment) const;
+
 	template<typename T>
 	T* AllocateBuiltData() const
 	{
 		static_assert(TIsTrivial<T>::Value, "Only trivial types can be used for built data");
+		void* NewData = AllocateBuiltData(sizeof(T), alignof(T));
+		return new(NewData) T();
+	}
 
-		const int32 Offset = Align(BuiltData.Num(), alignof(T));
-		BuiltData.AddZeroed(Offset + sizeof(T) - BuiltData.Num());
-		void* NewData = BuiltData.GetData() + Offset;
+	template<typename T>
+	T* AllocateSharedBuiltData() const
+	{
+		void* NewData = AllocateSharedBuiltData(T::GetName(), sizeof(T), alignof(T));
 		return new(NewData) T();
 	}
 
@@ -240,16 +241,11 @@ public:
 	NIAGARA_API int32 FindParticleVariableIndex(const FNiagaraVariableBase& Variable) const;
 
 private:
-	FNiagaraDataSetCompiledData&					ParticleDataSet;
-	FNiagaraParameterStore&							RendererBindings;
-	TArray<TPair<int32, FInstancedStruct>>&			Expressions;
-	TArray<uint8>&									BuiltData;
-	TArray<uint8>&									StaticDataBuffer;
-	NiagaraStateless::FParticleSimulationExecData*	ParticleExecData = nullptr;
-
-	int32											ModuleBuiltDataOffset = 0;
-	int32											ShaderParameterOffset = 0;
-	int32											RandomSeedOffest = 0;
+	FNiagaraStatelessEmitterData&	EmitterData;
+
+	int32							ModuleBuiltDataOffset = 0;
+	int32							ShaderParameterOffset = 0;
+	int32							RandomSeedOffest = 0;
 
 	struct FTransientObject
 	{
diff --git a/Engine/Plugins/FX/Niagara/Source/Niagara/Internal/Stateless/NiagaraStatelessModule.h b/Engine/Plugins/FX/Niagara/Source/Niagara/Internal/Stateless/NiagaraStatelessModule.h
--- a/Engine/Plugins/FX/Niagara/Source/Niagara/Internal/Stateless/NiagaraStatelessModule.h
+++ b/Engine/Plugins/FX/Niagara/Source/Niagara/Internal/Stateless/NiagaraStatelessModule.h
@@ -48,10 +48,21 @@ public:
 	};
 #endif
 
+	// The feature mask is used to define when they can be used
+	// For example if the module does not support GPU make sure to remove the GPU feature mask
+	// Module & template authors should be careful to not get into a situation where the feature mask is none
 	virtual ENiagaraStatelessFeatureMask GetFeatureMask() const { return ENiagaraStatelessFeatureMask::All; }
 
+	// Called during the data build process on each module sequentially
+	// Modules add build data that contain all the information needed to setup shader parameters for simulation
 	virtual void BuildEmitterData(const FNiagaraStatelessEmitterDataBuildContext& BuildContext) const {}
+
+	// Called during the data build process to create the shader parameter strucutre
+	// Modules pack structure or loose parameters that will be used to simulate
 	virtual void BuildShaderParameters(FNiagaraStatelessShaderParametersBuilder& ShaderParametersBuilder) const {}
+
+	// Called at runtime to set the shader parameter structures
+	// Modules should keep this simple and used the build process to optimize data as much as possible
 	virtual void SetShaderParameters(const FNiagaraStatelessSetShaderParameterContext& SetShaderParameterContext) const {}
 
 #if WITH_EDITOR
diff --git a/Engine/Plugins/FX/Niagara/Source/Niagara/Internal/Stateless/NiagaraStatelessParticleSimContext.h b/Engine/Plugins/FX/Niagara/Source/Niagara/Internal/Stateless/NiagaraStatelessParticleSimContext.h
--- a/Engine/Plugins/FX/Niagara/Source/Niagara/Internal/Stateless/NiagaraStatelessParticleSimContext.h
+++ b/Engine/Plugins/FX/Niagara/Source/Niagara/Internal/Stateless/NiagaraStatelessParticleSimContext.h
@@ -252,10 +252,13 @@ public:
 	FVector2f		RandomUnitFloat2(uint32 iInstance, uint32 RandomSeedOffset) const { return SafeNormalize(RandomFloat2(iInstance, RandomSeedOffset) - 0.5f); }
 	FVector3f		RandomUnitFloat3(uint32 iInstance, uint32 RandomSeedOffset) const { return SafeNormalize(RandomFloat3(iInstance, RandomSeedOffset) - 0.5f); }
 
-	FVector2f		SafeNormalize(const FVector2f& v, const FVector2f& Fallback) const { const float l2 = v.SquaredLength(); return l2 < UE_KINDA_SMALL_NUMBER ? Fallback : v * (1.0f / FMath::Sqrt(l2)); }
-	FVector3f		SafeNormalize(const FVector3f& v, const FVector3f& Fallback) const { const float l2 = v.SquaredLength(); return l2 < UE_KINDA_SMALL_NUMBER ? Fallback : v * (1.0f / FMath::Sqrt(l2)); }
-	FVector2f		SafeNormalize(const FVector2f& v) const { return SafeNormalize(v, FVector2f(1.0f, 0.0f)); }
-	FVector3f		SafeNormalize(const FVector3f& v) const { return SafeNormalize(v, FVector3f(1.0f, 0.0f, 0.0f)); }
+	static FVector2f	SafeNormalize(const FVector2f& v, const FVector2f& Fallback) { const float l2 = v.SquaredLength(); return l2 < UE_KINDA_SMALL_NUMBER ? Fallback : v * (1.0f / FMath::Sqrt(l2)); }
+	static FVector3f	SafeNormalize(const FVector3f& v, const FVector3f& Fallback) { const float l2 = v.SquaredLength(); return l2 < UE_KINDA_SMALL_NUMBER ? Fallback : v * (1.0f / FMath::Sqrt(l2)); }
+	static FVector2f	SafeNormalize(const FVector2f& v) { return SafeNormalize(v, FVector2f(1.0f, 0.0f)); }
+	static FVector3f	SafeNormalize(const FVector3f& v) { return SafeNormalize(v, FVector3f(1.0f, 0.0f, 0.0f)); }
+
+	static float		SafeRcpSqrt(const float v) { return v > UE_SMALL_NUMBER ? FMath::InvSqrt(v) : 0.0f; }
+	static FVector3f	SafeRcpSqrt(const FVector3f& v) { return FVector3f(SafeRcpSqrt(v.X), SafeRcpSqrt(v.Y), SafeRcpSqrt(v.Z)); }
 
 	void			ConvertRangeToScaleBias(const FNiagaraStatelessRangeInt&     Range, int32& OutScale, int32& OutBias) const { OutScale = Range.ParameterOffset == INDEX_NONE ? Range.GetScale() : 0; OutBias = Range.ParameterOffset == INDEX_NONE ? Range.Min : GetParameterBufferValue<int32>(Range.ParameterOffset, 0); }
 	void			ConvertRangeToScaleBias(const FNiagaraStatelessRangeFloat&   Range, float& OutScale, float& OutBias) const { OutScale = Range.ParameterOffset == INDEX_NONE ? Range.GetScale() : 0.0f; OutBias = Range.ParameterOffset == INDEX_NONE ? Range.Min : GetParameterBufferValue<float>(Range.ParameterOffset, 0); }
diff --git a/Engine/Plugins/FX/Niagara/Source/Niagara/Internal/Stateless/NiagaraStatelessSetShaderParameterContext.h b/Engine/Plugins/FX/Niagara/Source/Niagara/Internal/Stateless/NiagaraStatelessSetShaderParameterContext.h
--- a/Engine/Plugins/FX/Niagara/Source/Niagara/Internal/Stateless/NiagaraStatelessSetShaderParameterContext.h
+++ b/Engine/Plugins/FX/Niagara/Source/Niagara/Internal/Stateless/NiagaraStatelessSetShaderParameterContext.h
@@ -11,7 +11,7 @@ class FNiagaraStatelessSetShaderParameterContext
 public:
 	UE_NONCOPYABLE(FNiagaraStatelessSetShaderParameterContext);
 
-	explicit FNiagaraStatelessSetShaderParameterContext(const FNiagaraStatelessSpaceTransforms& InSpaceTransforms, TConstArrayView<uint8> InRendererParameterData, TConstArrayView<uint8> InBuiltData, const FShaderParametersMetadata* InShaderParametersMetadata, uint8* InShaderParameters);
+	explicit FNiagaraStatelessSetShaderParameterContext(const FNiagaraStatelessSpaceTransforms& InSpaceTransforms, TConstArrayView<uint8> InRendererParameterData, TConstArrayView<uint8> InBuiltData, TConstArrayView<uint8> InSharedBuiltData, TConstArrayView<TPair<FName, uint32>> InSharedBuiltDataOffsets, const FShaderParametersMetadata* InShaderParametersMetadata, uint8* InShaderParameters);
 
 	const FNiagaraStatelessSpaceTransforms& GetSpaceTransforms() const { return SpaceTransforms; }
 
@@ -35,6 +35,20 @@ public:
 		return reinterpret_cast<const T*>(BuiltData.GetData() + Offset);
 	}
 
+	template<typename T>
+	const T* ReadSharedBuiltData() const
+	{
+		const FName DataName = T::GetName();
+		for (const TPair<FName, uint32>& Pair : SharedBuiltDataOffsets)
+		{
+			if (Pair.Key == DataName)
+			{
+				return reinterpret_cast<const T*>(SharedBuiltData.GetData() + Pair.Value);
+			}
+		}
+		return nullptr;
+	}
+
 	template<typename TValueType>
 	TValueType GetRendererParameterValue(int32 Offset, const TValueType& DefaultValue) const
 	{
@@ -78,6 +92,8 @@ private:
 	TConstArrayView<uint8>					RendererParameterData;
 	TConstArrayView<uint8>					BuiltData;
 	mutable int32							BuiltDataOffset = 0;
+	TConstArrayView<uint8>					SharedBuiltData;
+	TConstArrayView<TPair<FName, uint32>>	SharedBuiltDataOffsets;
 	uint8*									ShaderParametersBase = nullptr;
 	mutable uint32							ParameterOffset = 0;
 	const FShaderParametersMetadata*		ShaderParametersMetadata = nullptr;
diff --git a/Engine/Plugins/FX/Niagara/Source/Niagara/Private/Stateless/Modules/NiagaraStatelessModuleCommon.cpp b/Engine/Plugins/FX/Niagara/Source/Niagara/Private/Stateless/Modules/NiagaraStatelessModuleCommon.cpp
new file mode 100644
--- /dev/null
+++ b/Engine/Plugins/FX/Niagara/Source/Niagara/Private/Stateless/Modules/NiagaraStatelessModuleCommon.cpp
@@ -0,0 +1,34 @@
+// Copyright Epic Games, Inc. All Rights Reserved.
+
+#include "Stateless/Modules/NiagaraStatelessModuleCommon.h"
+
+FName FNiagaraStatelessSystemScaleBuildData::GetName()
+{
+	return StaticStruct()->GetFName();
+}
+
+namespace NiagaraStateless
+{
+	FName FPhysicsBuildData::GetName()
+	{
+		return FName("FPhysicsBuildData");
+	}
+
+	FQuat4f DirectionToQuat(FVector3f Direction)
+	{
+		Direction = Direction.GetSafeNormal(UE_KINDA_SMALL_NUMBER, FVector3f::ZAxisVector);
+
+		const FVector3f Forward = FVector3f::ZAxisVector;
+		const float Dot = FVector3f::DotProduct(Forward, Direction);
+		if (Dot < -0.99999f)
+		{
+			return FQuat4f(FVector3f::YAxisVector, PI);
+		}
+
+		const FVector3f Axis = FVector3f::CrossProduct(Forward, Direction);
+		const float Angle = FMath::Acos(Dot);
+		FQuat4f Quat(Axis, Angle);
+		Quat.Normalize();
+		return Quat;
+	}
+}
diff --git a/Engine/Plugins/FX/Niagara/Source/Niagara/Private/Stateless/Modules/NiagaraStatelessModule_ApplyOwnerScaleToAttributes.cpp b/Engine/Plugins/FX/Niagara/Source/Niagara/Private/Stateless/Modules/NiagaraStatelessModule_ApplyOwnerScaleToAttributes.cpp
new file mode 100644
--- /dev/null
+++ b/Engine/Plugins/FX/Niagara/Source/Niagara/Private/Stateless/Modules/NiagaraStatelessModule_ApplyOwnerScaleToAttributes.cpp
@@ -0,0 +1,31 @@
+// Copyright Epic Games, Inc. All Rights Reserved.
+
+#include "Stateless/Modules/NiagaraStatelessModule_ApplyOwnerScaleToAttributes.h"
+#include "Stateless/NiagaraStatelessEmitterDataBuildContext.h"
+
+#include UE_INLINE_GENERATED_CPP_BY_NAME(NiagaraStatelessModule_ApplyOwnerScaleToAttributes)
+
+void UNiagaraStatelessModule_ApplyOwnerScaleToAttributes::BuildEmitterData(const FNiagaraStatelessEmitterDataBuildContext& BuildContext) const
+{
+	if (!IsModuleEnabled())
+	{
+		return;
+	}
+
+	const bool bAnyEnabled = 
+		SystemScaleData.bScaleCameraOffset ||
+		//SystemScaleData.bScaleDrag ||
+		SystemScaleData.bScaleForces ||
+		SystemScaleData.bScaleInitialVelocity ||
+		SystemScaleData.bScaleInitialRibbonWidth ||
+		//SystemScaleData.bScaleInitialMeshSize ||
+		SystemScaleData.bScaleInitialSpriteSize;
+
+	if (bAnyEnabled == false)
+	{
+		return;
+	}
+
+	FNiagaraStatelessSystemScaleBuildData* BuildData = BuildContext.AllocateSharedBuiltData<FNiagaraStatelessSystemScaleBuildData>();
+	*BuildData = SystemScaleData;
+}
diff --git a/Engine/Plugins/FX/Niagara/Source/Niagara/Private/Stateless/Modules/NiagaraStatelessModule_CameraOffset.cpp b/Engine/Plugins/FX/Niagara/Source/Niagara/Private/Stateless/Modules/NiagaraStatelessModule_CameraOffset.cpp
--- a/Engine/Plugins/FX/Niagara/Source/Niagara/Private/Stateless/Modules/NiagaraStatelessModule_CameraOffset.cpp
+++ b/Engine/Plugins/FX/Niagara/Source/Niagara/Private/Stateless/Modules/NiagaraStatelessModule_CameraOffset.cpp
@@ -2,6 +2,7 @@
 
 #include "Stateless/Modules/NiagaraStatelessModule_CameraOffset.h"
 
+#include "Stateless/Modules/NiagaraStatelessModuleCommon.h"
 #include "Stateless/NiagaraStatelessEmitterDataBuildContext.h"
 #include "Stateless/NiagaraStatelessModuleShaderParameters.h"
 #include "Stateless/NiagaraStatelessParticleSimContext.h"
@@ -24,13 +25,16 @@ namespace NSMCameraOffsetPrivate
 		using namespace NiagaraStateless;
 
 		const FModuleBuiltData* ModuleBuiltData = ParticleSimulationContext.ReadBuiltData<FModuleBuiltData>();
+		const FParameters* Parameters = ParticleSimulationContext.ReadParameterNestedStruct<FParameters>();
 		const float* NormalizedAgeData = ParticleSimulationContext.GetParticleNormalizedAge();
 		const float* PreviousNormalizedAgeData = ParticleSimulationContext.GetParticlePreviousNormalizedAge();
 
 		for (uint32 i = 0; i < ParticleSimulationContext.GetNumInstances(); ++i)
 		{
 			float CameraOffsets[2];
-			ParticleSimulationContext.SampleDistributionValues(ModuleBuiltData->DistributionParameters, i, 0, NormalizedAgeData[i], PreviousNormalizedAgeData[i], CameraOffsets[0], CameraOffsets[1]);
+			ParticleSimulationContext.SampleDistributionValues(Parameters->CameraOffset_Distribution, i, 0, NormalizedAgeData[i], PreviousNormalizedAgeData[i], CameraOffsets[0], CameraOffsets[1]);
+			CameraOffsets[0] *= Parameters->CameraOffset_Scale;
+			CameraOffsets[1] *= Parameters->CameraOffset_Scale;
 
 			ParticleSimulationContext.WriteParticleVariable(ModuleBuiltData->CameraVariableOffset, i, CameraOffsets[0]);
 			ParticleSimulationContext.WriteParticleVariable(ModuleBuiltData->PreviousCameraVariableOffset, i, CameraOffsets[1]);
@@ -71,7 +75,18 @@ void UNiagaraStatelessModule_CameraOffset::SetShaderParameters(const FNiagaraSta
 	const FModuleBuiltData* ModuleBuiltData = SetShaderParameterContext.ReadBuiltData<FModuleBuiltData>();
 
 	FParameters* Parameters = SetShaderParameterContext.GetParameterNestedStruct<FParameters>();
-	Parameters->CameraOffset_Distribution = ModuleBuiltData->DistributionParameters;
+	Parameters->CameraOffset_Distribution	= ModuleBuiltData->DistributionParameters;
+	Parameters->CameraOffset_Scale			= 1.0f;
+	if (const FNiagaraStatelessSystemScaleBuildData* SystemScaleBuildData = SetShaderParameterContext.ReadSharedBuiltData<FNiagaraStatelessSystemScaleBuildData>())
+	{
+		if (SystemScaleBuildData->bScaleCameraOffset)
+		{
+			const FNiagaraStatelessSpaceTransforms& SpaceTransforms = SetShaderParameterContext.GetSpaceTransforms();
+			const FVector3f ComponentScale = SpaceTransforms.GetLocalToWorld().GetScale3D();
+
+			Parameters->CameraOffset_Scale = FMath::Lerp(1.0f, ComponentScale.GetMax(), SystemScaleBuildData->ScaleCameraOffsetAmount);
+		}
+	}
 }
 
 #if WITH_EDITORONLY_DATA
diff --git a/Engine/Plugins/FX/Niagara/Source/Niagara/Private/Stateless/Modules/NiagaraStatelessModule_InitializeParticle.cpp b/Engine/Plugins/FX/Niagara/Source/Niagara/Private/Stateless/Modules/NiagaraStatelessModule_InitializeParticle.cpp
--- a/Engine/Plugins/FX/Niagara/Source/Niagara/Private/Stateless/Modules/NiagaraStatelessModule_InitializeParticle.cpp
+++ b/Engine/Plugins/FX/Niagara/Source/Niagara/Private/Stateless/Modules/NiagaraStatelessModule_InitializeParticle.cpp
@@ -139,6 +139,54 @@ void UNiagaraStatelessModule_InitializeParticle::SetShaderParameters(const FNiag
 	SetShaderParameterContext.ConvertRangeToScaleBias(ModuleBuiltData->SpriteRotationRange, Parameters->InitializeParticle_SpriteRotationScale, Parameters->InitializeParticle_SpriteRotationBias);
 	SetShaderParameterContext.ConvertRangeToScaleBias(ModuleBuiltData->MeshScaleRange,		Parameters->InitializeParticle_MeshScaleScale, Parameters->InitializeParticle_MeshScaleBias);
 	SetShaderParameterContext.ConvertRangeToScaleBias(ModuleBuiltData->RibbonWidthRange,	Parameters->InitializeParticle_RibbonWidthScale, Parameters->InitializeParticle_RibbonWidthBias);
+
+	if (const FNiagaraStatelessSystemScaleBuildData* SystemScaleBuildData = SetShaderParameterContext.ReadSharedBuiltData<FNiagaraStatelessSystemScaleBuildData>())
+	{
+		if (SystemScaleBuildData->bScaleInitialSpriteSize)
+		{
+			const FVector3f TransformScale = SetShaderParameterContext.GetSpaceTransforms().GetLocalToWorld().GetScale3D();
+			float Scale = 0.0f;
+			if (SystemScaleBuildData->bScaleInitialSpriteSizeBySingleAxis)
+			{
+				Scale = TransformScale.Component(int32(SystemScaleBuildData->ScaleInitialSpriteSizeAxis));
+			}
+			else
+			{
+				Scale = TransformScale.GetMax();
+			}
+			Scale = FMath::Lerp(1.0f, Scale, SystemScaleBuildData->ScaleInitialSpriteSizeScaleAmount);
+
+			Parameters->InitializeParticle_SpriteSizeScale *= Scale;
+			Parameters->InitializeParticle_SpriteSizeBias *= Scale;
+		}
+
+		//if (SystemScaleBuildData->bScaleInitialMeshSize)
+		//{
+		//	const FVector3f TransformScale = SetShaderParameterContext.GetSpaceTransforms().GetLocalToWorld().GetScale3D();
+		//	const FVector3f MeshSpaceScale = SystemScaleBuildData->bScaleInitialMeshSizeAbsoluteScale ? TransformScale.GetAbs() : TransformScale;
+		//	const FVector3f FinalScale = FMath::Lerp(FVector3f::OneVector, MeshSpaceScale, SystemScaleBuildData->ScaleInitialMeshSizeScaleAmount);
+		//	Parameters->InitializeParticle_MeshScaleScale *= FinalScale;
+		//	Parameters->InitializeParticle_MeshScaleBias *= FinalScale;
+		//}
+
+		if (SystemScaleBuildData->bScaleInitialRibbonWidth)
+		{
+			const FVector3f TransformScale = SetShaderParameterContext.GetSpaceTransforms().GetLocalToWorld().GetScale3D();
+			float Scale = 0.0f;
+			if (SystemScaleBuildData->bScaleInitialRibbonWidthBySingleAxis)
+			{
+				Scale = TransformScale.Component(int32(SystemScaleBuildData->ScaleInitialRibbonWidthScalingAxis));
+			}
+			else
+			{
+				Scale = TransformScale.GetMax();
+			}
+			Scale = FMath::Lerp(1.0f, Scale, SystemScaleBuildData->ScaleInitialRibbonWidthScaleAmount);
+
+			Parameters->InitializeParticle_RibbonWidthScale *= Scale;
+			Parameters->InitializeParticle_RibbonWidthBias *= Scale;
+		}
+	}
 }
 
 #if WITH_EDITORONLY_DATA
diff --git a/Engine/Plugins/FX/Niagara/Source/Niagara/Private/Stateless/Modules/NiagaraStatelessModule_SolveVelocitiesAndForces.cpp b/Engine/Plugins/FX/Niagara/Source/Niagara/Private/Stateless/Modules/NiagaraStatelessModule_SolveVelocitiesAndForces.cpp
--- a/Engine/Plugins/FX/Niagara/Source/Niagara/Private/Stateless/Modules/NiagaraStatelessModule_SolveVelocitiesAndForces.cpp
+++ b/Engine/Plugins/FX/Niagara/Source/Niagara/Private/Stateless/Modules/NiagaraStatelessModule_SolveVelocitiesAndForces.cpp
@@ -92,10 +92,10 @@ namespace NiagaraStatelessModuleSolveVelocitiesAndForcesPrivate
 
 		for (uint32 i = 0; i < ParticleSimulationContext.GetNumInstances(); ++i)
 		{
-			const float Mass				= ParticleSimulationContext.RandomScaleBiasFloat(i, 0, Parameters->SolveVelocitiesAndForces_MassScale, Parameters->SolveVelocitiesAndForces_MassBias);
-			const float Drag				= ParticleSimulationContext.RandomScaleBiasFloat(i, 1, Parameters->SolveVelocitiesAndForces_DragScale, Parameters->SolveVelocitiesAndForces_DragBias);
-			FVector3f InitialVelocity		= ParticleSimulationContext.RandomScaleBiasFloat(i, 2, Parameters->SolveVelocitiesAndForces_VelocityScale, Parameters->SolveVelocitiesAndForces_VelocityBias);
-			const FVector3f Wind			= ParticleSimulationContext.RandomScaleBiasFloat(i, 3, Parameters->SolveVelocitiesAndForces_WindScale, Parameters->SolveVelocitiesAndForces_WindBias);
+			const float Mass			= ParticleSimulationContext.RandomScaleBiasFloat(i, 0, Parameters->SolveVelocitiesAndForces_MassScale, Parameters->SolveVelocitiesAndForces_MassBias);
+			float Drag					= ParticleSimulationContext.RandomScaleBiasFloat(i, 1, Parameters->SolveVelocitiesAndForces_DragScale, Parameters->SolveVelocitiesAndForces_DragBias);
+			FVector3f InitialVelocity	= ParticleSimulationContext.RandomScaleBiasFloat(i, 2, Parameters->SolveVelocitiesAndForces_VelocityScale, Parameters->SolveVelocitiesAndForces_VelocityBias);
+			FVector3f Wind				= ParticleSimulationContext.RandomScaleBiasFloat(i, 3, Parameters->SolveVelocitiesAndForces_WindScale, Parameters->SolveVelocitiesAndForces_WindBias);
 
 			const float MassFactor	= 1 / (FMath::Max(Mass, UE_KINDA_SMALL_NUMBER)); // factor in mass for acceleration force
 			FVector3f Acceleration	= MassFactor * ParticleSimulationContext.RandomScaleBiasFloat(i, 4, Parameters->SolveVelocitiesAndForces_AccelerationScale, Parameters->SolveVelocitiesAndForces_AccelerationBias);
@@ -139,6 +139,12 @@ namespace NiagaraStatelessModuleSolveVelocitiesAndForcesPrivate
 				EvaluateNoise(NoiseLUT, Parameters, LifetimeData[i], PreviousNormalizedAgeData[i], NoiseFieldOffset, PreviousPosition);
 			}
 
+			// Apply owner scale before integration
+			//Drag			*= Parameters->SolveVelocitiesAndForces_OwnerDragScale;
+			InitialVelocity	*= Parameters->SolveVelocitiesAndForces_OwnerVelocityScale;
+			Wind			*= Parameters->SolveVelocitiesAndForces_OwnerAccelerationScale;
+			Acceleration	*= Parameters->SolveVelocitiesAndForces_OwnerAccelerationScale;
+
 			Position += IntegratePosition(FMath::Max(AgeData[i], UE_KINDA_SMALL_NUMBER), Drag, InitialVelocity, Wind, Acceleration);
 			PreviousPosition += IntegratePosition(PreviousAgeData[i], Drag, InitialVelocity, Wind, Acceleration);
 
@@ -194,6 +200,30 @@ void UNiagaraStatelessModule_SolveVelocitiesAndForces::SetShaderParameters(const
 	const FNiagaraStatelessSpaceTransforms& SpaceTransforms = SetShaderParameterContext.GetSpaceTransforms();
 
 	FParameters* Parameters = SetShaderParameterContext.GetParameterNestedStruct<FParameters>();
+	Parameters->SolveVelocitiesAndForces_OwnerAccelerationScale	= FVector3f::OneVector;
+	//Parameters->SolveVelocitiesAndForces_OwnerDragScale			= 1.0f;
+	Parameters->SolveVelocitiesAndForces_OwnerVelocityScale		= FVector3f::OneVector;
+	if (const FNiagaraStatelessSystemScaleBuildData* SystemScaleBuildData = SetShaderParameterContext.ReadSharedBuiltData<FNiagaraStatelessSystemScaleBuildData>())
+	{
+		if (SystemScaleBuildData->bScaleForces)
+		{
+			const FVector3f ComponentScale = SpaceTransforms.GetLocalToWorld().GetScale3D();
+			const FVector3f U = SystemScaleBuildData->bScaleForcesInverseScaling ? NiagaraStateless::FParticleSimulationContext::SafeRcpSqrt(SystemScaleBuildData->ScaleForcesScaleAmount) : SystemScaleBuildData->ScaleForcesScaleAmount;
+			Parameters->SolveVelocitiesAndForces_OwnerAccelerationScale = FMath::Lerp(FVector3f::OneVector, ComponentScale, U);
+		}
+		//if (SystemScaleBuildData->bScaleDrag)
+		//{
+		//	const float ComponentScale = SpaceTransforms.GetLocalToWorld().GetScale3D().GetMax();
+		//	const float U = SystemScaleBuildData->bInverseDragScaling ? NiagaraStateless::FParticleSimulationContext::SafeRcpSqrt(SystemScaleBuildData->ScaleDragScaleAmount) : SystemScaleBuildData->ScaleDragScaleAmount;
+		//	Parameters->SolveVelocitiesAndForces_OwnerDragScale = FMath::Lerp(1.0f, ComponentScale, U);
+		//}
+		if (SystemScaleBuildData->bScaleInitialVelocity)
+		{
+			const FVector3f ComponentScale = SpaceTransforms.GetLocalToWorld().GetScale3D();
+			Parameters->SolveVelocitiesAndForces_OwnerVelocityScale = FMath::Lerp(FVector3f::OneVector, ComponentScale, SystemScaleBuildData->ScaleInitialVelocityScaleAmount);
+		}
+	}
+
 	SetShaderParameterContext.ConvertRangeToScaleBias(PhysicsData.MassRange, Parameters->SolveVelocitiesAndForces_MassScale, Parameters->SolveVelocitiesAndForces_MassBias);
 	SetShaderParameterContext.ConvertRangeToScaleBias(PhysicsData.DragRange, Parameters->SolveVelocitiesAndForces_DragScale, Parameters->SolveVelocitiesAndForces_DragBias);
 
diff --git a/Engine/Plugins/FX/Niagara/Source/Niagara/Private/Stateless/NiagaraStatelessEmitter.cpp b/Engine/Plugins/FX/Niagara/Source/Niagara/Private/Stateless/NiagaraStatelessEmitter.cpp
--- a/Engine/Plugins/FX/Niagara/Source/Niagara/Private/Stateless/NiagaraStatelessEmitter.cpp
+++ b/Engine/Plugins/FX/Niagara/Source/Niagara/Private/Stateless/NiagaraStatelessEmitter.cpp
@@ -440,15 +440,7 @@ void UNiagaraStatelessEmitter::CacheFromCompiledData()
 			StatelessEmitterData->ParticleSimExecData = new NiagaraStateless::FParticleSimulationExecData(*StatelessEmitterData->ParticleDataSetCompiledData.Get());
 		}
 
-		FNiagaraStatelessEmitterDataBuildContext EmitterBuildContext(
-			*StatelessEmitterData->ParticleDataSetCompiledData.Get(),
-			StatelessEmitterData->RendererBindings,
-			StatelessEmitterData->Expressions,
-			StatelessEmitterData->BuiltData,
-			StatelessEmitterData->StaticDataBufferCpu,
-			StatelessEmitterData->ParticleSimExecData
-		);
-
+		FNiagaraStatelessEmitterDataBuildContext EmitterBuildContext(*StatelessEmitterData);
 		FNiagaraStatelessShaderParametersBuilder ShaderParametersBuilder;
 		ShaderParametersBuilder.AddParameterNestedStruct<NiagaraStateless::FCommonShaderParameters>();
 
@@ -696,6 +688,8 @@ NiagaraStateless::FCommonShaderParameters* UNiagaraStatelessEmitter::AllocateSha
 		SpaceTransforms,
 		RendererBindings.GetParameterDataArray(),
 		StatelessEmitterData->BuiltData,
+		StatelessEmitterData->SharedBuiltData,
+		StatelessEmitterData->SharedBuiltDataOffsets,
 		ShaderParametersMetadata,
 		static_cast<uint8*>(UntypedShaderParameters)
 	);
diff --git a/Engine/Plugins/FX/Niagara/Source/Niagara/Private/Stateless/NiagaraStatelessEmitterDataBuildContext.cpp b/Engine/Plugins/FX/Niagara/Source/Niagara/Private/Stateless/NiagaraStatelessEmitterDataBuildContext.cpp
--- a/Engine/Plugins/FX/Niagara/Source/Niagara/Private/Stateless/NiagaraStatelessEmitterDataBuildContext.cpp
+++ b/Engine/Plugins/FX/Niagara/Source/Niagara/Private/Stateless/NiagaraStatelessEmitterDataBuildContext.cpp
@@ -1,13 +1,13 @@
 // Copyright Epic Games, Inc. All Rights Reserved.
 
 #include "Stateless/NiagaraStatelessEmitterDataBuildContext.h"
+#include "Stateless/NiagaraStatelessEmitterData.h"
 #include "Stateless/NiagaraStatelessExpression.h"
 #include "Stateless/NiagaraStatelessParticleSimExecData.h"
 
 #include "NiagaraDataSetCompiledData.h"
 #include "NiagaraParameterBinding.h"
 #include "NiagaraParameterStore.h"
-#include "NiagaraParameterStore.h"
 
 namespace NiagaraStatelessEmitterDataBuildContextPrivate
 {
@@ -40,46 +40,81 @@ namespace NiagaraStatelessEmitterDataBuildContextPrivate
 	}
 };
 
+
+FNiagaraStatelessEmitterDataBuildContext::FNiagaraStatelessEmitterDataBuildContext(FNiagaraStatelessEmitterData& InEmitterData)
+	: EmitterData(InEmitterData)
+{
+	check(EmitterData.ParticleDataSetCompiledData.IsValid());
+}
+
+FNiagaraStatelessEmitterDataBuildContext::~FNiagaraStatelessEmitterDataBuildContext()
+{
+	EmitterData.Expressions.Shrink();
+	EmitterData.BuiltData.Shrink();
+	EmitterData.SharedBuiltData.Shrink();
+	EmitterData.StaticDataBufferCpu.Shrink();
+	//EmitterData.ParticleSimExecData.Shrink();
+}
+
 void FNiagaraStatelessEmitterDataBuildContext::PreModuleBuild(int32 InShaderParameterOffset)
 {
-	ModuleBuiltDataOffset = BuiltData.Num();
+	ModuleBuiltDataOffset = EmitterData.BuiltData.Num();
 	ShaderParameterOffset = InShaderParameterOffset;
 	++RandomSeedOffest;
 }
 
 uint32 FNiagaraStatelessEmitterDataBuildContext::AddStaticData(TConstArrayView<uint32> StaticData) const
 {
-	return NiagaraStatelessEmitterDataBuildContextPrivate::AddStaticData(StaticDataBuffer, StaticData.GetData(), StaticData.GetTypeSize() * StaticData.Num());
+	return NiagaraStatelessEmitterDataBuildContextPrivate::AddStaticData(EmitterData.StaticDataBufferCpu, StaticData.GetData(), StaticData.GetTypeSize() * StaticData.Num());
 }
 
 uint32 FNiagaraStatelessEmitterDataBuildContext::AddStaticData(TConstArrayView<int32> StaticData) const
 {
-	return NiagaraStatelessEmitterDataBuildContextPrivate::AddStaticData(StaticDataBuffer, StaticData.GetData(), StaticData.GetTypeSize() * StaticData.Num());
+	return NiagaraStatelessEmitterDataBuildContextPrivate::AddStaticData(EmitterData.StaticDataBufferCpu, StaticData.GetData(), StaticData.GetTypeSize() * StaticData.Num());
 }
 
 uint32 FNiagaraStatelessEmitterDataBuildContext::AddStaticData(TConstArrayView<float> StaticData) const
 {
-	return NiagaraStatelessEmitterDataBuildContextPrivate::AddStaticData(StaticDataBuffer, StaticData.GetData(), StaticData.GetTypeSize() * StaticData.Num());
+	return NiagaraStatelessEmitterDataBuildContextPrivate::AddStaticData(EmitterData.StaticDataBufferCpu, StaticData.GetData(), StaticData.GetTypeSize() * StaticData.Num());
 }
 
 uint32 FNiagaraStatelessEmitterDataBuildContext::AddStaticData(TConstArrayView<FVector2f> StaticData) const
 {
-	return NiagaraStatelessEmitterDataBuildContextPrivate::AddStaticData(StaticDataBuffer, StaticData.GetData(), StaticData.GetTypeSize() * StaticData.Num());
+	return NiagaraStatelessEmitterDataBuildContextPrivate::AddStaticData(EmitterData.StaticDataBufferCpu, StaticData.GetData(), StaticData.GetTypeSize() * StaticData.Num());
 }
 
 uint32 FNiagaraStatelessEmitterDataBuildContext::AddStaticData(TConstArrayView<FVector3f> StaticData) const
 {
-	return NiagaraStatelessEmitterDataBuildContextPrivate::AddStaticData(StaticDataBuffer, StaticData.GetData(), StaticData.GetTypeSize() * StaticData.Num());
+	return NiagaraStatelessEmitterDataBuildContextPrivate::AddStaticData(EmitterData.StaticDataBufferCpu, StaticData.GetData(), StaticData.GetTypeSize() * StaticData.Num());
 }
 
 uint32 FNiagaraStatelessEmitterDataBuildContext::AddStaticData(TConstArrayView<FVector4f> StaticData) const
 {
-	return NiagaraStatelessEmitterDataBuildContextPrivate::AddStaticData(StaticDataBuffer, StaticData.GetData(), StaticData.GetTypeSize() * StaticData.Num());
+	return NiagaraStatelessEmitterDataBuildContextPrivate::AddStaticData(EmitterData.StaticDataBufferCpu, StaticData.GetData(), StaticData.GetTypeSize() * StaticData.Num());
 }
 
 uint32 FNiagaraStatelessEmitterDataBuildContext::AddStaticData(TConstArrayView<FLinearColor> StaticData) const
 {
-	return NiagaraStatelessEmitterDataBuildContextPrivate::AddStaticData(StaticDataBuffer, StaticData.GetData(), StaticData.GetTypeSize() * StaticData.Num());
+	return NiagaraStatelessEmitterDataBuildContextPrivate::AddStaticData(EmitterData.StaticDataBufferCpu, StaticData.GetData(), StaticData.GetTypeSize() * StaticData.Num());
+}
+
+void* FNiagaraStatelessEmitterDataBuildContext::AllocateBuiltData(uint32 Size, uint32 Alingment) const
+{
+	const int32 Offset = Align(EmitterData.BuiltData.Num(), Alingment);
+	EmitterData.BuiltData.AddZeroed(Offset + Size - EmitterData.BuiltData.Num());
+	return EmitterData.BuiltData.GetData() + Offset;
+}
+
+void* FNiagaraStatelessEmitterDataBuildContext::AllocateSharedBuiltData(FName DataName, uint32 Size, uint32 Alingment) const
+{
+	check(EmitterData.SharedBuiltDataOffsets.ContainsByPredicate([DataName](const TPair<FName, uint32>& Pair) { return Pair.Key == DataName; }) == false);
+
+	const int32 Offset = Align(EmitterData.SharedBuiltData.Num(), Alingment);
+	EmitterData.SharedBuiltData.AddZeroed(Offset + Size - EmitterData.SharedBuiltData.Num());
+	void* NewData = EmitterData.SharedBuiltData.GetData() + Offset;
+
+	EmitterData.SharedBuiltDataOffsets.Emplace(DataName, reinterpret_cast<uintptr_t>(NewData) - reinterpret_cast<uintptr_t>(EmitterData.SharedBuiltData.GetData()));
+	return NewData;
 }
 
 int32 FNiagaraStatelessEmitterDataBuildContext::AddRendererBinding(const FNiagaraVariableBase& Variable) const
@@ -88,7 +123,7 @@ int32 FNiagaraStatelessEmitterDataBuildContext::AddRendererBinding(const FNiagar
 	if (Variable.IsValid())
 	{
 		FNiagaraVariable Var(Variable);
-		RendererBindings.AddParameter(Var, false, false, &DataOffset);
+		EmitterData.RendererBindings.AddParameter(Var, false, false, &DataOffset);
 		DataOffset /= sizeof(uint32);
 	}
 
@@ -101,7 +136,7 @@ int32 FNiagaraStatelessEmitterDataBuildContext::AddRendererBinding(const FNiagar
 	if (Binding.ResolvedParameter.IsValid())
 	{
 		FNiagaraVariable Var(Binding.ResolvedParameter);
-		RendererBindings.AddParameter(Var, false, false, &DataOffset);
+		EmitterData.RendererBindings.AddParameter(Var, false, false, &DataOffset);
 		DataOffset /= sizeof(uint32);
 	}
 
@@ -114,14 +149,14 @@ int32 FNiagaraStatelessEmitterDataBuildContext::AddRendererBinding(const FNiagar
 	if (Binding.ResolvedParameter.IsValid())
 	{
 		FNiagaraVariable Var(Binding.ResolvedParameter);
-		RendererBindings.AddParameter(Var, false, false, &DataOffset);
+		EmitterData.RendererBindings.AddParameter(Var, false, false, &DataOffset);
 
 		TConstArrayView<uint8> DefaultValue = Binding.GetDefaultValueArray();
 		if (DefaultValue.Num() > 0)
 		{
 			check(DataOffset != INDEX_NONE);
 			check(DefaultValue.Num() == Var.GetSizeInBytes());
-			RendererBindings.SetParameterData(DefaultValue.GetData(), DataOffset, DefaultValue.Num());
+			EmitterData.RendererBindings.SetParameterData(DefaultValue.GetData(), DataOffset, DefaultValue.Num());
 		}
 
 		DataOffset /= sizeof(uint32);
@@ -136,11 +171,11 @@ int32 FNiagaraStatelessEmitterDataBuildContext::AddExpression(const FInstancedSt
 	int32 DataOffset = INDEX_NONE;
 	if (const FNiagaraStatelessExpression* Expression = ExpressionStruct.GetPtr<FNiagaraStatelessExpression>())
 	{
-		const FName ExpressionName("__StatelessInternal__.Expression", Expressions.Num());
+		const FName ExpressionName("__StatelessInternal__.Expression", EmitterData.Expressions.Num());
 		const FNiagaraVariable ExpressionVariable(Expression->GetOutputTypeDef(), ExpressionName);
-		RendererBindings.AddParameter(ExpressionVariable, false, false, &DataOffset);
+		EmitterData.RendererBindings.AddParameter(ExpressionVariable, false, false, &DataOffset);
 
-		Expressions.Emplace(DataOffset, Expression->Build(*this));
+		EmitterData.Expressions.Emplace(DataOffset, Expression->Build(*this));
 		DataOffset /= sizeof(uint32);
 	}
 	return DataOffset;
@@ -148,15 +183,15 @@ int32 FNiagaraStatelessEmitterDataBuildContext::AddExpression(const FInstancedSt
 
 void FNiagaraStatelessEmitterDataBuildContext::AddParticleSimulationExecSimulate(TFunction<void(const NiagaraStateless::FParticleSimulationContext&)> Func) const
 {
-	if (!ParticleExecData)
+	if (!EmitterData.ParticleSimExecData)
 	{
 		return;
 	}
 
-	ParticleExecData->SimulateFunctions.Emplace(MoveTemp(Func), ModuleBuiltDataOffset, ShaderParameterOffset, RandomSeedOffest);
+	EmitterData.ParticleSimExecData->SimulateFunctions.Emplace(MoveTemp(Func), ModuleBuiltDataOffset, ShaderParameterOffset, RandomSeedOffest);
 }
 
 int32 FNiagaraStatelessEmitterDataBuildContext::FindParticleVariableIndex(const FNiagaraVariableBase& Variable) const
 {
-	return ParticleDataSet.Variables.IndexOfByKey(Variable);
+	return EmitterData.ParticleDataSetCompiledData->Variables.IndexOfByKey(Variable);
 }
diff --git a/Engine/Plugins/FX/Niagara/Source/Niagara/Private/Stateless/NiagaraStatelessEmitterTemplate.cpp b/Engine/Plugins/FX/Niagara/Source/Niagara/Private/Stateless/NiagaraStatelessEmitterTemplate.cpp
--- a/Engine/Plugins/FX/Niagara/Source/Niagara/Private/Stateless/NiagaraStatelessEmitterTemplate.cpp
+++ b/Engine/Plugins/FX/Niagara/Source/Niagara/Private/Stateless/NiagaraStatelessEmitterTemplate.cpp
@@ -13,6 +13,7 @@
 
 #include "Stateless/Modules/NiagaraStatelessModule_AddVelocity.h"
 #include "Stateless/Modules/NiagaraStatelessModule_AccelerationForce.h"
+#include "Stateless/Modules/NiagaraStatelessModule_ApplyOwnerScaleToAttributes.h"
 #include "Stateless/Modules/NiagaraStatelessModule_CalculateAccurateVelocity.h"
 #include "Stateless/Modules/NiagaraStatelessModule_CameraOffset.h"
 #include "Stateless/Modules/NiagaraStatelessModule_CurlNoiseForce.h"
@@ -238,6 +239,7 @@ void UNiagaraStatelessEmitterDefault::InitModulesAndAttributes()
 		UNiagaraStatelessModule_CurlNoiseForce::StaticClass(),
 		UNiagaraStatelessModule_Drag::StaticClass(),
 		UNiagaraStatelessModule_GravityForce::StaticClass(),
+		UNiagaraStatelessModule_ApplyOwnerScaleToAttributes::StaticClass(),
 		UNiagaraStatelessModule_SolveVelocitiesAndForces::StaticClass(),
 		// Post Solve
 		UNiagaraStatelessModule_CameraOffset::StaticClass(),
diff --git a/Engine/Plugins/FX/Niagara/Source/Niagara/Private/Stateless/NiagaraStatelessSetShaderParameterContext.cpp b/Engine/Plugins/FX/Niagara/Source/Niagara/Private/Stateless/NiagaraStatelessSetShaderParameterContext.cpp
--- a/Engine/Plugins/FX/Niagara/Source/Niagara/Private/Stateless/NiagaraStatelessSetShaderParameterContext.cpp
+++ b/Engine/Plugins/FX/Niagara/Source/Niagara/Private/Stateless/NiagaraStatelessSetShaderParameterContext.cpp
@@ -2,10 +2,12 @@
 
 #include "Stateless/NiagaraStatelessSetShaderParameterContext.h"
 
-FNiagaraStatelessSetShaderParameterContext::FNiagaraStatelessSetShaderParameterContext(const FNiagaraStatelessSpaceTransforms& InSpaceTransforms, TConstArrayView<uint8> InRendererParameterData, TConstArrayView<uint8> InBuiltData, const FShaderParametersMetadata* InShaderParametersMetadata, uint8* InShaderParameters)
+FNiagaraStatelessSetShaderParameterContext::FNiagaraStatelessSetShaderParameterContext(const FNiagaraStatelessSpaceTransforms& InSpaceTransforms, TConstArrayView<uint8> InRendererParameterData, TConstArrayView<uint8> InBuiltData, TConstArrayView<uint8> InSharedBuiltData, TConstArrayView<TPair<FName, uint32>> InSharedBuiltDataOffsets, const FShaderParametersMetadata* InShaderParametersMetadata, uint8* InShaderParameters)
 	: SpaceTransforms(InSpaceTransforms)
 	, RendererParameterData(InRendererParameterData)
 	, BuiltData(InBuiltData)
+	, SharedBuiltData(InSharedBuiltData)
+	, SharedBuiltDataOffsets(InSharedBuiltDataOffsets)
 	, ShaderParametersBase(InShaderParameters)
 	, ParameterOffset(0)
 	, ShaderParametersMetadata(InShaderParametersMetadata)
diff --git a/Engine/Plugins/FX/Niagara/Source/NiagaraShader/Internal/Stateless/NiagaraStatelessModuleShaderParameters.h b/Engine/Plugins/FX/Niagara/Source/NiagaraShader/Internal/Stateless/NiagaraStatelessModuleShaderParameters.h
--- a/Engine/Plugins/FX/Niagara/Source/NiagaraShader/Internal/Stateless/NiagaraStatelessModuleShaderParameters.h
+++ b/Engine/Plugins/FX/Niagara/Source/NiagaraShader/Internal/Stateless/NiagaraStatelessModuleShaderParameters.h
@@ -69,6 +69,7 @@ namespace NiagaraStateless
 
 	BEGIN_SHADER_PARAMETER_STRUCT(FCameraOffsetModule_ShaderParameters, NIAGARASHADER_API)
 		SHADER_PARAMETER(FNiagaraStatelessBuiltDistributionType,	CameraOffset_Distribution)
+		SHADER_PARAMETER(float,										CameraOffset_Scale)
 	END_SHADER_PARAMETER_STRUCT()
 
 	BEGIN_SHADER_PARAMETER_STRUCT(FScaleColorModule_ShaderParameters, NIAGARASHADER_API)
@@ -149,6 +150,9 @@ namespace NiagaraStateless
 	END_SHADER_PARAMETER_STRUCT()
 
 	BEGIN_SHADER_PARAMETER_STRUCT(FSolveVelocitiesAndForcesModule_ShaderParameters, NIAGARASHADER_API)
+		SHADER_PARAMETER(FVector3f,		SolveVelocitiesAndForces_OwnerAccelerationScale)
+		//SHADER_PARAMETER(float,			SolveVelocitiesAndForces_OwnerDragScale)
+		SHADER_PARAMETER(FVector3f,		SolveVelocitiesAndForces_OwnerVelocityScale)
 		SHADER_PARAMETER(float,			SolveVelocitiesAndForces_MassScale)
 		SHADER_PARAMETER(float,			SolveVelocitiesAndForces_MassBias)
 		SHADER_PARAMETER(float,			SolveVelocitiesAndForces_DragScale)
