- Enabled Async Factory Workflows for Niagara & Material Wizard
- Removed unused cruft/functions
- Added delegate for factory window to close the window automatically when focus is lost

diff --git a/Engine/Plugins/Editor/MaterialAssetWizard/Source/MaterialAssetWizard/Private/MaterialFactoryWithAssetWizard.cpp b/Engine/Plugins/Editor/MaterialAssetWizard/Source/MaterialAssetWizard/Private/MaterialFactoryWithAssetWizard.cpp
--- a/Engine/Plugins/Editor/MaterialAssetWizard/Source/MaterialAssetWizard/Private/MaterialFactoryWithAssetWizard.cpp
+++ b/Engine/Plugins/Editor/MaterialAssetWizard/Source/MaterialAssetWizard/Private/MaterialFactoryWithAssetWizard.cpp
@@ -21,15 +21,16 @@ bool UMaterialFactoryWithAssetWizardBase::ConfigureProperties()
 		return false;
 	}
 	
-	TArray<FAssetData> SelectedAssetData;
-	
+	// Needs to return true to allow async path
+	return true;
+}
+
+bool UMaterialFactoryWithAssetWizardBase::ConfigurePropertiesAsync(FOnFactoryConfigurePropertiesAsyncComplete OnComplete, FOnFactoryConfigurePropertiesAsyncCancelled OnCancelled)
+{	
 	FSoftObjectPath Path = GetDefault<UMaterialEditorSettings>()->MaterialWizardConfiguration;
 	UObject* ConfigUObject = Path.TryLoad();
 	if(UTaggedAssetBrowserConfiguration* ConfigurationAsset = Cast<UTaggedAssetBrowserConfiguration>(ConfigUObject))
-	{ 
-		IMainFrameModule& MainFrame = FModuleManager::LoadModuleChecked<IMainFrameModule>("MainFrame");
-		TSharedPtr<SWindow>	ParentWindow = MainFrame.GetParentWindow();
-
+	{
 		SWindow::FArguments WindowArgs = CreateWindowArguments();
 		STaggedAssetBrowser::FArguments AssetBrowserArgs = CreateTaggedAssetBrowserArguments();
 		
@@ -37,30 +38,36 @@ bool UMaterialFactoryWithAssetWizardBase::ConfigureProperties()
 		AssetBrowserWindowArgs.AssetBrowserArgs(AssetBrowserArgs);
 		AssetBrowserWindowArgs.WindowArgs(WindowArgs);
 		
+		STaggedAssetBrowserAssetFactoryWindow::FAsyncFactoryArguments AsyncArguments;
+		AsyncArguments.Factory.Reset(this);
+		AsyncArguments.OnAssetsActivated = FOnAssetsActivated::CreateUObject(this, &UMaterialFactoryWithAssetWizardBase::OnAssetsActivated);
+		AsyncArguments.OnFactoryConfigurePropertiesComplete = OnComplete;
+		AsyncArguments.OnFactoryConfigurePropertiesCancelled = OnCancelled;
+		
 		TSharedRef<STaggedAssetBrowserAssetFactoryWindow> CreateAssetBrowserWindow = SNew(STaggedAssetBrowserAssetFactoryWindow, *ConfigurationAsset, *SupportedClass.Get())
-			.AssetBrowserWindowArgs(AssetBrowserWindowArgs);
+			.AssetBrowserWindowArgs(AssetBrowserWindowArgs)
+			.AsyncFactoryArguments(AsyncArguments);
 		
-		FSlateApplication::Get().AddModalWindow(CreateAssetBrowserWindow, ParentWindow);
+		FSlateApplication::Get().AddWindow(CreateAssetBrowserWindow, true);
 
-		if(CreateAssetBrowserWindow->ShouldProceedWithAction() == false)
-		{
-			return false;	
-		}
-
-		SelectedAssetData = CreateAssetBrowserWindow->GetSelectedAssets();
+		return true;
 	}
 
-	if(SelectedAssetData.Num() == 1)
+	OnComplete.ExecuteIfBound(this);
+	return true;
+}
+
+void UMaterialFactoryWithAssetWizardBase::OnAssetsActivated(const TArray<FAssetData>& AssetData, EAssetTypeActivationMethod::Type Method)
+{
+	if(AssetData.Num() == 1)
 	{
-		FAssetData SelectedAsset = SelectedAssetData[0];
+		FAssetData SelectedAsset = AssetData[0];
 		ensure(SelectedAsset.GetClass() == UMaterial::StaticClass());
 		
 		UMaterial* MaterialAsset = Cast<UMaterial>(SelectedAsset.GetAsset());
 			
 		BaseMaterial = MaterialAsset;
 	}
-	
-	return true;
 }
 
 SWindow::FArguments UMaterialFactoryWithAssetWizardBase::CreateWindowArguments() const
diff --git a/Engine/Plugins/Editor/MaterialAssetWizard/Source/MaterialAssetWizard/Public/MaterialFactoryWithAssetWizard.h b/Engine/Plugins/Editor/MaterialAssetWizard/Source/MaterialAssetWizard/Public/MaterialFactoryWithAssetWizard.h
--- a/Engine/Plugins/Editor/MaterialAssetWizard/Source/MaterialAssetWizard/Public/MaterialFactoryWithAssetWizard.h
+++ b/Engine/Plugins/Editor/MaterialAssetWizard/Source/MaterialAssetWizard/Public/MaterialFactoryWithAssetWizard.h
@@ -20,9 +20,12 @@ public:
 	UMaterialFactoryWithAssetWizardBase() = default;
 	
 	virtual bool ConfigureProperties() override;
+	virtual bool ConfigurePropertiesAsync(FOnFactoryConfigurePropertiesAsyncComplete OnComplete, FOnFactoryConfigurePropertiesAsyncCancelled OnCancelled) override;
 	virtual bool ShouldShowInNewMenu() const override { return false; }
 
 private:
+	void OnAssetsActivated(const TArray<FAssetData>& AssetData, EAssetTypeActivationMethod::Type Method);
+	
 	SWindow::FArguments CreateWindowArguments() const;
 	STaggedAssetBrowser::FArguments CreateTaggedAssetBrowserArguments() const;
 	
diff --git a/Engine/Plugins/FX/Niagara/Source/NiagaraEditor/Private/NiagaraEmitterFactoryNew.cpp b/Engine/Plugins/FX/Niagara/Source/NiagaraEditor/Private/NiagaraEmitterFactoryNew.cpp
--- a/Engine/Plugins/FX/Niagara/Source/NiagaraEditor/Private/NiagaraEmitterFactoryNew.cpp
+++ b/Engine/Plugins/FX/Niagara/Source/NiagaraEditor/Private/NiagaraEmitterFactoryNew.cpp
@@ -17,11 +17,7 @@
 #include "NiagaraEmitterEditorData.h"
 #include "NiagaraNodeAssignment.h"
 #include "NiagaraRecentAndFavoritesManager.h"
-#include "UserAssetTagEditorUtilities.h"
 #include "Assets/TaggedAssetBrowserConfiguration.h"
-#include "Misc/MessageDialog.h"
-#include "Modules/ModuleManager.h"
-#include "Interfaces/IMainFrameModule.h"
 #include "Framework/Application/SlateApplication.h"
 #include "Widgets/STaggedAssetBrowser.h"
 #include "Widgets/STaggedAssetBrowserAssetFactoryWindow.h"
@@ -31,8 +27,7 @@
 
 #define LOCTEXT_NAMESPACE "NiagaraEmitterFactory"
 
-UNiagaraEmitterFactoryNew::UNiagaraEmitterFactoryNew(const FObjectInitializer& ObjectInitializer)
-	: Super(ObjectInitializer)
+UNiagaraEmitterFactoryNew::UNiagaraEmitterFactoryNew()
 {
 	SupportedClass = UNiagaraEmitter::StaticClass();
 	bEditAfterNew = true;
@@ -40,12 +35,34 @@ UNiagaraEmitterFactoryNew::UNiagaraEmitterFactoryNew(const FObjectInitializer& O
 	EmitterToCopy = nullptr;
 	bUseInheritance = false;
 	bAddDefaultModulesAndRenderersToEmptyEmitter = true;
+	SupportedWorkflows = (uint8) (EFactoryCreateWorkflow::Default | EFactoryCreateWorkflow::Asynchronous);
 }
 
 bool UNiagaraEmitterFactoryNew::ConfigureProperties()
-{	
-	TArray<FAssetData> SelectedAssetData;
+{
+	// By default, we use the default empty emitter, or a truly empty emitter.
+	// This gets potentially overridden in ConfigurePropertiesAsync
+	FSoftObjectPath MinimalEmitterPath = GetDefault<UNiagaraEditorSettings>()->DefaultEmptyEmitter;
+	if(MinimalEmitterPath.IsValid() && MinimalEmitterPath.IsAsset())
+	{
+		if(UNiagaraEmitter* MinimalEmitter = Cast<UNiagaraEmitter>(MinimalEmitterPath.TryLoad()))
+		{
+			EmitterToCopy = MinimalEmitter;
+			bUseInheritance = MinimalEmitter->bIsInheritable;
+		}
+	}
+	
+	if(EmitterToCopy.IsExplicitlyNull())
+	{
+		bUseInheritance = true;
+		bAddDefaultModulesAndRenderersToEmptyEmitter = false;
+	}
 	
+	return true;
+}
+
+bool UNiagaraEmitterFactoryNew::ConfigurePropertiesAsync(FOnFactoryConfigurePropertiesAsyncComplete OnComplete,	FOnFactoryConfigurePropertiesAsyncCancelled OnCancelled)
+{	
 	FSoftObjectPath Path = GetDefault<UNiagaraEditorSettings>()->EmitterAssetWizardConfiguration;
     UObject* ConfigUObject = Path.TryLoad();
 	if(UTaggedAssetBrowserConfiguration* ConfigurationAsset = Cast<UTaggedAssetBrowserConfiguration>(ConfigUObject))
@@ -53,9 +70,6 @@ bool UNiagaraEmitterFactoryNew::ConfigureProperties()
 		FDefaultDetailsTabConfiguration DefaultDetailsTabConfiguration;
 		DefaultDetailsTabConfiguration.bUseDefaultDetailsTab = true;
 		DefaultDetailsTabConfiguration.EmptySelectionMessage = LOCTEXT("EmptyEmitterFactorySelectionMessage", "Select an emitter as a starting point for your new emitter.\n");
- 
-		IMainFrameModule& MainFrame = FModuleManager::LoadModuleChecked<IMainFrameModule>("MainFrame");
-		TSharedPtr<SWindow>	ParentWindow = MainFrame.GetParentWindow();
 
 		STaggedAssetBrowser::FInterfaceOverrideProfiles OverrideProfiles;
 		OverrideProfiles.AssetViewOptionsProfileName = FName("TaggedAssetBrowser");
@@ -93,51 +107,24 @@ bool UNiagaraEmitterFactoryNew::ConfigureProperties()
 		STaggedAssetBrowserWindow::FArguments AssetBrowserWindowArgs;
 		AssetBrowserWindowArgs.AssetBrowserArgs(AssetBrowserArgs);
 		AssetBrowserWindowArgs.WindowArgs(WindowArgs);
+
+		STaggedAssetBrowserAssetFactoryWindow::FAsyncFactoryArguments AsyncArguments;
+		AsyncArguments.Factory.Reset(this);
+		AsyncArguments.OnAssetsActivated = FOnAssetsActivated::CreateUObject(this, &UNiagaraEmitterFactoryNew::OnAssetsActivated);
+		AsyncArguments.OnFactoryConfigurePropertiesComplete = OnComplete;
+		AsyncArguments.OnFactoryConfigurePropertiesCancelled = OnCancelled;
 		
 		TSharedRef<STaggedAssetBrowserAssetFactoryWindow> CreateAssetBrowserWindow = SNew(STaggedAssetBrowserAssetFactoryWindow, *ConfigurationAsset, *UNiagaraEmitter::StaticClass())
-			.AssetBrowserWindowArgs(AssetBrowserWindowArgs);
+			.AssetBrowserWindowArgs(AssetBrowserWindowArgs)
+			.AsyncFactoryArguments(AsyncArguments);
 		
-		FSlateApplication::Get().AddModalWindow(CreateAssetBrowserWindow, ParentWindow);
-
-		if(CreateAssetBrowserWindow->ShouldProceedWithAction() == false)
-		{
-			return false;	
-		}
-
-		SelectedAssetData = CreateAssetBrowserWindow->GetSelectedAssets();
-	}
-	
-	// Give the minimal emitter a chance to be selected if the user chose 'Minimal Emitter'
-	if(SelectedAssetData.Num() == 0)
-	{
-		FSoftObjectPath MinimalEmitterPath = GetDefault<UNiagaraEditorSettings>()->DefaultEmptyEmitter;
-		if(MinimalEmitterPath.IsValid() && MinimalEmitterPath.IsAsset())
-		{
-			if(UObject* MinimalEmitter = MinimalEmitterPath.TryLoad())
-			{
-				SelectedAssetData.Add(FAssetData(MinimalEmitter));
-			}
-		}
-	}
-
-	if(SelectedAssetData.Num() == 1)
-	{
-		FAssetData SelectedAsset = SelectedAssetData[0];
-		ensure(SelectedAsset.GetClass() == UNiagaraEmitter::StaticClass());
+		FSlateApplication::Get().AddWindow(CreateAssetBrowserWindow, true);
 		
-		UNiagaraEmitter* EmitterAsset = Cast<UNiagaraEmitter>(SelectedAsset.GetAsset());
-			
-		EmitterToCopy = EmitterAsset;
-		bUseInheritance = EmitterAsset->bIsInheritable;
-	}
-	else
-	{
-		// If we haven't selected an emitter asset, nor had a valid default empty emitter, we add a truly empty emitter
-		EmitterToCopy = nullptr;
-		bUseInheritance = true;
-		bAddDefaultModulesAndRenderersToEmptyEmitter = false;
+		return true;
 	}
-	
+
+	// If no browser configuration was found, we call OnComplete immediately, creating our asset
+	OnComplete.ExecuteIfBound(this);
 	return true;
 }
 
@@ -292,6 +279,42 @@ void UNiagaraEmitterFactoryNew::InitializeEmitter(UNiagaraEmitter* NewEmitter, b
 	}
 }
 
+void UNiagaraEmitterFactoryNew::OnAssetsActivated(const TArray<FAssetData>& AssetData, EAssetTypeActivationMethod::Type Method)
+{
+	TArray<FAssetData> AssetDataCpy = AssetData;
+	
+	// Give the minimal emitter a chance to be selected if the user chose 'Minimal Emitter'
+	if(AssetDataCpy.Num() == 0)
+	{
+		FSoftObjectPath MinimalEmitterPath = GetDefault<UNiagaraEditorSettings>()->DefaultEmptyEmitter;
+		if(MinimalEmitterPath.IsValid() && MinimalEmitterPath.IsAsset())
+		{
+			if(UObject* MinimalEmitter = MinimalEmitterPath.TryLoad())
+			{
+				AssetDataCpy.Add(FAssetData(MinimalEmitter));
+			}
+		}
+	}
+
+	if(AssetDataCpy.Num() == 1)
+	{
+		FAssetData SelectedAsset = AssetDataCpy[0];
+		ensure(SelectedAsset.GetClass() == UNiagaraEmitter::StaticClass());
+		
+		UNiagaraEmitter* EmitterAsset = Cast<UNiagaraEmitter>(SelectedAsset.GetAsset());
+			
+		EmitterToCopy = EmitterAsset;
+		bUseInheritance = EmitterAsset->bIsInheritable;
+	}
+	else
+	{
+		// If we haven't selected an emitter asset, nor had a valid default empty emitter, we add a truly empty emitter
+		EmitterToCopy = nullptr;
+		bUseInheritance = true;
+		bAddDefaultModulesAndRenderersToEmptyEmitter = false;
+	}
+}
+
 bool UNiagaraEmitterFactoryNew::OnAdditionalShouldFilterAsset(const FAssetData& AssetData) const
 {
 	const UNiagaraEditorSettings* Settings = GetDefault<UNiagaraEditorSettings>();
diff --git a/Engine/Plugins/FX/Niagara/Source/NiagaraEditor/Private/NiagaraSystemFactoryNew.cpp b/Engine/Plugins/FX/Niagara/Source/NiagaraEditor/Private/NiagaraSystemFactoryNew.cpp
--- a/Engine/Plugins/FX/Niagara/Source/NiagaraEditor/Private/NiagaraSystemFactoryNew.cpp
+++ b/Engine/Plugins/FX/Niagara/Source/NiagaraEditor/Private/NiagaraSystemFactoryNew.cpp
@@ -11,16 +11,12 @@
 #include "EditorDirectories.h"
 #include "ImageUtils.h"
 #include "ViewModels/Stack/NiagaraStackGraphUtilities.h"
-#include "Misc/MessageDialog.h"
 #include "ViewModels/NiagaraSystemViewModel.h"
-#include "Modules/ModuleManager.h"
-#include "Interfaces/IMainFrameModule.h"
 #include "Framework/Application/SlateApplication.h"
 #include "NiagaraEditorUtilities.h"
 #include "NiagaraRecentAndFavoritesManager.h"
 #include "NiagaraSettings.h"
 #include "ObjectTools.h"
-#include "UserAssetTagEditorUtilities.h"
 #include "Assets/TaggedAssetBrowserConfiguration.h"
 #include "Widgets/STaggedAssetBrowser.h"
 #include "Widgets/STaggedAssetBrowserAssetFactoryWindow.h"
@@ -30,25 +26,26 @@
 
 #define LOCTEXT_NAMESPACE "NiagaraSystemFactory"
 
-UNiagaraSystemFactoryNew::UNiagaraSystemFactoryNew(const FObjectInitializer& ObjectInitializer)
-: Super(ObjectInitializer)
+UNiagaraSystemFactoryNew::UNiagaraSystemFactoryNew()
 {
 	SupportedClass = UNiagaraSystem::StaticClass();
 	bEditAfterNew = true;
 	bCreateNew = true;
+	SupportedWorkflows = (uint8) (EFactoryCreateWorkflow::Default | EFactoryCreateWorkflow::Asynchronous);
 }
 
 bool UNiagaraSystemFactoryNew::ConfigureProperties()
 {
-	TArray<FAssetData> SelectedAssetData;
-	
+	// Needs to return true to allow async path
+	return true;
+}
+
+bool UNiagaraSystemFactoryNew::ConfigurePropertiesAsync(FOnFactoryConfigurePropertiesAsyncComplete OnComplete, FOnFactoryConfigurePropertiesAsyncCancelled OnCancelled)
+{	
 	FSoftObjectPath Path = GetDefault<UNiagaraEditorSettings>()->SystemAssetWizardConfiguration;
 	UObject* ConfigUObject = Path.TryLoad();
 	if(UTaggedAssetBrowserConfiguration* ConfigurationAsset = Cast<UTaggedAssetBrowserConfiguration>(ConfigUObject))
-	{
-		IMainFrameModule& MainFrame = FModuleManager::LoadModuleChecked<IMainFrameModule>("MainFrame");
-		TSharedPtr<SWindow>	ParentWindow = MainFrame.GetParentWindow();
-		
+	{		
 		STaggedAssetBrowser::FInterfaceOverrideProfiles OverrideProfiles;
 		OverrideProfiles.AssetViewOptionsProfileName = FName("TaggedAssetBrowser");
 		OverrideProfiles.DefaultFilterMenuExpansion = EAssetTypeCategories::FX;
@@ -88,36 +85,24 @@ bool UNiagaraSystemFactoryNew::ConfigureProperties()
 		STaggedAssetBrowserWindow::FArguments AssetBrowserWindowArgs;
 		AssetBrowserWindowArgs.AssetBrowserArgs(AssetBrowserArgs);
 		AssetBrowserWindowArgs.WindowArgs(WindowArgs);
+
+		STaggedAssetBrowserAssetFactoryWindow::FAsyncFactoryArguments AsyncArguments;
+		AsyncArguments.Factory.Reset(this);
+		AsyncArguments.OnAssetsActivated = FOnAssetsActivated::CreateUObject(this, &UNiagaraSystemFactoryNew::OnAssetsActivated);
+		AsyncArguments.OnFactoryConfigurePropertiesComplete = OnComplete;
+		AsyncArguments.OnFactoryConfigurePropertiesCancelled = OnCancelled;
 		
 		TSharedRef<STaggedAssetBrowserAssetFactoryWindow> CreateAssetBrowserWindow = SNew(STaggedAssetBrowserAssetFactoryWindow, *ConfigurationAsset, *UNiagaraSystem::StaticClass())
-			.AssetBrowserWindowArgs(AssetBrowserWindowArgs);
+			.AssetBrowserWindowArgs(AssetBrowserWindowArgs)
+			.AsyncFactoryArguments(AsyncArguments);
 
-		FSlateApplication::Get().AddModalWindow(CreateAssetBrowserWindow, ParentWindow);
+		FSlateApplication::Get().AddWindow(CreateAssetBrowserWindow, true);
 
-		if(CreateAssetBrowserWindow->ShouldProceedWithAction() == false)
-		{
-			return false;
-		}
-
-		SelectedAssetData = CreateAssetBrowserWindow->GetSelectedAssets();
-	}
-	
-	if(SelectedAssetData.Num() == 1)
-	{
-		FAssetData AssetData = SelectedAssetData[0];
-		if(AssetData.GetClass() == UNiagaraSystem::StaticClass())
-		{
-			SystemToCopy = Cast<UNiagaraSystem>(AssetData.GetAsset());
-		}
-		else if(AssetData.GetClass() == UNiagaraEmitter::StaticClass())
-		{
-			UNiagaraEmitter* EmitterAsset = Cast<UNiagaraEmitter>(AssetData.GetAsset());
-			FVersionedNiagaraEmitter VersionedNiagaraEmitter(EmitterAsset, EmitterAsset->GetExposedVersion().VersionGuid);
-			
-			EmittersToAddToNewSystem.Add(VersionedNiagaraEmitter);
-		}
+		return true;
 	}
-	
+
+	// If no browser configuration was found, we call OnComplete immediately, creating our asset
+	OnComplete.ExecuteIfBound(this);
 	return true;
 }
 
@@ -130,7 +115,7 @@ UObject* UNiagaraSystemFactoryNew::FactoryCreateNew(UClass* Class, UObject* InPa
 
 	UNiagaraSystem* NewSystem;
 	
-	if (SystemToCopy != nullptr)
+	if (SystemToCopy.IsValid())
 	{
 		if (SystemToCopy->IsReadyToRun() == false)
 		{
@@ -139,7 +124,7 @@ UObject* UNiagaraSystemFactoryNew::FactoryCreateNew(UClass* Class, UObject* InPa
 		
 		FNiagaraEditorModule::Get().GetRecentsManager()->SystemUsed(*SystemToCopy);
 
-		NewSystem = Cast<UNiagaraSystem>(StaticDuplicateObject(SystemToCopy, InParent, Name, Flags, Class));
+		NewSystem = Cast<UNiagaraSystem>(StaticDuplicateObject(SystemToCopy.Get(), InParent, Name, Flags, Class));
 		NewSystem->TemplateAssetDescription = FText();
 		NewSystem->Category = FText();
 		//NewSystem->AssetTags.Empty();
@@ -237,6 +222,25 @@ void UNiagaraSystemFactoryNew::TryAssignDefaultEffectType(UNiagaraSystem* System
 	}
 }
 
+void UNiagaraSystemFactoryNew::OnAssetsActivated(const TArray<FAssetData>& AssetData, EAssetTypeActivationMethod::Type ActivationMethod)
+{
+	if(AssetData.Num() == 1)
+	{
+		const FAssetData& FirstAssetData = AssetData[0];
+		if(FirstAssetData.GetClass() == UNiagaraSystem::StaticClass())
+		{
+			SystemToCopy = Cast<UNiagaraSystem>(FirstAssetData.GetAsset());
+		}
+		else if(FirstAssetData.GetClass() == UNiagaraEmitter::StaticClass())
+		{
+			UNiagaraEmitter* EmitterAsset = Cast<UNiagaraEmitter>(FirstAssetData.GetAsset());
+			FVersionedNiagaraEmitter VersionedNiagaraEmitter(EmitterAsset, EmitterAsset->GetExposedVersion().VersionGuid);
+				
+			EmittersToAddToNewSystem.Add(VersionedNiagaraEmitter);
+		}
+	}
+}
+
 bool UNiagaraSystemFactoryNew::OnAdditionalShouldFilterAsset(const FAssetData& AssetData) const
 {
 	const UNiagaraEditorSettings* Settings = GetDefault<UNiagaraEditorSettings>();
diff --git a/Engine/Plugins/FX/Niagara/Source/NiagaraEditor/Public/NiagaraEmitterFactoryNew.h b/Engine/Plugins/FX/Niagara/Source/NiagaraEditor/Public/NiagaraEmitterFactoryNew.h
--- a/Engine/Plugins/FX/Niagara/Source/NiagaraEditor/Public/NiagaraEmitterFactoryNew.h
+++ b/Engine/Plugins/FX/Niagara/Source/NiagaraEditor/Public/NiagaraEmitterFactoryNew.h
@@ -12,14 +12,19 @@
 UCLASS(hidecategories = Object, MinimalAPI)
 class UNiagaraEmitterFactoryNew : public UFactory
 {
-	GENERATED_UCLASS_BODY()
+	GENERATED_BODY()
 
+public:
+	UNiagaraEmitterFactoryNew();
+	
 	TWeakObjectPtr<UNiagaraEmitter> EmitterToCopy;
 	bool bUseInheritance;
 	bool bAddDefaultModulesAndRenderersToEmptyEmitter;
 
+private:
 	//~ Begin UFactory Interface
 	virtual bool ConfigureProperties() override;
+	virtual bool ConfigurePropertiesAsync(FOnFactoryConfigurePropertiesAsyncComplete OnComplete, FOnFactoryConfigurePropertiesAsyncCancelled OnCancelled) override;
 	virtual UObject* FactoryCreateNew(UClass* Class, UObject* InParent, FName Name, EObjectFlags Flags, UObject* Context, FFeedbackContext* Warn) override;
 	//~ Begin UFactory Interface	
 
@@ -27,6 +32,7 @@ public:
 	NIAGARAEDITOR_API static void InitializeEmitter(UNiagaraEmitter* NewEmitter, bool bAddDefaultModulesAndRenderers);
 
 private:
+	void OnAssetsActivated(const TArray<FAssetData>& AssetData, EAssetTypeActivationMethod::Type Method);
 	bool OnAdditionalShouldFilterAsset(const FAssetData& AssetData) const;
 	void OnExtendAddFilterMenu(UToolMenu* ToolMenu) const;
 	TArray<TSharedRef<FFrontendFilter>> OnGetExtraFrontendFilters() const;
diff --git a/Engine/Plugins/FX/Niagara/Source/NiagaraEditor/Public/NiagaraSystemFactoryNew.h b/Engine/Plugins/FX/Niagara/Source/NiagaraEditor/Public/NiagaraSystemFactoryNew.h
--- a/Engine/Plugins/FX/Niagara/Source/NiagaraEditor/Public/NiagaraSystemFactoryNew.h
+++ b/Engine/Plugins/FX/Niagara/Source/NiagaraEditor/Public/NiagaraSystemFactoryNew.h
@@ -15,13 +15,18 @@ class UNiagaraEmitter;
 UCLASS(hidecategories = Object)
 class UNiagaraSystemFactoryNew : public UFactory
 {
-	GENERATED_UCLASS_BODY()
+	GENERATED_BODY()
 
-	UNiagaraSystem* SystemToCopy = nullptr;
+public:
+	UNiagaraSystemFactoryNew();
+	
+	TWeakObjectPtr<UNiagaraSystem> SystemToCopy;
 	TArray<FVersionedNiagaraEmitter> EmittersToAddToNewSystem;
 
+private:
 	//~ Begin UFactory Interface
 	virtual bool ConfigureProperties() override;
+	virtual bool ConfigurePropertiesAsync(FOnFactoryConfigurePropertiesAsyncComplete OnComplete, FOnFactoryConfigurePropertiesAsyncCancelled OnCancelled) override;
 	virtual UObject* FactoryCreateNew(UClass* Class, UObject* InParent, FName Name, EObjectFlags Flags, UObject* Context, FFeedbackContext* Warn) override;
 	//~ Begin UFactory Interface	
 
@@ -30,7 +35,8 @@ public:
 
 private:
 	void TryAssignDefaultEffectType(UNiagaraSystem* System);
-	
+
+	void OnAssetsActivated(const TArray<FAssetData>& AssetData, EAssetTypeActivationMethod::Type ActivationMethod);
 	bool OnAdditionalShouldFilterAsset(const FAssetData& AssetData) const;
 	void OnExtendAddFilterMenu(UToolMenu* ToolMenu) const;
 	TArray<TSharedRef<FFrontendFilter>> OnGetExtraFrontendFilters() const;
diff --git a/Engine/Source/Editor/UserAssetTagsEditor/Private/Widgets/STaggedAssetBrowser.cpp b/Engine/Source/Editor/UserAssetTagsEditor/Private/Widgets/STaggedAssetBrowser.cpp
--- a/Engine/Source/Editor/UserAssetTagsEditor/Private/Widgets/STaggedAssetBrowser.cpp
+++ b/Engine/Source/Editor/UserAssetTagsEditor/Private/Widgets/STaggedAssetBrowser.cpp
@@ -1245,11 +1245,6 @@ TArray<FAssetData> STaggedAssetBrowserWindow::GetSelectedAssets() const
 	return AssetBrowser->GetSelectedAssets();
 }
 
-void STaggedAssetBrowserWindow::OnFocusLost(const FFocusEvent& InFocusEvent)
-{
-	SWindow::OnFocusLost(InFocusEvent);
-}
-
 FReply STaggedAssetBrowserWindow::OnKeyDown(const FGeometry& MyGeometry, const FKeyEvent& InKeyEvent)
 {
 	FReply Reply = SWindow::OnKeyDown(MyGeometry, InKeyEvent);
diff --git a/Engine/Source/Editor/UserAssetTagsEditor/Private/Widgets/STaggedAssetBrowserAssetFactoryWindow.cpp b/Engine/Source/Editor/UserAssetTagsEditor/Private/Widgets/STaggedAssetBrowserAssetFactoryWindow.cpp
--- a/Engine/Source/Editor/UserAssetTagsEditor/Private/Widgets/STaggedAssetBrowserAssetFactoryWindow.cpp
+++ b/Engine/Source/Editor/UserAssetTagsEditor/Private/Widgets/STaggedAssetBrowserAssetFactoryWindow.cpp
@@ -23,6 +23,8 @@ void STaggedAssetBrowserAssetFactoryWindow::Construct(const FArguments& InArgs,
 	{
 		AssetDefinition = UAssetDefinitionRegistry::Get()->GetAssetDefinitionForClass(CreatedClass.Get());
 	}
+
+	FSlateApplication::Get().OnFocusChanging().AddSP(this, &STaggedAssetBrowserAssetFactoryWindow::CloseIfOtherWindowFocused);
 	
 	AsyncFactoryArguments = InArgs._AsyncFactoryArguments;
 	AdditionalFactorySettingsClass = InArgs._AdditionalFactorySettingsClass;
@@ -116,6 +118,11 @@ void STaggedAssetBrowserAssetFactoryWindow::Construct(const FArguments& InArgs,
 	STaggedAssetBrowserWindow::Construct(Args._AssetBrowserWindowArgs, Configuration);
 }
 
+STaggedAssetBrowserAssetFactoryWindow::~STaggedAssetBrowserAssetFactoryWindow()
+{
+	FSlateApplication::Get().OnFocusChanging().RemoveAll(this);
+}
+
 bool STaggedAssetBrowserAssetFactoryWindow::CheckValidArguments(const FArguments& InArgs) const
 {
 	if(!ensureMsgf(!InArgs._AssetBrowserWindowArgs._AssetBrowserArgs._OnAssetsActivated.IsBound(), TEXT("OnAssetsActivated should not be bound.")))
@@ -217,4 +224,13 @@ void STaggedAssetBrowserAssetFactoryWindow::AddReferencedObjects(FReferenceColle
 	}
 }
 
+void STaggedAssetBrowserAssetFactoryWindow::CloseIfOtherWindowFocused(const FFocusEvent& FocusEvent, const FWeakWidgetPath& OldWidgetPath,	const TSharedPtr<SWidget>& OldFocusedWidget, const FWidgetPath& NewWidgetPath, const TSharedPtr<SWidget>& NewFocusedWidget)
+{
+	// If the focus path no longer contains our window, close it
+	if (!NewWidgetPath.ContainsWidget(this))
+	{
+		Cancel();
+	}
+}
+
 #undef LOCTEXT_NAMESPACE
diff --git a/Engine/Source/Editor/UserAssetTagsEditor/Public/Widgets/STaggedAssetBrowser.h b/Engine/Source/Editor/UserAssetTagsEditor/Public/Widgets/STaggedAssetBrowser.h
--- a/Engine/Source/Editor/UserAssetTagsEditor/Public/Widgets/STaggedAssetBrowser.h
+++ b/Engine/Source/Editor/UserAssetTagsEditor/Public/Widgets/STaggedAssetBrowser.h
@@ -316,9 +316,7 @@ public:
 	
 	UE_API TArray<FAssetData> GetSelectedAssets() const;
 	
-protected:
-	UE_API virtual void OnFocusLost(const FFocusEvent& InFocusEvent) override;
-	
+protected:	
 	UE_API virtual FReply OnKeyDown(const FGeometry& MyGeometry, const FKeyEvent& InKeyEvent) override;
 	
 protected:
diff --git a/Engine/Source/Editor/UserAssetTagsEditor/Public/Widgets/STaggedAssetBrowserAssetFactoryWindow.h b/Engine/Source/Editor/UserAssetTagsEditor/Public/Widgets/STaggedAssetBrowserAssetFactoryWindow.h
--- a/Engine/Source/Editor/UserAssetTagsEditor/Public/Widgets/STaggedAssetBrowserAssetFactoryWindow.h
+++ b/Engine/Source/Editor/UserAssetTagsEditor/Public/Widgets/STaggedAssetBrowserAssetFactoryWindow.h
@@ -41,7 +41,7 @@ public:
 	UE_API void Construct(const FArguments& InArgs, const UTaggedAssetBrowserConfiguration& Configuration, UClass& CreatedClass);
 
 	STaggedAssetBrowserAssetFactoryWindow() {}
-	virtual ~STaggedAssetBrowserAssetFactoryWindow() override {}
+	UE_API virtual ~STaggedAssetBrowserAssetFactoryWindow() override;
 
 	/** For use in a modal context: retrieve whether the user activated assets or clicked a Proceed button, or closed the window.*/
 	bool ShouldProceedWithAction() const { return bProceedWithAction; }
@@ -66,6 +66,9 @@ private:
 
 	UE_API virtual FString GetReferencerName() const override;
 	UE_API virtual void AddReferencedObjects(FReferenceCollector& Collector) override;
+	
+	void CloseIfOtherWindowFocused(const FFocusEvent& FocusEvent, const FWeakWidgetPath& OldWidgetPath, const TSharedPtr<SWidget>& OldFocusedWidget, const FWidgetPath& NewWidgetPath, const TSharedPtr<SWidget>& NewFocusedWidget);
+
 private:
 	TWeakObjectPtr<UClass> CreatedClass;
 	TWeakObjectPtr<UClass> AdditionalFactorySettingsClass;
