- Scale by X modules fix for CPU path when using approximate motion vectors

diff --git a/Engine/Plugins/FX/Niagara/Source/Niagara/Private/Stateless/Modules/NiagaraStatelessModule_ScaleMeshSizeBySpeed.cpp b/Engine/Plugins/FX/Niagara/Source/Niagara/Private/Stateless/Modules/NiagaraStatelessModule_ScaleMeshSizeBySpeed.cpp
--- a/Engine/Plugins/FX/Niagara/Source/Niagara/Private/Stateless/Modules/NiagaraStatelessModule_ScaleMeshSizeBySpeed.cpp
+++ b/Engine/Plugins/FX/Niagara/Source/Niagara/Private/Stateless/Modules/NiagaraStatelessModule_ScaleMeshSizeBySpeed.cpp
@@ -17,6 +17,7 @@ namespace NSMScaleMeshSizeBySpeedPrivate
 
 		int32							PositionVariableOffset = INDEX_NONE;
 		int32							PreviousPositionVariableOffset = INDEX_NONE;
+		int32							VelocityVariableOffset = INDEX_NONE;
 		int32							ScaleVariableVariableOffset = INDEX_NONE;
 		int32							PreviousScaleVariableVariableOffset = INDEX_NONE;
 	};
@@ -31,9 +32,17 @@ namespace NSMScaleMeshSizeBySpeedPrivate
 
 		for (uint32 i = 0; i < ParticleSimulationContext.GetNumInstances(); ++i)
 		{
-			const FVector3f Position = ParticleSimulationContext.ReadParticleVariable(ModuleBuiltData->PositionVariableOffset, i, FVector3f::OneVector);
-			const FVector3f PreviousPosition = ParticleSimulationContext.ReadParticleVariable(ModuleBuiltData->PreviousPositionVariableOffset, i, FVector3f::OneVector);
-			const FVector3f Velocity = (Position - PreviousPosition) * ParticleSimulationContext.GetInvDeltaTime();
+			const FVector3f Position = ParticleSimulationContext.ReadParticleVariable(ModuleBuiltData->PositionVariableOffset, i, FVector3f::ZeroVector);
+			FVector3f Velocity;
+			if (ModuleBuiltData->PreviousPositionVariableOffset != INDEX_NONE)
+			{
+				const FVector3f PreviousPosition = ParticleSimulationContext.ReadParticleVariable(ModuleBuiltData->PreviousPositionVariableOffset, i, Position);
+				Velocity = (Position - PreviousPosition) * ParticleSimulationContext.GetInvDeltaTime();
+			}
+			else
+			{
+				Velocity = ParticleSimulationContext.ReadParticleVariable(ModuleBuiltData->VelocityVariableOffset, i, FVector3f::ZeroVector);
+			}
 			const float Speed = Velocity.SquaredLength();
 			const float NormSpeed = FMath::Clamp(Speed * Parameters->ScaleMeshSizeBySpeed_VelocityNorm, 0.0f, 1.0f);
 			const float Interp = ParticleSimulationContext.SampleDistributionValue<float>(Parameters->ScaleMeshSizeBySpeed_ScaleDistribution, i, 0, NormSpeed);
@@ -64,6 +73,7 @@ void UNiagaraStatelessModule_ScaleMeshSizeBySpeed::BuildEmitterData(const FNiaga
 	const FNiagaraStatelessGlobals& StatelessGlobals = FNiagaraStatelessGlobals::Get();
 	BuiltData->PositionVariableOffset				= BuildContext.FindParticleVariableIndex(StatelessGlobals.PositionVariable);
 	BuiltData->PreviousPositionVariableOffset		= BuildContext.FindParticleVariableIndex(StatelessGlobals.PreviousPositionVariable);
+	BuiltData->VelocityVariableOffset				= BuildContext.FindParticleVariableIndex(StatelessGlobals.VelocityVariable);
 	BuiltData->ScaleVariableVariableOffset			= BuildContext.FindParticleVariableIndex(StatelessGlobals.ScaleVariable);
 	BuiltData->PreviousScaleVariableVariableOffset	= BuildContext.FindParticleVariableIndex(StatelessGlobals.PreviousScaleVariable);
 
@@ -84,7 +94,11 @@ void UNiagaraStatelessModule_ScaleMeshSizeBySpeed::BuildEmitterData(const FNiaga
 		const FNiagaraDistributionFloat DefaultDistribution({ 0.0f, 1.0f });
 		BuiltData->ScaleDistribution = BuildContext.AddDistribution(DefaultDistribution);
 	}
-	BuildContext.AddParticleSimulationExecSimulate(&NSMScaleMeshSizeBySpeedPrivate::ParticleSimulate);
+	
+	if (BuiltData->PreviousPositionVariableOffset != INDEX_NONE || BuiltData->VelocityVariableOffset != INDEX_NONE)
+	{
+		BuildContext.AddParticleSimulationExecSimulate(&NSMScaleMeshSizeBySpeedPrivate::ParticleSimulate);
+	}
 }
 
 void UNiagaraStatelessModule_ScaleMeshSizeBySpeed::BuildShaderParameters(FNiagaraStatelessShaderParametersBuilder& ShaderParametersBuilder) const
diff --git a/Engine/Plugins/FX/Niagara/Source/Niagara/Private/Stateless/Modules/NiagaraStatelessModule_ScaleSpriteSizeBySpeed.cpp b/Engine/Plugins/FX/Niagara/Source/Niagara/Private/Stateless/Modules/NiagaraStatelessModule_ScaleSpriteSizeBySpeed.cpp
--- a/Engine/Plugins/FX/Niagara/Source/Niagara/Private/Stateless/Modules/NiagaraStatelessModule_ScaleSpriteSizeBySpeed.cpp
+++ b/Engine/Plugins/FX/Niagara/Source/Niagara/Private/Stateless/Modules/NiagaraStatelessModule_ScaleSpriteSizeBySpeed.cpp
@@ -16,6 +16,7 @@ namespace NSMScaleSpriteSizeBySpeedPrivate
 
 		int32							PositionVariableOffset = INDEX_NONE;
 		int32							PreviousPositionVariableOffset = INDEX_NONE;
+		int32							VelocityVariableOffset = INDEX_NONE;
 		int32							SpriteSizeVariableOffset = INDEX_NONE;
 		int32							PreviousSpriteSizeVariableOffset = INDEX_NONE;
 	};
@@ -29,9 +30,17 @@ namespace NSMScaleSpriteSizeBySpeedPrivate
 
 		for (uint32 i = 0; i < ParticleSimulationContext.GetNumInstances(); ++i)
 		{
-			const FVector3f Position			= ParticleSimulationContext.ReadParticleVariable(ModuleBuiltData->PositionVariableOffset, i, FVector3f::OneVector);
-			const FVector3f PreviousPosition	= ParticleSimulationContext.ReadParticleVariable(ModuleBuiltData->PreviousPositionVariableOffset, i, FVector3f::OneVector);
-			const FVector3f Velocity			= (Position - PreviousPosition) * ParticleSimulationContext.GetInvDeltaTime();
+			const FVector3f Position			= ParticleSimulationContext.ReadParticleVariable(ModuleBuiltData->PositionVariableOffset, i, FVector3f::ZeroVector);
+			FVector3f Velocity;
+			if (ModuleBuiltData->PreviousPositionVariableOffset != INDEX_NONE)
+			{
+				const FVector3f PreviousPosition = ParticleSimulationContext.ReadParticleVariable(ModuleBuiltData->PreviousPositionVariableOffset, i, Position);
+				Velocity = (Position - PreviousPosition) * ParticleSimulationContext.GetInvDeltaTime();
+			}
+			else
+			{
+				Velocity = ParticleSimulationContext.ReadParticleVariable(ModuleBuiltData->VelocityVariableOffset, i, FVector3f::ZeroVector);
+			}
 			const float Speed					= Velocity.SquaredLength();
 			const float NormSpeed				= FMath::Clamp(Speed * Parameters->ScaleSpriteSizeBySpeed_VelocityNorm, 0.0f, 1.0f);
 			const float Interp					= ParticleSimulationContext.SampleDistributionValue<float>(Parameters->ScaleSpriteSizeBySpeed_ScaleDistribution, i, 0, NormSpeed);
@@ -62,6 +71,7 @@ void UNiagaraStatelessModule_ScaleSpriteSizeBySpeed::BuildEmitterData(const FNia
 	const FNiagaraStatelessGlobals& StatelessGlobals = FNiagaraStatelessGlobals::Get();
 	BuiltData->PositionVariableOffset = BuildContext.FindParticleVariableIndex(StatelessGlobals.PositionVariable);
 	BuiltData->PreviousPositionVariableOffset = BuildContext.FindParticleVariableIndex(StatelessGlobals.PreviousPositionVariable);
+	BuiltData->VelocityVariableOffset = BuildContext.FindParticleVariableIndex(StatelessGlobals.VelocityVariable);
 	BuiltData->SpriteSizeVariableOffset = BuildContext.FindParticleVariableIndex(StatelessGlobals.SpriteSizeVariable);
 	BuiltData->PreviousSpriteSizeVariableOffset = BuildContext.FindParticleVariableIndex(StatelessGlobals.PreviousSpriteSizeVariable);
 
@@ -82,7 +92,11 @@ void UNiagaraStatelessModule_ScaleSpriteSizeBySpeed::BuildEmitterData(const FNia
 		const FNiagaraDistributionFloat DefaultDistribution({ 0.0f, 1.0f });
 		BuiltData->ScaleDistribution = BuildContext.AddDistribution(DefaultDistribution);
 	}
-	BuildContext.AddParticleSimulationExecSimulate(&ParticleSimulate);
+
+	if (BuiltData->PreviousPositionVariableOffset != INDEX_NONE || BuiltData->VelocityVariableOffset != INDEX_NONE)
+	{
+		BuildContext.AddParticleSimulationExecSimulate(&ParticleSimulate);
+	}
 }
 
 void UNiagaraStatelessModule_ScaleSpriteSizeBySpeed::BuildShaderParameters(FNiagaraStatelessShaderParametersBuilder& ShaderParametersBuilder) const
