- Fix for initializing GPU data while shrink might occur

#rnx

diff --git a/Engine/Plugins/FX/Niagara/Source/Niagara/Private/Stateless/NiagaraStatelessEmitter.cpp b/Engine/Plugins/FX/Niagara/Source/Niagara/Private/Stateless/NiagaraStatelessEmitter.cpp
--- a/Engine/Plugins/FX/Niagara/Source/Niagara/Private/Stateless/NiagaraStatelessEmitter.cpp
+++ b/Engine/Plugins/FX/Niagara/Source/Niagara/Private/Stateless/NiagaraStatelessEmitter.cpp
@@ -452,14 +452,6 @@ void UNiagaraStatelessEmitter::CacheFromCompiledData()
 		}
 		StatelessEmitterData->bModulesHaveRendererBindings = StatelessEmitterData->RendererBindings.Num() > 0;
 
-		// Make sure our static data buffer allows us to sample the max size in one chunk without having to worry about overflow
-		constexpr int MinSafeBufferSize = 16;
-		if (StatelessEmitterData->StaticDataBufferCpu.Num() < MinSafeBufferSize)
-		{
-			StatelessEmitterData->StaticDataBufferCpu.AddZeroed(MinSafeBufferSize - StatelessEmitterData->StaticDataBufferCpu.Num());
-		}
-		StatelessEmitterData->InitRenderResources();
-
 		// Gather the parameter collections we need to bind to
 	#if WITH_EDITOR
 		// UNiagaraStatelessEmitter does not get a notification when bindings change.
diff --git a/Engine/Plugins/FX/Niagara/Source/Niagara/Private/Stateless/NiagaraStatelessEmitterDataBuildContext.cpp b/Engine/Plugins/FX/Niagara/Source/Niagara/Private/Stateless/NiagaraStatelessEmitterDataBuildContext.cpp
--- a/Engine/Plugins/FX/Niagara/Source/Niagara/Private/Stateless/NiagaraStatelessEmitterDataBuildContext.cpp
+++ b/Engine/Plugins/FX/Niagara/Source/Niagara/Private/Stateless/NiagaraStatelessEmitterDataBuildContext.cpp
@@ -52,8 +52,18 @@ FNiagaraStatelessEmitterDataBuildContext::~FNiagaraStatelessEmitterDataBuildCont
 	EmitterData.Expressions.Shrink();
 	EmitterData.BuiltData.Shrink();
 	EmitterData.SharedBuiltData.Shrink();
-	EmitterData.StaticDataBufferCpu.Shrink();
 	//EmitterData.ParticleSimExecData.Shrink();
+
+	// Make sure our static data buffer allows us to sample the max size in one chunk without having to worry about overflow
+	constexpr int MinSafeBufferSize = 16;
+	if (EmitterData.StaticDataBufferCpu.Num() < MinSafeBufferSize)
+	{
+		EmitterData.StaticDataBufferCpu.AddZeroed(MinSafeBufferSize - EmitterData.StaticDataBufferCpu.Num());
+	}
+	EmitterData.StaticDataBufferCpu.Shrink();
+
+	// Initialize the GPU data
+	EmitterData.InitRenderResources();
 }
 
 void FNiagaraStatelessEmitterDataBuildContext::PreModuleBuild(int32 InShaderParameterOffset)
