- Fix for simulations sometimes having the incorrect ticking group assigned to them

#tests In code randomly forced instances to reset while ticking

diff --git a/Engine/Plugins/FX/Niagara/Source/Niagara/Private/NiagaraWorldManager.cpp b/Engine/Plugins/FX/Niagara/Source/Niagara/Private/NiagaraWorldManager.cpp
--- a/Engine/Plugins/FX/Niagara/Source/Niagara/Private/NiagaraWorldManager.cpp
+++ b/Engine/Plugins/FX/Niagara/Source/Niagara/Private/NiagaraWorldManager.cpp
@@ -319,6 +319,12 @@ namespace NiagaraWorldManagerInternal
 		return ActualTickGroup;
 	}
 
+	ETickingGroup GetTickingGroupFromNiagara(int ActualTickGroup)
+	{
+		const int32 TickGroup = FMath::Clamp(ActualTickGroup + NiagaraFirstTickGroup, 0, TG_MAX - 1);
+		return ETickingGroup(TickGroup);
+	}
+
 	bool GetNiagaraTickGroupPriority(int32 NiagaraTickGroup)
 	{
 		return NiagaraTickGroup >= GFirstHighPriTickGroup;
@@ -707,6 +713,9 @@ FNiagaraSystemSimulationRef FNiagaraWorldManager::GetSystemSimulation(ETickingGr
 	{
 		int32 DemotedTickGroup = NiagaraWorldManagerInternal::GetNiagaraTickGroup((ETickingGroup)(TickGroup + 1));
 		ActualTickGroup = DemotedTickGroup == ActualTickGroup ? 0 : DemotedTickGroup;
+		
+		// Note we must change TickGroup to the new tick group in case we need to create a new simulation, otherwise it will have the wrong TickGroup assigned to it
+		TickGroup = NiagaraWorldManagerInternal::GetTickingGroupFromNiagara(ActualTickGroup);
 	}
 
 	const bool bCanOverlapTickGroups = System->AsyncWorkCanOverlapTickGroups();
