PR #13968: Fix Niagara sample normal of static mesh mismatch

diff --git a/Engine/Plugins/FX/Niagara/Shaders/Private/NiagaraDataInterfaceStaticMeshTemplate.ush b/Engine/Plugins/FX/Niagara/Shaders/Private/NiagaraDataInterfaceStaticMeshTemplate.ush
--- a/Engine/Plugins/FX/Niagara/Shaders/Private/NiagaraDataInterfaceStaticMeshTemplate.ush
+++ b/Engine/Plugins/FX/Niagara/Shaders/Private/NiagaraDataInterfaceStaticMeshTemplate.ush
@@ -313,9 +313,9 @@ void GetTriangle_{ParameterName}(int Tri, float3 BaryCoord, out float3 Position,
 
 	Position	= Positions[0] * BaryCoord.x + Positions[1] * BaryCoord.y + Positions[2] * BaryCoord.z;
 	Velocity	= float3(0, 0, 0);
-	Normal		= Normals[0] * BaryCoord.x + Normals[1] * BaryCoord.y + Normals[2] * BaryCoord.x;
-	Bitangent	= Bitangents[0] * BaryCoord.x + Bitangents[1] * BaryCoord.y + Bitangents[2] * BaryCoord.x;
-	Tangent		= Tangents[0] * BaryCoord.x + Tangents[1] * BaryCoord.y + Tangents[2] * BaryCoord.x;
+	Normal		= Normals[0] * BaryCoord.x + Normals[1] * BaryCoord.y + Normals[2] * BaryCoord.z;
+	Bitangent	= Bitangents[0] * BaryCoord.x + Bitangents[1] * BaryCoord.y + Bitangents[2] * BaryCoord.z;
+	Tangent		= Tangents[0] * BaryCoord.x + Tangents[1] * BaryCoord.y + Tangents[2] * BaryCoord.z;
 }
 
 void GetTriangleWS_{ParameterName}(int Tri, float3 BaryCoord, out float3 Position, out float3 Velocity, out float3 Normal, out float3 Bitangent, out float3 Tangent)
