- Fix for potentially generating a NaN when using local space mesh emitters

#tests Make a mesh emitter that orients to velocity on the renderer, place in level, rotate the component/actor on the X axis

diff --git a/Engine/Plugins/FX/Niagara/Shaders/Private/NiagaraMeshParticleUtils.ush b/Engine/Plugins/FX/Niagara/Shaders/Private/NiagaraMeshParticleUtils.ush
--- a/Engine/Plugins/FX/Niagara/Shaders/Private/NiagaraMeshParticleUtils.ush
+++ b/Engine/Plugins/FX/Niagara/Shaders/Private/NiagaraMeshParticleUtils.ush
@@ -307,7 +307,8 @@ float3x3 NiagaraGetMeshFacingMatrix(FLWCVector3 ParticleSimPosition, NiagaraMesh
 		else
 		{
 			// Prefer to use world up as a reference vector, fall back to world X-axis when facing up or down
-			float DotWorldZ = dot(FacingDir, WorldZ);
+			const float3 SimFacingDir = NiagaraSimToWorldVec(FacingDir, Params);
+			float DotWorldZ = dot(SimFacingDir, WorldZ);
 			RefVector = abs(DotWorldZ) > 0.99f ? (-sign(DotWorldZ) * WorldX) : WorldZ;
 		}
 
