[Niagara] Expose batching registry plugins unmounts.

#tests local run, replay runs
#rnx

diff --git a/Engine/Plugins/FX/Niagara/Source/Niagara/Private/NiagaraTypeRegistry.cpp b/Engine/Plugins/FX/Niagara/Source/Niagara/Private/NiagaraTypeRegistry.cpp
--- a/Engine/Plugins/FX/Niagara/Source/Niagara/Private/NiagaraTypeRegistry.cpp
+++ b/Engine/Plugins/FX/Niagara/Source/Niagara/Private/NiagaraTypeRegistry.cpp
@@ -3,11 +3,17 @@
 #include "NiagaraTypeRegistry.h"
 
 #include "Algo/Transform.h"
+#include "Containers/Array.h"
+#include "Containers/ArrayView.h"
+#include "Containers/ContainersFwd.h"
+#include "HAL/CriticalSection.h"
 #include "Interfaces/IPluginManager.h"
 #include "Misc/LazySingleton.h"
 #include "NiagaraDataInterfaceBase.h"
 #include "NiagaraDebugVis.h" // IWYU pragma: keep
 #include "UObject/CoreRedirects.h"
+#include "UObject/NameTypes.h"
+#include <atomic>
 
 // these two globals are intended to help ensure that a global variable is accessible to natvis (as defined in Niagara.natvis)
 // while debugging.  GCoreTypeRegistrySingletonPtr is the pointer to the actual data stored within the TLazySingleton<FNiagaraTypeRegistry>
@@ -33,12 +39,24 @@ namespace NiagaraTypeRegistryLocal
 
 	static FDelegateHandle PluginUnmountedHandle;
 
+	static std::atomic<int32> GSuspendPluginUnmountCounter = 0;
+	static FCriticalSection GPluginUnmountSync;
+	static TArray<FString> GPluginMountPointsToUnmount;
+
 	void PluginUnmounted(IPlugin& Plugin)
 	{
 		const FString MountedPath = Plugin.GetMountedAssetPath();
 
-		// go through all of the registered types and see if any of them are from the plugin being unmounted.
-		FNiagaraTypeRegistry::Get().InvalidateTypesByPath(MountedPath);
+		if (GSuspendPluginUnmountCounter.load(std::memory_order_acquire) > 0)
+		{
+			FScopeLock Guard(&NiagaraTypeRegistryLocal::GPluginUnmountSync);
+			NiagaraTypeRegistryLocal::GPluginMountPointsToUnmount.Add(MountedPath);
+		}
+		else
+		{
+			// go through all of the registered types and see if any of them are from the plugin being unmounted.
+			FNiagaraTypeRegistry::Get().InvalidateTypesByPath(MountedPath);
+		}
 	}
 }
 
@@ -151,8 +169,36 @@ void FNiagaraTypeRegistry::TearDown()
 	TLazySingleton<FNiagaraTypeRegistry>::TearDown();
 }
 
+void FNiagaraTypeRegistry::SuspendPluginUnmount()
+{
+	NiagaraTypeRegistryLocal::GSuspendPluginUnmountCounter.fetch_add(1, std::memory_order_acq_rel);
+}
+
+void FNiagaraTypeRegistry::ResumePluginUnmount()
+{
+	if (NiagaraTypeRegistryLocal::GSuspendPluginUnmountCounter.fetch_sub(1, std::memory_order_acq_rel) == 1)
+	{
+		TArray<FString> LocalListOfPaths;
+		{
+			FScopeLock Guard(&NiagaraTypeRegistryLocal::GPluginUnmountSync);
+			LocalListOfPaths = MoveTemp(NiagaraTypeRegistryLocal::GPluginMountPointsToUnmount);
+		}
+		InvalidateTypesByPaths(LocalListOfPaths);
+	}
+}
+
 void FNiagaraTypeRegistry::InvalidateTypesByPath(const FString& AssetPath)
 {
+	InvalidateTypesByPaths({AssetPath});
+}
+
+void FNiagaraTypeRegistry::InvalidateTypesByPaths(TConstArrayView<FString> AssetPaths)
+{
+	if (AssetPaths.IsEmpty())
+	{
+		return;
+	}
+
 	FNiagaraTypeRegistry& Registry = Get();
 
 	if (!Registry.bModuleInitialized)
@@ -160,17 +206,25 @@ void FNiagaraTypeRegistry::InvalidateTypesByPath(const FString& AssetPath)
 		return;
 	}
 
+	TRACE_CPUPROFILER_EVENT_SCOPE(NiagaraTypeRegistry_InvalidateTypesByPaths);
+
 	TArray<int32> TypesToInvalidate;
 
 	{
 		FReadScopeLock Lock(Registry.RegistrationLock);
 
+		FNameBuilder RegistryTypeNameStr;
 		const int32 TypeCount = Registry.RegisteredSoftObjectPaths.Num();
 		for (int32 TypeIt = 0; TypeIt < TypeCount; ++TypeIt)
 		{
-			if (Registry.RegisteredSoftObjectPaths[TypeIt].GetLongPackageName().StartsWith(AssetPath))
+			Registry.RegisteredSoftObjectPaths[TypeIt].GetLongPackageFName().ToString(RegistryTypeNameStr);
+			const FStringView RegistryTypeName = RegistryTypeNameStr.ToView();
+			for (const FString& AssetPath : AssetPaths)
 			{
-				TypesToInvalidate.Add(TypeIt);
+				if (RegistryTypeName.StartsWith(AssetPath))
+				{
+					TypesToInvalidate.Add(TypeIt);
+				}
 			}
 		}
 	}
@@ -196,6 +250,7 @@ void FNiagaraTypeRegistry::InvalidateTypesByPath(const FString& AssetPath)
 			}
 		}
 
+#if DO_ENSURE
 		{
 			FReadScopeLock Lock(Registry.RegistrationLock);
 
@@ -220,6 +275,7 @@ void FNiagaraTypeRegistry::InvalidateTypesByPath(const FString& AssetPath)
 				//RegisteredSoftObjectPaths
 			}
 		}
+#endif
 	}
 }
 
diff --git a/Engine/Plugins/FX/Niagara/Source/Niagara/Public/NiagaraTypeRegistry.h b/Engine/Plugins/FX/Niagara/Source/Niagara/Public/NiagaraTypeRegistry.h
--- a/Engine/Plugins/FX/Niagara/Source/Niagara/Public/NiagaraTypeRegistry.h
+++ b/Engine/Plugins/FX/Niagara/Source/Niagara/Public/NiagaraTypeRegistry.h
@@ -2,7 +2,9 @@
 
 #pragma once
 
+#include "Containers/ContainersFwd.h"
 #include "CoreMinimal.h"
+#include "HAL/Platform.h"
 #include "Misc/TransactionallySafeRWLock.h"
 #include "Misc/ScopeLock.h"
 #include "Misc/SecureHash.h"
@@ -78,6 +80,7 @@ public:
 
 	static NIAGARA_API FNiagaraTypeDefinition GetTypeForStruct(UScriptStruct* InStruct);
 	static NIAGARA_API void InvalidateTypesByPath(const FString& AssetPath);
+	static NIAGARA_API void InvalidateTypesByPaths(TConstArrayView<FString> AssetPaths);
 
 	/** LazySingleton interface */
 	static NIAGARA_API FNiagaraTypeRegistry& Get();
@@ -101,6 +104,13 @@ public:
 		Registry.ForAllSelectedRegisteredTypes(TrivialSelect, Func);
 	}
 
+	class FSuspendPluginUnmounting
+	{
+	public:
+		UE_NODISCARD_CTOR FSuspendPluginUnmounting() { SuspendPluginUnmount(); }
+		~FSuspendPluginUnmounting() { ResumePluginUnmount(); }
+	};
+
 private:
 	struct FQueuedRegistryEntry
 	{
@@ -113,6 +123,9 @@ private:
 	NIAGARA_API FNiagaraTypeRegistry();
 	NIAGARA_API virtual ~FNiagaraTypeRegistry();
 
+	static NIAGARA_API void SuspendPluginUnmount();
+	static NIAGARA_API void ResumePluginUnmount();
+
 	int32 RegisterTypeInternal(const FNiagaraTypeDefinition& NewType);
 	const FNiagaraTypeDefinition& GetRegisteredType(int32 TypeIndex) const;
 
