- Add reset to default for material parameter bindings, where the default is pulled from the material

#tests Create a material with scalar / vector / texture parameters, binding to a renderer and reset to the material default value

diff --git a/Engine/Plugins/FX/Niagara/Source/NiagaraEditor/Private/Customizations/NiagaraTypeCustomizations.cpp b/Engine/Plugins/FX/Niagara/Source/NiagaraEditor/Private/Customizations/NiagaraTypeCustomizations.cpp
--- a/Engine/Plugins/FX/Niagara/Source/NiagaraEditor/Private/Customizations/NiagaraTypeCustomizations.cpp
+++ b/Engine/Plugins/FX/Niagara/Source/NiagaraEditor/Private/Customizations/NiagaraTypeCustomizations.cpp
@@ -16,6 +16,7 @@
 #include "Framework/Application/SlateApplication.h"
 #include "Framework/MultiBox/MultiBoxBuilder.h"
 #include "Materials/MaterialInterface.h"
+#include "Misc/DefaultValueHelper.h"
 #include "Modules/ModuleManager.h"
 #include "ViewModels/NiagaraSystemViewModel.h"
 #include "ViewModels/TNiagaraViewModelManager.h"
@@ -58,6 +59,274 @@
 #define LOCTEXT_NAMESPACE "FNiagaraVariableAttributeBindingCustomization"
 #define ALLOW_LIBRARY_TO_LIBRARY_DEFAULT_BINDING 0
 
+//////////////////////////////////////////////////////////////////////////////////////////////////////////////////////////////////////////////////////////////////
+
+struct FNiagaraMaterialParameterValueHelper : public TSharedFromThis<FNiagaraMaterialParameterValueHelper>
+{
+	FNiagaraMaterialParameterValueHelper(TSharedPtr<IPropertyHandle> InNamePropertyHandle, TSharedPtr<IPropertyHandle> InValuePropertyHandle, TWeakObjectPtr<UNiagaraRendererProperties> InWeakRenderProperties)
+		: NamePropertyHandle(InNamePropertyHandle.ToSharedRef())
+		, ValuePropertyHandle(InValuePropertyHandle.ToSharedRef())
+		, WeakRenderProperties(InWeakRenderProperties)
+	{
+	}
+
+	template<typename TType>
+	TOptional<TType> GetMaterialDefaultValue(UMaterialInterface* Material, FName ParameterName) const;
+
+	template<>
+	TOptional<float> GetMaterialDefaultValue(UMaterialInterface* Material, FName ParameterName) const
+	{
+		float DefaultValue;
+		if (Material->GetScalarParameterDefaultValue(ParameterName, DefaultValue))
+		{
+			return DefaultValue;
+		}
+		return {};
+	}
+
+	template<>
+	TOptional<FLinearColor> GetMaterialDefaultValue(UMaterialInterface* Material, FName ParameterName) const
+	{
+		FLinearColor DefaultValue;
+		if (Material->GetVectorParameterDefaultValue(ParameterName, DefaultValue))
+		{
+			return DefaultValue;
+		}
+		return {};
+	}
+
+	template<>
+	TOptional<UTexture*> GetMaterialDefaultValue(UMaterialInterface* Material, FName ParameterName) const
+	{
+		UTexture* DefaultValue;
+		if (Material->GetTextureParameterDefaultValue(ParameterName, DefaultValue))
+		{
+			return DefaultValue;
+		}
+		return {};
+	}
+
+	void AppendValueToString(float V, FString& ValuesString) const
+	{
+		ValuesString.Appendf(TEXT("%f"), V);
+	}
+
+	void AppendValueToString(const FLinearColor& V, FString& ValuesString) const
+	{
+		ValuesString.Appendf(TEXT("{R=%f, G=%f, B=%f, A=%f}"), V.R, V.G, V.B, V.A);
+	}
+
+	void AppendValueToString(const UTexture* V, FString& ValuesString) const
+	{
+		ValuesString.Append(*GetNameSafe(V));
+	}
+
+	template<typename TType>
+	FText FromValuesAsText(TArray<TType> Values) const
+	{
+		FString ValuesString;
+		for (const TType& V : Values)
+		{
+			if (!ValuesString.IsEmpty())
+			{
+				ValuesString.Append(TEXT(", "));
+			}
+			AppendValueToString(V, ValuesString);
+		}
+		return FText::FromString(ValuesString);
+	}
+
+	FName GetParameterName() const
+	{
+		FName ParameterName;
+		NamePropertyHandle->GetValue(ParameterName);
+		return ParameterName;
+	}
+
+	template<typename TType>
+	TArray<TType> TryResolveMaterialDefaultValue() const
+	{
+		TArray<TType> DefaultValues;
+
+		if (UNiagaraRendererProperties* RenderProperties = WeakRenderProperties.Get())
+		{
+			const FName ParameterName = GetParameterName();
+
+			TArray<UMaterialInterface*> UsedMaterials;
+			RenderProperties->GetUsedMaterials(nullptr, UsedMaterials);
+
+			for (UMaterialInterface* Material : UsedMaterials)
+			{
+				if (Material == nullptr)
+				{
+					continue;
+				}
+
+				TOptional<TType> DefaultValue = GetMaterialDefaultValue<TType>(Material, ParameterName);
+				if (DefaultValue.IsSet())
+				{
+					DefaultValues.AddUnique(DefaultValue.GetValue());
+				}
+			}
+		}
+
+		return DefaultValues;
+	}
+
+	template<typename TType>
+	FText GetToolTipText() const
+	{
+		const TArray<TType> Values = TryResolveMaterialDefaultValue<TType>();
+		const FText ParameterNameText = FText::FromName(GetParameterName());
+		if (Values.Num() == 0)
+		{
+			return FText::Format(LOCTEXT("GetDefaultMaterialValue_InvalidBinding", "Parameter '{0}' is either invalid or not found on the material"), ParameterNameText);
+		}
+		else if (Values.Num() == 1)
+		{
+			return FText::Format(LOCTEXT("GetDefaultMaterialValue_Value", "Reset To Default\nSet to parameter '{0}' default value '{1}'"), ParameterNameText, FromValuesAsText(Values));
+		}
+		else
+		{
+			return FText::Format(LOCTEXT("GetDefaultMaterialValue_Conflict", "Reset To Default\nParameter '{0}' has conflicting values '{1}'"), ParameterNameText, FromValuesAsText(Values));
+		}
+	}
+
+	template<typename TType>
+	bool GetPropertyValue(TType& OutValue) const
+	{
+		return ValuePropertyHandle->GetValue(OutValue) == FPropertyAccess::Success;
+	}
+
+	template<>
+	bool GetPropertyValue(UTexture*& OutValue) const
+	{
+		UObject* TempValue = nullptr;
+		bool bSuccess = ValuePropertyHandle->GetValue(TempValue) == FPropertyAccess::Success;
+		OutValue = Cast<UTexture>(TempValue);
+		return bSuccess;
+	}
+
+	template<>
+	bool GetPropertyValue(FLinearColor& OutValue) const
+	{
+		TArray<void*> RawDatas;
+		ValuePropertyHandle->AccessRawData(RawDatas);
+
+		if (RawDatas.Num() == 1 && RawDatas[0] != nullptr)
+		{
+			FMemory::Memcpy(&OutValue, RawDatas[0], sizeof(FLinearColor));
+			return true;
+		}
+		return false;
+	}
+
+	template<typename TType>
+	EVisibility IsVisible() const
+	{
+		TType Value;
+		if (GetPropertyValue(Value))
+		{
+			const TArray<TType> Values = TryResolveMaterialDefaultValue<TType>();
+			if (Values.Num() > 0 && Value != Values[0])
+			{
+				return EVisibility::Visible;
+			}
+		}
+		return EVisibility::Hidden;
+	}
+
+	template<typename TType>
+	void SetPropertyValue(TType Value) const
+	{
+		ValuePropertyHandle->SetValue(Value);
+	}
+
+	template<>
+	void SetPropertyValue(FLinearColor Value) const
+	{
+		TArray<void*> RawDatas;
+		ValuePropertyHandle->AccessRawData(RawDatas);
+
+		for (void* RawData : RawDatas)
+		{
+			if (RawData)
+			{
+				FMemory::Memcpy(RawData, &Value, sizeof(FLinearColor));
+			}
+		}
+	}
+
+	template<typename TType>
+	FReply OnClicked() const
+	{
+		const TArray<TType> Values = TryResolveMaterialDefaultValue<TType>();
+		if (Values.Num() > 0)
+		{
+			FScopedTransaction Transaction(LOCTEXT("SetValueToMaterialDefault", "Set Value To Material Default"));
+
+			TArray<UObject*> OuterObjects;
+			ValuePropertyHandle->GetOuterObjects(OuterObjects);
+			for (UObject* OuterObject : OuterObjects)
+			{
+				OuterObject->Modify();
+			}
+
+			ValuePropertyHandle->NotifyPreChange();
+
+			SetPropertyValue(Values[0]);
+
+			ValuePropertyHandle->NotifyPostChange(EPropertyChangeType::ValueSet);
+		}
+		return FReply::Handled();
+	}
+
+	template<typename TType>
+	bool CreateWidget(IDetailChildrenBuilder& ChildBuilder)
+	{
+		ChildBuilder.AddCustomRow(ValuePropertyHandle->GetPropertyDisplayName())
+		.NameWidget
+		[
+			ValuePropertyHandle->CreatePropertyNameWidget()
+		]
+		.ValueWidget
+		[
+			SNew(SHorizontalBox)
+			+ SHorizontalBox::Slot()
+			[
+				ValuePropertyHandle->CreatePropertyValueWidget()
+			]
+			+ SHorizontalBox::Slot()
+			.AutoWidth()
+			.VAlign(VAlign_Center)
+			.Padding(2, 1)
+			[
+				SNew(SButton)
+				.IsFocusable(false)
+				.ToolTipText(this, &FNiagaraMaterialParameterValueHelper::GetToolTipText<TType>)
+				.ButtonStyle(FAppStyle::Get(), "NoBorder")
+				.ContentPadding(0)
+				.Visibility(this, &FNiagaraMaterialParameterValueHelper::IsVisible<TType>)
+				.OnClicked(this, &FNiagaraMaterialParameterValueHelper::OnClicked<TType>)
+				.Content()
+				[
+					SNew(SImage)
+					.Image(FAppStyle::GetBrush("PropertyWindow.DiffersFromDefault"))
+					.ColorAndOpacity(FSlateColor::UseForeground())
+				]
+			]
+		];
+
+		return true;
+	}
+
+	TSharedRef<IPropertyHandle>					NamePropertyHandle;
+	TSharedRef<IPropertyHandle>					ValuePropertyHandle;
+	TWeakObjectPtr<UNiagaraRendererProperties>	WeakRenderProperties;
+};
+
+//////////////////////////////////////////////////////////////////////////////////////////////////////////////////////////////////////////////////////////////////
+
 void FNiagaraNumericCustomization::CustomizeHeader(TSharedRef<IPropertyHandle> PropertyHandle, FDetailWidgetRow& HeaderRow, IPropertyTypeCustomizationUtils& CustomizationUtils)
 {
 	TSharedPtr<IPropertyHandle> ValueHandle = PropertyHandle->GetChildHandle(TEXT("Value"));
@@ -2306,6 +2575,7 @@ void FNiagaraRendererMaterialParameterCustomization::CustomizeChildren(TSharedRe
 		TSharedPtr<IPropertyHandle> ChildHandle = StructPropertyHandle->GetChildHandle(ChildIndex);
 		if ( ChildHandle->GetProperty()->GetFName() == NAME_MaterialParameterName )
 		{
+			NamePropertyHandle = ChildHandle;
 			ChildBuilder.AddCustomRow(ChildHandle->GetPropertyDisplayName())
 			.NameWidget
 			[
@@ -2441,18 +2711,48 @@ void FNiagaraRendererMaterialScalarParameterCustomization::GetMaterialParameterI
 	Material->GetAllScalarParameterInfo(OutMaterialParameterInfos, Guids);
 }
 
+bool FNiagaraRendererMaterialScalarParameterCustomization::CustomizeChildProperty(IDetailChildrenBuilder& ChildBuilder, TSharedPtr<IPropertyHandle> PropertyHandle)
+{
+	if (NamePropertyHandle.IsValid() && PropertyHandle->GetProperty()->GetFName() == GET_MEMBER_NAME_CHECKED(FNiagaraRendererMaterialScalarParameter, Value))
+	{
+		ValueHelper = MakeShared<FNiagaraMaterialParameterValueHelper>(NamePropertyHandle, PropertyHandle, WeakRenderProperties);
+		return ValueHelper->CreateWidget<float>(ChildBuilder);
+	}
+	return false;
+}
+
 void FNiagaraRendererMaterialVectorParameterCustomization::GetMaterialParameterInfos(UMaterialInterface* Material, TArray<FMaterialParameterInfo>& OutMaterialParameterInfos) const
 {
 	TArray<FGuid> Guids;
 	Material->GetAllVectorParameterInfo(OutMaterialParameterInfos, Guids);
 }
 
+bool FNiagaraRendererMaterialVectorParameterCustomization::CustomizeChildProperty(IDetailChildrenBuilder& ChildBuilder, TSharedPtr<IPropertyHandle> PropertyHandle)
+{
+	if (NamePropertyHandle.IsValid() && PropertyHandle->GetProperty()->GetFName() == GET_MEMBER_NAME_CHECKED(FNiagaraRendererMaterialVectorParameter, Value))
+	{
+		ValueHelper = MakeShared<FNiagaraMaterialParameterValueHelper>(NamePropertyHandle, PropertyHandle, WeakRenderProperties);
+		return ValueHelper->CreateWidget<FLinearColor>(ChildBuilder);
+	}
+	return false;
+}
+
 void FNiagaraRendererMaterialTextureParameterCustomization::GetMaterialParameterInfos(UMaterialInterface* Material, TArray<FMaterialParameterInfo>& OutMaterialParameterInfos) const
 {
 	TArray<FGuid> Guids;
 	Material->GetAllTextureParameterInfo(OutMaterialParameterInfos, Guids);
 }
 
+bool FNiagaraRendererMaterialTextureParameterCustomization::CustomizeChildProperty(IDetailChildrenBuilder& ChildBuilder, TSharedPtr<IPropertyHandle> PropertyHandle)
+{
+	if (NamePropertyHandle.IsValid() && PropertyHandle->GetProperty()->GetFName() == GET_MEMBER_NAME_CHECKED(FNiagaraRendererMaterialTextureParameter, Texture))
+	{
+		ValueHelper = MakeShared<FNiagaraMaterialParameterValueHelper>(NamePropertyHandle, PropertyHandle, WeakRenderProperties);
+		return ValueHelper->CreateWidget<UTexture*>(ChildBuilder);
+	}
+	return false;
+}
+
 bool FNiagaraRendererMaterialStaticBoolParameterCustomization::CustomizeChildProperty(class IDetailChildrenBuilder& ChildBuilder, TSharedPtr<IPropertyHandle> PropertyHandle)
 {
 	if (PropertyHandle->GetProperty()->GetFName() != GET_MEMBER_NAME_CHECKED(FNiagaraRendererMaterialStaticBoolParameter, StaticVariableName))
diff --git a/Engine/Plugins/FX/Niagara/Source/NiagaraEditor/Private/Customizations/NiagaraTypeCustomizations.h b/Engine/Plugins/FX/Niagara/Source/NiagaraEditor/Private/Customizations/NiagaraTypeCustomizations.h
--- a/Engine/Plugins/FX/Niagara/Source/NiagaraEditor/Private/Customizations/NiagaraTypeCustomizations.h
+++ b/Engine/Plugins/FX/Niagara/Source/NiagaraEditor/Private/Customizations/NiagaraTypeCustomizations.h
@@ -405,6 +405,8 @@ public:
 
 protected:
 	TWeakObjectPtr<UNiagaraRendererProperties>	WeakRenderProperties;
+	TSharedPtr<IPropertyHandle>					NamePropertyHandle;
+	TSharedPtr<struct FNiagaraMaterialParameterValueHelper>	ValueHelper;
 };
 
 class FNiagaraRendererMaterialScalarParameterCustomization : public FNiagaraRendererMaterialParameterCustomization
@@ -416,6 +418,7 @@ public:
 	}
 
 	virtual void GetMaterialParameterInfos(UMaterialInterface* Material, TArray<FMaterialParameterInfo>& OutMaterialParameterInfos) const override;
+	virtual bool CustomizeChildProperty(class IDetailChildrenBuilder& ChildBuilder, TSharedPtr<IPropertyHandle> PropertyHandle) override;
 };
 
 class FNiagaraRendererMaterialVectorParameterCustomization : public FNiagaraRendererMaterialParameterCustomization
@@ -427,6 +430,7 @@ public:
 	}
 
 	virtual void GetMaterialParameterInfos(UMaterialInterface* Material, TArray<FMaterialParameterInfo>& OutMaterialParameterInfos) const override;
+	virtual bool CustomizeChildProperty(IDetailChildrenBuilder& ChildBuilder, TSharedPtr<IPropertyHandle> PropertyHandle) override;
 };
 
 class FNiagaraRendererMaterialTextureParameterCustomization : public FNiagaraRendererMaterialParameterCustomization
@@ -438,6 +442,7 @@ public:
 	}
 
 	virtual void GetMaterialParameterInfos(UMaterialInterface* Material, TArray<FMaterialParameterInfo>& OutMaterialParameterInfos) const override;
+	virtual bool CustomizeChildProperty(IDetailChildrenBuilder& ChildBuilder, TSharedPtr<IPropertyHandle> PropertyHandle) override;
 };
 
 class FNiagaraRendererMaterialStaticBoolParameterCustomization : public FNiagaraRendererMaterialParameterCustomization
