- Fix for distribution UI not always updating when a module is pasted over an existing module

diff --git a/Engine/Plugins/FX/Niagara/Source/NiagaraEditor/Private/Stateless/NiagaraDistributionPropertyCustomization.cpp b/Engine/Plugins/FX/Niagara/Source/NiagaraEditor/Private/Stateless/NiagaraDistributionPropertyCustomization.cpp
--- a/Engine/Plugins/FX/Niagara/Source/NiagaraEditor/Private/Stateless/NiagaraDistributionPropertyCustomization.cpp
+++ b/Engine/Plugins/FX/Niagara/Source/NiagaraEditor/Private/Stateless/NiagaraDistributionPropertyCustomization.cpp
@@ -44,8 +44,16 @@ public:
 		SourceNumChannels = InNumChannels;
 		if (IsValid())
 		{
+			// Use UNiagaraNotifyOnChanged if available as it is more robust against copy and paste changes vs SetOnPropertyValueChanged
+			if (UNiagaraNotifyOnChanged* OwnerAsNotifyOnChanged = Cast<UNiagaraNotifyOnChanged>(InOwningObject))
+			{
+				OwnerAsNotifyOnChanged->OnChanged().AddSP(AsShared(), &FNiagaraDistributionAdapter::DistributionPropertyChanged);
+			}
+			else
+			{
+				PropertyHandle->SetOnPropertyValueChanged(FSimpleDelegate::CreateSP(AsShared(), &FNiagaraDistributionAdapter::DistributionPropertyChanged));
+			}
 			bDisplayAsColor = InDistribution->DisplayAsColor();
-			PropertyHandle->SetOnPropertyValueChanged(FSimpleDelegate::CreateSP(this->AsShared(), &FNiagaraDistributionAdapter::DistributionPropertyChanged));
 			EditorMode = GetDistributionEditorModeFromSourceMode(SourceNumChannels, bDisplayAsColor, SourceDistribution->Mode);
 			EditorNumChannels = FNiagaraDistributionEditorUtilities::IsUniform(EditorMode) ? 1 : SourceNumChannels;
 			WidgetCustomizationOptions = FNiagaraInputParameterCustomization::MakeFromProperty(InPropertyHandle);
