- Fix for unbound buffers with ribbon renderer

#rnx

diff --git a/Engine/Plugins/FX/Niagara/Source/Niagara/Private/NiagaraRendererRibbons.cpp b/Engine/Plugins/FX/Niagara/Source/Niagara/Private/NiagaraRendererRibbons.cpp
--- a/Engine/Plugins/FX/Niagara/Source/Niagara/Private/NiagaraRendererRibbons.cpp
+++ b/Engine/Plugins/FX/Niagara/Source/Niagara/Private/NiagaraRendererRibbons.cpp
@@ -2745,6 +2745,9 @@ void FNiagaraRendererRibbons::InitializeVertexBuffersGPU(FRHICommandListImmediat
 		{
 			SwapAccumulationBuffers();
 
+			FNiagaraEmptyUAVPoolScopedAccess UAVPoolAccessScope(ComputeDispatchInterface->GetEmptyUAVPool());
+			TArray<FRHITransitionInfo, TInlineAllocator<4>> Transitions;
+
 			FNiagaraRibbonAggregationApplyParameters ApplyParams;
 			FMemory::Memzero(ApplyParams);
 			ApplyParams.Common = CommonParams;
@@ -2752,21 +2755,38 @@ void FNiagaraRendererRibbons::InitializeVertexBuffersGPU(FRHICommandListImmediat
 			ApplyParams.InputAccumulation = CurrentAccumulationSRV;
 			ApplyParams.OutputTangentsAndDistances = VertexBuffers.TangentsAndDistancesBuffer.UAV;
 			ApplyParams.OutputMultiRibbonIndices = VertexBuffers.MultiRibbonIndicesBuffer.UAV;
+			if (ApplyParams.OutputMultiRibbonIndices)
+			{
+				Transitions.Emplace(ApplyParams.OutputMultiRibbonIndices, ERHIAccess::SRVMask, ERHIAccess::UAVMask);
+			}
+			else
+			{
+				ApplyParams.OutputMultiRibbonIndices = ComputeDispatchInterface->GetEmptyUAVFromPool(RHICmdList, VertexBuffers.MultiRibbonIndicesBuffer.PixelFormat, ENiagaraEmptyUAVType::Buffer);
+			}
 			ApplyParams.OutputTessellationStats = TempBuffers.TransientTessellationStats.UAV;
+			if (ApplyParams.OutputTessellationStats)
+			{
+				Transitions.Emplace(ApplyParams.OutputTessellationStats, ERHIAccess::SRVMask, ERHIAccess::UAVMask);
+			}
+			else
+			{
+				ApplyParams.OutputTessellationStats = ComputeDispatchInterface->GetEmptyUAVFromPool(RHICmdList, EPixelFormat::PF_R32_FLOAT, ENiagaraEmptyUAVType::Buffer);
+			}
 			ApplyParams.OutputSegments = VertexBuffers.SegmentsBuffer.UAV;
 
-			RHICmdList.Transition(FRHITransitionInfo(ApplyParams.OutputTangentsAndDistances, ERHIAccess::SRVMask, ERHIAccess::UAVMask));
-			RHICmdList.Transition(FRHITransitionInfo(ApplyParams.OutputMultiRibbonIndices, ERHIAccess::SRVMask, ERHIAccess::UAVMask));
-			RHICmdList.Transition(FRHITransitionInfo(ApplyParams.OutputSegments, ERHIAccess::SRVMask, ERHIAccess::UAVMask));
-			RHICmdList.Transition(FRHITransitionInfo(ApplyParams.OutputTessellationStats, ERHIAccess::SRVMask, ERHIAccess::UAVMask));
+			Transitions.Emplace(ApplyParams.OutputTangentsAndDistances, ERHIAccess::SRVMask, ERHIAccess::UAVMask);
+			Transitions.Emplace(ApplyParams.OutputSegments, ERHIAccess::SRVMask, ERHIAccess::UAVMask);
+
+			RHICmdList.Transition(Transitions);
 
 			const uint32 NumThreadGroups = FMath::DivideAndRoundUp<uint32>(NumExecutableInstances, FNiagaraRibbonComputeCommon::VertexGenReductionPropagationThreadSize);
 			FComputeShaderUtils::Dispatch(RHICmdList, AggregationApplyShader, ApplyParams, FIntVector(NumThreadGroups, 1, 1));
 
-			RHICmdList.Transition(FRHITransitionInfo(ApplyParams.OutputTangentsAndDistances, ERHIAccess::UAVMask, ERHIAccess::SRVMask));
-			RHICmdList.Transition(FRHITransitionInfo(ApplyParams.OutputMultiRibbonIndices, ERHIAccess::UAVMask, ERHIAccess::SRVMask));
-			RHICmdList.Transition(FRHITransitionInfo(ApplyParams.OutputSegments, ERHIAccess::UAVMask, ERHIAccess::SRVMask));
-			RHICmdList.Transition(FRHITransitionInfo(ApplyParams.OutputTessellationStats, ERHIAccess::UAVMask, ERHIAccess::SRVMask));
+			for (FRHITransitionInfo& Transition : Transitions)
+			{
+				Swap(Transition.AccessBefore, Transition.AccessAfter);
+			}
+			RHICmdList.Transition(Transitions);
 		}
 
 		static constexpr int32 CommandBufferOffset = 0;
@@ -2779,9 +2799,9 @@ void FNiagaraRendererRibbons::InitializeVertexBuffersGPU(FRHICommandListImmediat
 			FinalizationParams.Common = CommonParams;
 			FinalizationParams.SortedIndices = VertexBuffers.SortedIndicesBuffer.SRV;
 			FinalizationParams.TangentsAndDistances = VertexBuffers.TangentsAndDistancesBuffer.SRV;
-			FinalizationParams.MultiRibbonIndices = VertexBuffers.MultiRibbonIndicesBuffer.SRV;
+			FinalizationParams.MultiRibbonIndices = FNiagaraRenderer::GetSrvOrDefaultUInt(VertexBuffers.MultiRibbonIndicesBuffer.SRV);
+			FinalizationParams.TessellationStats = FNiagaraRenderer::GetSrvOrDefaultFloat(TempBuffers.TransientTessellationStats.SRV);
 			FinalizationParams.Segments = VertexBuffers.SegmentsBuffer.SRV;
-			FinalizationParams.TessellationStats = TempBuffers.TransientTessellationStats.SRV;
 			FinalizationParams.AccumulationBuffer = CurrentAccumulationSRV;
 			FinalizationParams.PackedPerRibbonData = VertexBuffers.RibbonLookupTableBuffer.UAV;
 			FinalizationParams.OutputCommandBuffer = VertexBuffers.GPUComputeCommandBuffer.UAV;
