- Wait for concurrent tick if we need to flush parameter stores in GetMaterialStreamingInfo
- Speculative fix for low repro issue

diff --git a/Engine/Plugins/FX/Niagara/Source/Niagara/Private/NiagaraSystemInstanceController.cpp b/Engine/Plugins/FX/Niagara/Source/Niagara/Private/NiagaraSystemInstanceController.cpp
--- a/Engine/Plugins/FX/Niagara/Source/Niagara/Private/NiagaraSystemInstanceController.cpp
+++ b/Engine/Plugins/FX/Niagara/Source/Niagara/Private/NiagaraSystemInstanceController.cpp
@@ -162,11 +162,24 @@ void FNiagaraSystemInstanceController::GetMaterialStreamingInfo(FNiagaraMaterial
 		return;
 	}
 
-	if (OverrideParameters)
+	// If renderers read from parameters that might have changed in the parent we need to ensure it's flushed
 	{
-		OverrideParameters->Tick();
+		FNiagaraParameterStore& InstanceParameters = SystemInstance->GetInstanceParameters();
+		const bool bInstancesDirty = InstanceParameters.GetInterfacesDirty() || InstanceParameters.GetUObjectsDirty();
+		const bool bOverridesDirty = OverrideParameters && (OverrideParameters->GetInterfacesDirty() || OverrideParameters->GetUObjectsDirty());
+		if (bInstancesDirty || bOverridesDirty)
+		{
+			SystemInstance->WaitForConcurrentTickDoNotFinalize();
+			if (bOverridesDirty)
+			{
+				OverrideParameters->Tick();
+			}
+			if (bInstancesDirty)
+			{
+				InstanceParameters.Tick();
+			}
+		}
 	}
-	SystemInstance->GetInstanceParameters().Tick();
 
 	for (const FNiagaraEmitterInstanceRef& EmitterInst : SystemInstance->GetEmitters())
 	{
