- Add support for large numbers of particles to Niagara keygen
- This fixes overflowing the thread group counts

#tests Make a LWE that spawns ~10 million sprite particles and enable sorting on the renderer

diff --git a/Engine/Plugins/FX/Niagara/Shaders/Private/NiagaraSortKeyGen.usf b/Engine/Plugins/FX/Niagara/Shaders/Private/NiagaraSortKeyGen.usf
--- a/Engine/Plugins/FX/Niagara/Shaders/Private/NiagaraSortKeyGen.usf
+++ b/Engine/Plugins/FX/Niagara/Shaders/Private/NiagaraSortKeyGen.usf
@@ -15,6 +15,7 @@
 ------------------------------------------------------------------------------*/
 
 #include "/Engine/Private/Common.ush"
+#include "/Engine/Private/ComputeShaderUtils.ush"
 
 Buffer<float> NiagaraParticleDataFloat;
 Buffer<half> NiagaraParticleDataHalf;
@@ -216,9 +217,11 @@ bool CullParticle(in uint ParticleIndex)
 }
 
 [numthreads(THREAD_COUNT,1,1)]
-void GenerateParticleSortKeys(uint ParticleIndex : SV_DispatchThreadID, uint GroupTID : SV_GroupThreadID)
+void GenerateParticleSortKeys(uint3 GroupId : SV_GroupID, int GroupThreadIndex : SV_GroupIndex)
 {
-	CullingPrologue(GroupTID);
+	const uint ParticleIndex = GetUnWrappedDispatchThreadId(GroupId, GroupThreadIndex, THREAD_COUNT);
+	
+	CullingPrologue(GroupThreadIndex);
 
 	const uint GPUParticleCount = GPUParticleCountOffset != -1 ? GPUParticleCountBuffer[GPUParticleCountOffset] : ParticleCount;
 	
@@ -287,6 +290,6 @@ void GenerateParticleSortKeys(uint ParticleIndex : SV_DispatchThreadID, uint Gro
 		OutParticleIndices[OutputOffset + ParticleIndex] = ParticleIndex;
 	}
 
-	CullingEpilogue(bCulled, GroupTID);
+	CullingEpilogue(bCulled, GroupThreadIndex);
 }
 
diff --git a/Engine/Plugins/FX/Niagara/Source/Niagara/Private/NiagaraGpuComputeDispatch.cpp b/Engine/Plugins/FX/Niagara/Source/Niagara/Private/NiagaraGpuComputeDispatch.cpp
--- a/Engine/Plugins/FX/Niagara/Source/Niagara/Private/NiagaraGpuComputeDispatch.cpp
+++ b/Engine/Plugins/FX/Niagara/Source/Niagara/Private/NiagaraGpuComputeDispatch.cpp
@@ -2449,11 +2449,12 @@ void FNiagaraGpuComputeDispatch::GenerateSortKeys(FRHICommandListImmediate& RHIC
 			}
 
 			// Choose the shader to bind
+			const int32 NumThreadGroups = FMath::DivideAndRoundUp(SortInfo.ParticleCount, NIAGARA_KEY_GEN_THREAD_COUNT);
 			FComputeShaderUtils::Dispatch(
 				RHICmdList,
 				SortInfo.bEnableCulling ? SortAndCullKeyGenCS : SortKeyGenCS,
 				Params,
-				FIntVector(FMath::DivideAndRoundUp(SortInfo.ParticleCount, NIAGARA_KEY_GEN_THREAD_COUNT), 1, 1)
+				FComputeShaderUtils::GetGroupCountWrapped(NumThreadGroups)
 			);
 		}
 	}
