- Add stack information if the default mesh is set and the source mode is set to default.
- Often this is an incorrect setup as the mesh is spawned attached and the default mesh is likely ignored but still loaded / referenced
- Add post edit / can edit changes to the SKM DI to ensure the default mesh is cleared out when it can not be used

#tests Create a SKM DI, assign a default mesh, change source mode to one that does not sample the default mesh

diff --git a/Engine/Plugins/FX/Niagara/Source/Niagara/Private/DataInterface/NiagaraDataInterfaceStaticMesh.cpp b/Engine/Plugins/FX/Niagara/Source/Niagara/Private/DataInterface/NiagaraDataInterfaceStaticMesh.cpp
--- a/Engine/Plugins/FX/Niagara/Source/Niagara/Private/DataInterface/NiagaraDataInterfaceStaticMesh.cpp
+++ b/Engine/Plugins/FX/Niagara/Source/Niagara/Private/DataInterface/NiagaraDataInterfaceStaticMesh.cpp
@@ -31,6 +31,7 @@
 #include "INiagaraEditorOnlyDataUtlities.h"
 #include "Modules/ModuleManager.h"
 #include "NiagaraModule.h"
+#include "ScopedTransaction.h"
 #endif
 
 #include "Engine/StaticMeshActor.h"
@@ -544,7 +545,7 @@ namespace NDIStaticMeshLocal
 			bIsValid = false;
 
 		#if STATS
-			INC_MEMORY_STAT_BY(STAT_NiagaraGPUDataInterfaceMemory, GPUMemoryUsage);
+			DEC_MEMORY_STAT_BY(STAT_NiagaraGPUDataInterfaceMemory, GPUMemoryUsage);
 			GPUMemoryUsage = 0;
 		#endif
 
@@ -725,7 +726,7 @@ namespace NDIStaticMeshLocal
 			int32 MaxSampleLOD = StaticMesh->GetNumLODs() - 1;
 
 			OutLOD = RequestedLOD >= 0 ? RequestedLOD : (StaticMesh->GetNumLODs() + RequestedLOD);
-			OutLOD = FMath::Clamp(CachedLODIdx, MinSampleLOD, MaxSampleLOD);
+			OutLOD = FMath::Clamp(OutLOD, MinSampleLOD, MaxSampleLOD);
 
 			return StaticMesh->GetRenderData()->LODResources.IsValidIndex(OutLOD);
 		}
@@ -3414,6 +3415,23 @@ void UNiagaraDataInterfaceStaticMesh::GetFeedback(UNiagaraSystem* Asset, UNiagar
 		OutWarnings.Add(NoMeshAssignedError);
 	}
 
+	if (SourceMode == ENDIStaticMesh_SourceMode::Default && DefaultMesh != nullptr)
+	{
+		OutWarnings.Emplace(
+			LOCTEXT("DefaultSourceAndDefaultMesh", "Default Mesh is set but Source Mode is set to Default.  Verify you have not left it set incorrectly to avoid unused meshes being referenced."),
+			LOCTEXT("DefaultSourceAndDefaultMeshSummary", "Potentially unused Default Mesh set!"),
+			FNiagaraDataInterfaceFix::CreateLambda(
+				[this]()
+				{
+					FScopedTransaction ST(LOCTEXT("ClearDefaultMesh", "Clear Default Mesh"));
+					Modify();
+					DefaultMesh = nullptr;
+					return true;
+				}
+			)
+		);
+	}
+
 	if (GetDefault<UNiagaraSettings>()->NDIStaticMesh_AllowDistanceFields == false)
 	{
 		auto GenerateWarnings = 
@@ -3604,8 +3622,8 @@ void UNiagaraDataInterfaceStaticMesh::VMGetMeshBounds(FVectorVMExternalFunctionC
 	const FBox MeshBounds(StaticMeshHelper.InstanceData->PreSkinnedLocalBoundsCenter - StaticMeshHelper.InstanceData->PreSkinnedLocalBoundsExtents, StaticMeshHelper.InstanceData->PreSkinnedLocalBoundsCenter + StaticMeshHelper.InstanceData->PreSkinnedLocalBoundsExtents);
 	const FBox WorldBounds = MeshBounds.TransformBy(StaticMeshHelper.InstanceData->Transform);
 
-	const FVector3f Center = bLocalSpace ? FVector3f(MeshBounds.GetCenter())  : FVector3f(WorldBounds.GetCenter());
-	const FVector3f HalfExtents = FVector3f(WorldBounds.GetExtent());
+	const FVector3f Center = bLocalSpace ? FVector3f(MeshBounds.GetCenter()) : FVector3f(WorldBounds.GetCenter());
+	const FVector3f HalfExtents = bLocalSpace ? FVector3f(MeshBounds.GetExtent()) : FVector3f(WorldBounds.GetExtent());
 	const FVector3f Extents = HalfExtents * 2.0f;
 	const FVector3f ExtentsMin = Center - HalfExtents;
 	const FVector3f ExtentsMax = Center + HalfExtents;
diff --git a/Engine/Plugins/FX/Niagara/Source/Niagara/Private/NiagaraDataInterfaceSkeletalMesh.cpp b/Engine/Plugins/FX/Niagara/Source/Niagara/Private/NiagaraDataInterfaceSkeletalMesh.cpp
--- a/Engine/Plugins/FX/Niagara/Source/Niagara/Private/NiagaraDataInterfaceSkeletalMesh.cpp
+++ b/Engine/Plugins/FX/Niagara/Source/Niagara/Private/NiagaraDataInterfaceSkeletalMesh.cpp
@@ -30,6 +30,10 @@
 #include "CachedGeometry.h"
 #include "SkeletalRenderPublic.h"
 
+#if WITH_EDITOR
+#include "ScopedTransaction.h"
+#endif
+
 #include UE_INLINE_GENERATED_CPP_BY_NAME(NiagaraDataInterfaceSkeletalMesh)
 
 #define LOCTEXT_NAMESPACE "NiagaraDataInterfaceSkeletalMesh"
@@ -2364,14 +2368,20 @@ void UNiagaraDataInterfaceSkeletalMesh::PostEditChangeProperty(FPropertyChangedE
 	Super::PostEditChangeProperty(PropertyChangedEvent);
 	ChangeId++;
 
-	if (PropertyChangedEvent.Property &&
-		PropertyChangedEvent.Property->GetFName() == GET_MEMBER_NAME_CHECKED(UNiagaraDataInterfaceSkeletalMesh, SourceMode) &&
-		SourceMode != ENDISkeletalMesh_SourceMode::Default &&
-		SourceMode != ENDISkeletalMesh_SourceMode::Source)
+	if (PropertyChangedEvent.Property && PropertyChangedEvent.Property->GetFName() == GET_MEMBER_NAME_CHECKED(UNiagaraDataInterfaceSkeletalMesh, SourceMode))
 	{
-		// Clear out any source that is set to prevent unnecessary references, since we won't even consider them
-		SoftSourceActor = nullptr;
-		SourceComponent = nullptr;
+		if (SourceMode != ENDISkeletalMesh_SourceMode::Default && SourceMode != ENDISkeletalMesh_SourceMode::Source)
+		{
+			// Ensure we don't have any reference to a source actor that we'll never use
+			SoftSourceActor = nullptr;
+			SourceComponent = nullptr;
+		}
+
+		if (SourceMode != ENDISkeletalMesh_SourceMode::Default && SourceMode != ENDISkeletalMesh_SourceMode::DefaultMeshOnly)
+		{
+			// Ensure we don't cook in a default mesh we'll never use
+			DefaultMesh = nullptr;
+		}
 	}
 }
 
@@ -2382,12 +2392,14 @@ bool UNiagaraDataInterfaceSkeletalMesh::CanEditChange(const FProperty* InPropert
 		return false;
 	}
 
-	if (InProperty->GetFName() == GET_MEMBER_NAME_CHECKED(UNiagaraDataInterfaceSkeletalMesh, SoftSourceActor) &&
-		SourceMode != ENDISkeletalMesh_SourceMode::Default &&
-		SourceMode != ENDISkeletalMesh_SourceMode::Source)
+	if (InProperty->GetFName() == GET_MEMBER_NAME_CHECKED(UNiagaraDataInterfaceSkeletalMesh, SoftSourceActor) || InProperty->GetFName() == GET_MEMBER_NAME_CHECKED(UNiagaraDataInterfaceSkeletalMesh, SourceComponent))
 	{
-		// Disable "Source" if it won't be considered
-		return false;
+		return SourceMode == ENDISkeletalMesh_SourceMode::Default || SourceMode == ENDISkeletalMesh_SourceMode::Source;
+	}
+
+	if (InProperty->GetFName() == GET_MEMBER_NAME_CHECKED(UNiagaraDataInterfaceSkeletalMesh, DefaultMesh))
+	{
+		return SourceMode == ENDISkeletalMesh_SourceMode::Default || SourceMode == ENDISkeletalMesh_SourceMode::DefaultMeshOnly;
 	}
 
 	return true;
@@ -2633,7 +2645,6 @@ void UNiagaraDataInterfaceSkeletalMesh::GetFeedback(UNiagaraSystem* Asset, UNiag
 	USkeletalMesh* SkelMesh = GetSkeletalMesh(Component);
 
 	// Collect Errors
-#if WITH_EDITORONLY_DATA
 	if (SkelMesh != nullptr)
 	{
 		bool bHasCPUAccess = true;
@@ -2773,7 +2784,6 @@ void UNiagaraDataInterfaceSkeletalMesh::GetFeedback(UNiagaraSystem* Asset, UNiag
 			);
 		}
 	}
-#endif
 
 	if (SoftSourceActor.Get() == nullptr && bHasNoMeshAssignedWarning)
 	{
@@ -2784,6 +2794,23 @@ void UNiagaraDataInterfaceSkeletalMesh::GetFeedback(UNiagaraSystem* Asset, UNiag
 		OutWarnings.Add(NoMeshAssignedError);
 	}
 
+	if (SourceMode == ENDISkeletalMesh_SourceMode::Default && DefaultMesh != nullptr)
+	{
+		OutWarnings.Emplace(
+			LOCTEXT("DefaultSourceAndDefaultMesh", "Default Mesh is set but Source Mode is set to Default.  Verify you have not left it set incorrectly to avoid unused meshes being referenced."),
+			LOCTEXT("DefaultSourceAndDefaultMeshSummary", "Potentially unused Default Mesh set!"),
+			FNiagaraDataInterfaceFix::CreateLambda(
+				[this]()
+				{
+					FScopedTransaction ST(LOCTEXT("ClearDefaultMesh", "Clear Default Mesh"));
+					Modify();
+					DefaultMesh = nullptr;
+					return true;
+				}
+			)
+		);
+	}
+
 	// Look for bones being used that are LOD'ed out
 	if ( (SkelMesh != nullptr) && (SkelMesh->GetResourceForRendering() != nullptr) )
 	{
