- Don't flush parameters if not on the game thread
- This can occur during AddPrimitive in EOF updates

#tests Ran ET / PIE testing
#rnx

diff --git a/Engine/Plugins/FX/Niagara/Source/Niagara/Private/NiagaraSystemInstanceController.cpp b/Engine/Plugins/FX/Niagara/Source/Niagara/Private/NiagaraSystemInstanceController.cpp
--- a/Engine/Plugins/FX/Niagara/Source/Niagara/Private/NiagaraSystemInstanceController.cpp
+++ b/Engine/Plugins/FX/Niagara/Source/Niagara/Private/NiagaraSystemInstanceController.cpp
@@ -163,6 +163,8 @@ void FNiagaraSystemInstanceController::GetMaterialStreamingInfo(FNiagaraMaterial
 	}
 
 	// If renderers read from parameters that might have changed in the parent we need to ensure it's flushed
+	// Note that this can be called during AddPrimitive on concurrent threads
+	if ( IsInGameThread() )
 	{
 		FNiagaraParameterStore& InstanceParameters = SystemInstance->GetInstanceParameters();
 		const bool bInstancesDirty = InstanceParameters.GetInterfacesDirty() || InstanceParameters.GetUObjectsDirty();
