- Fix for LW parameters not always being up to date on the RT
- This fixes a crash when forcing the CPU path

diff --git a/Engine/Plugins/FX/Niagara/Source/Niagara/Internal/Stateless/NiagaraStatelessEmitterInstance.h b/Engine/Plugins/FX/Niagara/Source/Niagara/Internal/Stateless/NiagaraStatelessEmitterInstance.h
--- a/Engine/Plugins/FX/Niagara/Source/Niagara/Internal/Stateless/NiagaraStatelessEmitterInstance.h
+++ b/Engine/Plugins/FX/Niagara/Source/Niagara/Internal/Stateless/NiagaraStatelessEmitterInstance.h
@@ -92,7 +92,7 @@ private:
 	void InitEmitterState();
 	void TickEmitterState();
 	void CalculateBounds();
-	void UpdateSimulationData(float DeltaSeconds);
+	void UpdateSimulationData(float DeltaSeconds, bool bForceParametersUpdate);
 
 	void InitSpawnInfos(float InitializationAge);
 	void InitSpawnInfosForLoop(float InitializationAge);
diff --git a/Engine/Plugins/FX/Niagara/Source/Niagara/Private/Stateless/NiagaraStatelessEmitterInstance.cpp b/Engine/Plugins/FX/Niagara/Source/Niagara/Private/Stateless/NiagaraStatelessEmitterInstance.cpp
--- a/Engine/Plugins/FX/Niagara/Source/Niagara/Private/Stateless/NiagaraStatelessEmitterInstance.cpp
+++ b/Engine/Plugins/FX/Niagara/Source/Niagara/Private/Stateless/NiagaraStatelessEmitterInstance.cpp
@@ -330,6 +330,7 @@ bool FNiagaraStatelessEmitterInstance::ShouldTick() const
 
 void FNiagaraStatelessEmitterInstance::Tick(float DeltaSeconds)
 {
+	const bool bForceParametersUpdate = bNeedsEmitterStateInit;
 	if ( bNeedsEmitterStateInit )
 	{
 		bNeedsEmitterStateInit = false;
@@ -342,7 +343,7 @@ void FNiagaraStatelessEmitterInstance::Tick(float DeltaSeconds)
 	TickSpawnInfos();
 	TickEmitterState();
 	CalculateBounds();
-	UpdateSimulationData(DeltaSeconds);
+	UpdateSimulationData(DeltaSeconds, bForceParametersUpdate);
 }
 
 void FNiagaraStatelessEmitterInstance::InitEmitterState()
@@ -543,13 +544,13 @@ void FNiagaraStatelessEmitterInstance::CalculateBounds()
 	}
 }
 
-void FNiagaraStatelessEmitterInstance::UpdateSimulationData(float DeltaSeconds)
+void FNiagaraStatelessEmitterInstance::UpdateSimulationData(float DeltaSeconds, bool bForceParametersUpdate)
 {
 	using namespace NiagaraStateless;
 
 	// Update parameter data (if needed)
 	bool bNeedsShaderParametersUpdate = false;
-	bool bNeedsParametersUpdate = false;
+	bool bNeedsParametersUpdate = bForceParametersUpdate;
 	if (RendererBindings.GetParametersDirty())
 	{
 		for (const TPair<int32, FInstancedStruct>& ExpressionMapping : EmitterData->Expressions)
