- Modify post reachability to only invalidate if the parameter store contains UObjects
- Change iterator to ignore invalid components

#tests Ran ET
#rnx

diff --git a/Engine/Plugins/FX/Niagara/Source/Niagara/Private/NiagaraWorldManager.cpp b/Engine/Plugins/FX/Niagara/Source/Niagara/Private/NiagaraWorldManager.cpp
--- a/Engine/Plugins/FX/Niagara/Source/Niagara/Private/NiagaraWorldManager.cpp
+++ b/Engine/Plugins/FX/Niagara/Source/Niagara/Private/NiagaraWorldManager.cpp
@@ -1023,10 +1023,15 @@ void FNiagaraWorldManager::PostReachabilityAnalysis()
 	// We must also mark the component render data as dirty to ensure any material references in renderers are updated as they could also have become nullptr.
 	for (TObjectIterator<UNiagaraComponent> ComponentIt; ComponentIt; ++ComponentIt)
 	{
-		ComponentIt->GetOverrideParameters().MarkUObjectsDirty();
-		if (ComponentIt->GetOverrideParameters().GetUObjects().Num() > 0)
+		UNiagaraComponent* Component = *ComponentIt;
+		if (IsValid(Component) && !Component->HasAnyFlags(EObjectFlags::RF_ClassDefaultObject))
 		{
-			ComponentIt->MarkRenderStateDirty();
+			FNiagaraUserRedirectionParameterStore& ParameterStore = Component->GetOverrideParameters();
+			if (ParameterStore.GetUObjects().Num() > 0)
+			{
+				ParameterStore.MarkUObjectsDirty();
+				Component->MarkRenderStateDirty();
+			}
 		}
 	}
 }
