- Fix a crash if ScriptRuntimeData is null
- This came up when duplicating a BP / NS then assigning the NS to the duplicated BP

#rnx

diff --git a/Engine/Plugins/FX/Niagara/Source/Niagara/Private/NiagaraScriptExecutionContext.cpp b/Engine/Plugins/FX/Niagara/Source/Niagara/Private/NiagaraScriptExecutionContext.cpp
--- a/Engine/Plugins/FX/Niagara/Source/Niagara/Private/NiagaraScriptExecutionContext.cpp
+++ b/Engine/Plugins/FX/Niagara/Source/Niagara/Private/NiagaraScriptExecutionContext.cpp
@@ -654,6 +654,11 @@ bool FNiagaraSystemScriptExecutionContext::Execute(FNiagaraSystemInstance* Paren
 
 bool FNiagaraSystemScriptExecutionContext::GeneratePerInstanceDIFunctionTable(FNiagaraSystemInstance* Inst, TArray<FNiagaraPerInstanceDIFuncInfo>& OutFunctions)
 {
+	if (ScriptRuntimeData == nullptr)
+	{
+		return false;
+	}
+
 	const FNiagaraScriptExecutionParameterStore* ScriptParameterStore = ScriptRuntimeData->GetExecutionReadyParameterStore();
 	const FNiagaraVMExecutableData& ScriptExecutableData = Script->GetVMExecutableData();
 	TConstArrayView<FNiagaraScriptResolvedDataInterfaceInfo> ScriptDataInterfaceInfos = ScriptRuntimeData->GetResolvedDataInterfaces();
