Fix presentation of user parameters in the Parameter panel.  Ensures that they are marked as bExternallyReferenced if they are...this prevents us from deleting them which can currently incorrect behavior.

#tests Created a system using a float, color and bool arrays as user DI, and then tried to delete the bool one from the parameter panel

diff --git a/Engine/Plugins/FX/Niagara/Source/NiagaraEditor/Private/ViewModels/NiagaraParameterPanelViewModel.cpp b/Engine/Plugins/FX/Niagara/Source/NiagaraEditor/Private/ViewModels/NiagaraParameterPanelViewModel.cpp
--- a/Engine/Plugins/FX/Niagara/Source/NiagaraEditor/Private/ViewModels/NiagaraParameterPanelViewModel.cpp
+++ b/Engine/Plugins/FX/Niagara/Source/NiagaraEditor/Private/ViewModels/NiagaraParameterPanelViewModel.cpp
@@ -1340,17 +1340,10 @@ void FNiagaraSystemToolkitParameterPanelViewModel::DeleteParameters(const TArray
 				continue;
 			}
 
-			UNiagaraSystem& System = SystemViewModel->GetSystem();
-			const FGuid& ScriptVarId = ItemToDelete.GetVariableMetaData().GetVariableGuid();
-			System.Modify();
-			System.GetExposedParameters().RemoveParameter(ItemToDelete.GetVariable());
-			UNiagaraEditorParametersAdapter* EditorParametersAdapter = SystemViewModel->GetEditorOnlyParametersAdapter();
-			EditorParametersAdapter->Modify();
-			EditorParametersAdapter->GetParameters().RemoveAll([&ScriptVarId](const UNiagaraScriptVariable* ScriptVar) { return ScriptVar->Metadata.GetVariableGuid() == ScriptVarId; });
-
-			// Update anything that was referencing that parameter
-			System.HandleVariableRemoved(ItemToDelete.GetVariable(), true);
-			bAnyChange = true;
+			if (SystemViewModel->DeleteParameter(ItemToDelete.GetVariable()))
+			{
+				bAnyChange = true;
+			}
 		}
 	}
 
@@ -2109,19 +2102,9 @@ TArray<FNiagaraParameterPanelItem> FNiagaraSystemToolkitParameterPanelViewModel:
 
 					TArray<TNiagaraParameterMapHistory<FNiagaraCompilationGraphBridge>::FReadHistory> ReadHistory = Builder.Histories[0].PerVariableReadHistory[VariableIndex];
 					TArray<TModuleScopedPin<UEdGraphPin>> WriteHistory = Builder.Histories[0].PerVariableWriteHistory[VariableIndex];
-					if (FNiagaraParameterPanelItem* ItemPtr = VisitedParameterToItemMap.Find(Var))
-					{
-						// This variable has already been registered, increment the reference counts.
-						for (const FNiagaraParameterMapHistory::FReadHistory& History : ReadHistory)
-						{
-							AddReadReference(History, *ItemPtr);
-						}
-						for (const FNiagaraParameterMapHistory::FModuleScopedPin& History : WriteHistory)
-						{
-							AddWriteReference(History, *ItemPtr);
-						}
-					}
-					else  // Add newly found variables
+					FNiagaraParameterPanelItem* ItemPtr = VisitedParameterToItemMap.Find(Var);
+					// Add newly found variables
+					if (ItemPtr == nullptr)
 					{
 						// This variable has not been registered, prepare the FNiagaraParameterPanelItem.
 						// -First make sure the variable namespace is in a valid category. If not, skip it.
@@ -2148,25 +2131,10 @@ TArray<FNiagaraParameterPanelItem> FNiagaraSystemToolkitParameterPanelViewModel:
 							TransientParameterToScriptVarMap.Add(Var, ScriptVar);
 						}
 
-						bool bVarOnlyInTopLevelGraph = true;
-						if (!bForceScript)
-						{
-							for (FNiagaraParameterMapHistory::FModuleScopedPin& WritePin : WriteHistory)
-							{
-								UEdGraphNode* VariableOwningNode = WritePin.Pin->GetOwningNode();
-								bVarOnlyInTopLevelGraph &= AllGraphs.Contains(static_cast<const UNiagaraGraph*>(VariableOwningNode->GetGraph()));
-							}
-							for (FNiagaraParameterMapHistory::FReadHistory& ReadPins : ReadHistory)
-							{
-								UEdGraphNode* VariableOwningNode = ReadPins.ReadPin.Pin->GetOwningNode();
-								bVarOnlyInTopLevelGraph &= !AllGraphs.Contains(static_cast<const UNiagaraGraph*>(VariableOwningNode->GetGraph()));
-							}
-						}
-
 						FNiagaraParameterPanelItem Item = FNiagaraParameterPanelItem();
 						Item.ScriptVariable = ScriptVar;
 						Item.NamespaceMetaData = CandidateNamespaceMetaData;
-						Item.bExternallyReferenced = !bVarOnlyInTopLevelGraph;
+						Item.bExternallyReferenced = false;
 
 						// -Determine whether the parameter is from a custom stack context.
 						Item.bSourcedFromCustomStackContext = false;
@@ -2182,17 +2150,37 @@ TArray<FNiagaraParameterPanelItem> FNiagaraSystemToolkitParameterPanelViewModel:
 						// Determine whether the item is name aliasing a parameter definition's parameter.
 						Item.DefinitionMatchState = FNiagaraParameterDefinitionsUtilities::GetDefinitionMatchStateForParameter(Item.ScriptVariable->Variable);
 
-						// -Increment the reference count.
-						for (const FNiagaraParameterMapHistory::FReadHistory& History : ReadHistory)
-						{
-							AddReadReference(History, Item);
-						}
-						for (const FNiagaraParameterMapHistory::FModuleScopedPin& History : WriteHistory)
+						ItemPtr = &VisitedParameterToItemMap.Add(Var, Item);
+					}
+
+					if (!ItemPtr->bExternallyReferenced)
+					{
+						bool bVarOnlyInTopLevelGraph = true;
+						if (!bForceScript)
 						{
-							AddWriteReference(History, Item);
+							for (FNiagaraParameterMapHistory::FModuleScopedPin& WritePin : WriteHistory)
+							{
+								UEdGraphNode* VariableOwningNode = WritePin.Pin->GetOwningNode();
+								bVarOnlyInTopLevelGraph &= AllGraphs.Contains(static_cast<const UNiagaraGraph*>(VariableOwningNode->GetGraph()));
+							}
+							for (FNiagaraParameterMapHistory::FReadHistory& ReadPins : ReadHistory)
+							{
+								UEdGraphNode* VariableOwningNode = ReadPins.ReadPin.Pin->GetOwningNode();
+								bVarOnlyInTopLevelGraph &= !AllGraphs.Contains(static_cast<const UNiagaraGraph*>(VariableOwningNode->GetGraph()));
+							}
 						}
 
-						VisitedParameterToItemMap.Add(Var, Item);
+						ItemPtr->bExternallyReferenced = !bVarOnlyInTopLevelGraph;
+					}
+
+					// -Increment the reference count.
+					for (const FNiagaraParameterMapHistory::FReadHistory& History : ReadHistory)
+					{
+						AddReadReference(History, *ItemPtr);
+					}
+					for (const FNiagaraParameterMapHistory::FModuleScopedPin& History : WriteHistory)
+					{
+						AddWriteReference(History, *ItemPtr);
 					}
 				}
 			}
diff --git a/Engine/Plugins/FX/Niagara/Source/NiagaraEditor/Private/ViewModels/NiagaraSystemViewModel.cpp b/Engine/Plugins/FX/Niagara/Source/NiagaraEditor/Private/ViewModels/NiagaraSystemViewModel.cpp
--- a/Engine/Plugins/FX/Niagara/Source/NiagaraEditor/Private/ViewModels/NiagaraSystemViewModel.cpp
+++ b/Engine/Plugins/FX/Niagara/Source/NiagaraEditor/Private/ViewModels/NiagaraSystemViewModel.cpp
@@ -1523,6 +1523,19 @@ bool FNiagaraSystemViewModel::DeleteParameter(const FNiagaraVariableBase Paramet
 		
 		System->GetExposedParameters().RemoveParameter(Parameter);
 	}
+	else
+	{
+		if (UNiagaraEditorParametersAdapter* EditorParametersAdapter = GetEditorOnlyParametersAdapter())
+		{
+			auto MatchingParameter = [&Parameter](const UNiagaraScriptVariable* ScriptVariable) -> bool
+			{
+				return ScriptVariable && ScriptVariable->Variable.IsEquivalent(Parameter);
+			};
+
+			EditorParametersAdapter->Modify();
+			EditorParametersAdapter->GetParameters().RemoveAll(MatchingParameter);
+		}
+	}
 
 	System->HandleVariableRemoved(Parameter, true);
 
