- Fix for Linear Color material parameters not having RGBA editable fields

#rnx

diff --git a/Engine/Plugins/FX/Niagara/Source/NiagaraEditor/Private/Customizations/NiagaraTypeCustomizations.cpp b/Engine/Plugins/FX/Niagara/Source/NiagaraEditor/Private/Customizations/NiagaraTypeCustomizations.cpp
--- a/Engine/Plugins/FX/Niagara/Source/NiagaraEditor/Private/Customizations/NiagaraTypeCustomizations.cpp
+++ b/Engine/Plugins/FX/Niagara/Source/NiagaraEditor/Private/Customizations/NiagaraTypeCustomizations.cpp
@@ -52,6 +52,7 @@
 #include "Stateless/NiagaraStatelessEmitter.h"
 //-TODO:stateless:Remove and unify emitter
 
+#include "Widgets/SNiagaraColorEditor.h"
 #include "Widgets/SNiagaraParameterMenu.h"
 
 #include UE_INLINE_GENERATED_CPP_BY_NAME(NiagaraTypeCustomizations)
@@ -281,6 +282,86 @@ struct FNiagaraMaterialParameterValueHelper : public TSharedFromThis<FNiagaraMat
 		return FReply::Handled();
 	}
 
+	template<typename TType>
+	TSharedRef<SWidget> CreatePropertyWidget()
+	{
+		return ValuePropertyHandle->CreatePropertyValueWidget();
+	}
+
+	template<>
+	TSharedRef<SWidget> CreatePropertyWidget<FLinearColor>()
+	{
+		return SNew(SNiagaraColorEditor)
+			.ShowAlpha(true)
+			.ShowExpander(true)
+			.Color(this, &FNiagaraMaterialParameterValueHelper::GetColorValue)
+			.OnColorChanged(this, &FNiagaraMaterialParameterValueHelper::SetColorValue)
+			.OnBeginEditing(this, &FNiagaraMaterialParameterValueHelper::BeginContinuousTransaction)
+			.OnEndEditing(this, &FNiagaraMaterialParameterValueHelper::EndContinuousTransaction)
+			.OnCancelEditing(this, &FNiagaraMaterialParameterValueHelper::CancelContinuousTransaction)
+		;
+	}
+
+	FLinearColor GetColorValue() const
+	{
+		FLinearColor Color = FLinearColor::White;
+		GetPropertyValue(Color);
+		return Color;
+	}
+
+	void SetColorValue(FLinearColor InColorValue)
+	{
+		const FText TransactionText = LOCTEXT("SetColorValue", "Set Color Value");
+		if (bContinuousTransactionPending)
+		{
+			bContinuousTransactionPending = false;
+			ContinuousTransactionIndex = GEditor->BeginTransaction(TransactionText);
+		}
+
+		const bool bIsContinousTransaction = ContinuousTransactionIndex.IsSet();
+
+		FScopedTransaction Transaction(TransactionText, !bIsContinousTransaction);
+
+		TArray<UObject*> OuterObjects;
+		ValuePropertyHandle->GetOuterObjects(OuterObjects);
+		for (UObject* OuterObject : OuterObjects)
+		{
+			OuterObject->Modify();
+		}
+
+		ValuePropertyHandle->NotifyPreChange();
+
+		SetPropertyValue(InColorValue);
+
+		ValuePropertyHandle->NotifyPostChange(bIsContinousTransaction ? EPropertyChangeType::Interactive : EPropertyChangeType::ValueSet);
+	}
+
+	void BeginContinuousTransaction()
+	{
+		bContinuousTransactionPending = true;
+	}
+
+	void EndContinuousTransaction()
+	{
+		bContinuousTransactionPending = false;
+		if (ContinuousTransactionIndex.IsSet())
+		{
+			GEditor->EndTransaction();
+			ContinuousTransactionIndex.Reset();
+		}
+	}
+
+	void CancelContinuousTransaction(FLinearColor OriginalValue)
+	{
+		bContinuousTransactionPending = false;
+		if (ContinuousTransactionIndex.IsSet())
+		{
+			GEditor->CancelTransaction(ContinuousTransactionIndex.GetValue());
+			ContinuousTransactionIndex.Reset();
+		}
+		SetPropertyValue(OriginalValue);
+	}
+
 	template<typename TType>
 	bool CreateWidget(IDetailChildrenBuilder& ChildBuilder)
 	{
@@ -294,7 +375,7 @@ struct FNiagaraMaterialParameterValueHelper : public TSharedFromThis<FNiagaraMat
 			SNew(SHorizontalBox)
 			+ SHorizontalBox::Slot()
 			[
-				ValuePropertyHandle->CreatePropertyValueWidget()
+				CreatePropertyWidget<TType>()
 			]
 			+ SHorizontalBox::Slot()
 			.AutoWidth()
@@ -323,6 +404,9 @@ struct FNiagaraMaterialParameterValueHelper : public TSharedFromThis<FNiagaraMat
 	TSharedRef<IPropertyHandle>					NamePropertyHandle;
 	TSharedRef<IPropertyHandle>					ValuePropertyHandle;
 	TWeakObjectPtr<UNiagaraRendererProperties>	WeakRenderProperties;
+
+	bool				bContinuousTransactionPending = false;
+	TOptional<int32>	ContinuousTransactionIndex;
 };
 
 //////////////////////////////////////////////////////////////////////////////////////////////////////////////////////////////////////////////////////////////////
