- Add allowed classes to texture sample DI
- This stops the user from being enable to select an RVT, Array, Volume, etc.

diff --git a/Engine/Plugins/FX/Niagara/Source/Niagara/Classes/NiagaraDataInterfaceTexture.h b/Engine/Plugins/FX/Niagara/Source/Niagara/Classes/NiagaraDataInterfaceTexture.h
--- a/Engine/Plugins/FX/Niagara/Source/Niagara/Classes/NiagaraDataInterfaceTexture.h
+++ b/Engine/Plugins/FX/Niagara/Source/Niagara/Classes/NiagaraDataInterfaceTexture.h
@@ -21,7 +21,7 @@ class UNiagaraDataInterfaceTexture : public UNiagaraDataInterface
 	END_SHADER_PARAMETER_STRUCT()
 
 public:
-	UPROPERTY(EditAnywhere, Category = "Texture")
+	UPROPERTY(EditAnywhere, Category = "Texture", meta = (AllowedClasses = "/Script/Engine.Texture2D,/Script/Engine.TextureRenderTarget2D"))
 	TObjectPtr<UTexture> Texture;
 
 	UPROPERTY(EditAnywhere, Category = "Texture", meta = (ToolTip = "When valid the user parameter is used as the texture rather than the one on the data interface"))
