- Fix asan race detection by removing ununeccsary assignment
- GPUScript_RT was already set and checked, this is just setting the exact same value again

#rnx

diff --git a/Engine/Plugins/FX/Niagara/Source/Niagara/Private/NiagaraEmitterInstanceImpl.cpp b/Engine/Plugins/FX/Niagara/Source/Niagara/Private/NiagaraEmitterInstanceImpl.cpp
--- a/Engine/Plugins/FX/Niagara/Source/Niagara/Private/NiagaraEmitterInstanceImpl.cpp
+++ b/Engine/Plugins/FX/Niagara/Source/Niagara/Private/NiagaraEmitterInstanceImpl.cpp
@@ -1366,7 +1366,6 @@ void FNiagaraEmitterInstanceImpl::Tick(float DeltaSeconds)
 	if (EmitterData->SimTarget == ENiagaraSimTarget::GPUComputeSim && GPUExecContext != nullptr)
 	{
 		check(GPUExecContext->GPUScript_RT == EmitterData->GetGPUComputeScript()->GetRenderThreadScript());
-		GPUExecContext->GPUScript_RT = EmitterData->GetGPUComputeScript()->GetRenderThreadScript();
 
 #if NIAGARA_SYSTEM_CAPTURE
 		if (ParentSystemInstance->ShouldCaptureThisFrame())
