Allow indirect reading of Mesh distance fields in Niagara.
Allows NDC systems to be used to read mesh distance fields of arbitrary meshes in the scene.

= Niagara Data Managers
Generic manager class for handling of data that needs to be marshalled at the world level for us in Niagara. The GT counterparts to FNiagaraGpuComputeDataManager.
Extend FNiagaraGpuComputeDataManager to add BeginFrame EndFrame and PreTick calls.
Also adding path to just get a data manager in places where our compute dispatch interface is const.
Implementation of the GT and Gpu managers for exposing Indirect Static Mesh info.

= Indirect Static mesh DI
New DI that allows access to the indirect static mesh data manager buffers plus associated shader code to query distance field info.

= Misc
- New Lazy NDC Handler Init
- Adding Cached Layout to handlers so writer helpes can be invalidated if the NDC layout changes
- Some utiliy code for easier and standardized passing of packed data via an NDC.
- FNDCAccessContext can now optionally return all currently active systems rather than just newly spawned ones.

diff --git a/Engine/Plugins/FX/Niagara/Shaders/Private/NiagaraDataInterfaceStaticMeshIndirectTemplate.ush b/Engine/Plugins/FX/Niagara/Shaders/Private/NiagaraDataInterfaceStaticMeshIndirectTemplate.ush
new file mode 100644
--- /dev/null
+++ b/Engine/Plugins/FX/Niagara/Shaders/Private/NiagaraDataInterfaceStaticMeshIndirectTemplate.ush
@@ -0,0 +1,69 @@
+// Copyright Epic Games, Inc. All Rights Reserved.
+
+#include "/Plugin/FX/Niagara/Private/NiagaraQuaternionUtils.ush"
+#include "/Plugin/FX/Niagara/Private/NiagaraDistanceFieldQueries.ush"
+#include "/Plugin/FX/Niagara/Private/NiagaraSharedDataGPU_StaticMesh.ush"
+
+//////////////////////////////////////////////////////////////////////////
+// Distance Field Functions.
+
+void QueryDistanceField_{ParameterName}(bool bExecute, float3 WorldPosition, NiagaraSharedDataID_StaticMesh StaticMeshID, bool bUseMaxDistance, float MaxDistance, out bool OutIsValid, out float OutDistance)
+{
+	int DFIndex = ResolveNiagaraManagedReferenceID_StaticMesh(StaticMeshID);
+	if ( !bUseMaxDistance )
+	{
+		MaxDistance = 1e38;
+	}
+
+	//const int DFIndex = {ParameterName}_InstanceDistanceFieldIndex;
+	OutIsValid = bExecute && DFIndex != -1 && DFIndex < (int)NumSceneObjects;
+	OutDistance = MaxDistance;
+	if ( OutIsValid )
+	{
+		// Needs to be a smidge over 0.0 for DistanceToNearestSurfaceForObject to query the SDF
+		MaxDistance = max(0.001f, MaxDistance);
+		const FDFVector3 LWCWorldPosition = DFFromTileOffset_Hack(MakeLWCVector3(GetEngineOwnerLWCTile(), WorldPosition));
+		OutDistance = DistanceToNearestSurfaceForObject(DFIndex, LWCWorldPosition, MaxDistance);
+	}
+	else
+	{
+		OutIsValid = false;
+		OutDistance = MaxDistance;
+	}
+}
+
+void GetClosestPointOnMeshDistanceField_{ParameterName}(bool bExecute, float3 WorldPosition, NiagaraSharedDataID_StaticMesh StaticMeshID, bool bUseMaxDistance, float MaxDistance, out bool OutIsValid, out float OutClosestDistance, out float3 OutClosestPosition, out float3 OutClosestNormal, out float OutMaxEncodedDistance, out bool OutNormalIsValid)
+{
+	if ( !bUseMaxDistance )
+	{
+		MaxDistance = MAX_DISTANCE;
+	}
+
+	OutIsValid = false;
+	OutClosestDistance = MaxDistance;	
+	OutClosestNormal = float3(0,0,1);
+	OutClosestPosition = WorldPosition;
+	OutMaxEncodedDistance = MaxDistance;
+	OutNormalIsValid = false;
+
+	// Needs to be a smidge over 0.0 for DistanceToNearestSurfaceForObject to query the SDF
+	MaxDistance = max(0.001f, MaxDistance);
+	int DFIndex = ResolveNiagaraManagedReferenceID_StaticMesh(StaticMeshID);
+	OutIsValid = bExecute && DFIndex != -1 && DFIndex < (int)NumSceneObjects;
+	OutClosestDistance = MaxDistance;
+	if ( OutIsValid )
+	{
+		const FDFVector3 LWCWorldPosition = DFFromTileOffset_Hack(MakeLWCVector3(GetEngineOwnerLWCTile(), WorldPosition));
+		OutClosestDistance = GetDistanceToMeshDistanceField(DFIndex, LWCWorldPosition, MaxDistance);
+		ComputeClosestPointMeshDistanceField_LWCTiled(DFIndex, WorldPosition, OutClosestDistance, OutClosestPosition, OutClosestNormal, OutMaxEncodedDistance, OutNormalIsValid);				
+	}
+	else
+	{
+		OutIsValid = false;
+		OutClosestDistance = MaxDistance;
+		OutClosestPosition = WorldPosition;
+		OutClosestNormal = float3(0,0,1);
+		OutMaxEncodedDistance = MaxDistance;
+		OutNormalIsValid = false;
+	}
+}
\ No newline at end of file
diff --git a/Engine/Plugins/FX/Niagara/Shaders/Private/NiagaraSharedDataGPU_StaticMesh.ush b/Engine/Plugins/FX/Niagara/Shaders/Private/NiagaraSharedDataGPU_StaticMesh.ush
new file mode 100644
--- /dev/null
+++ b/Engine/Plugins/FX/Niagara/Shaders/Private/NiagaraSharedDataGPU_StaticMesh.ush
@@ -0,0 +1,87 @@
+// Copyright Epic Games, Inc. All Rights Reserved.
+
+/*=============================================================================
+NiagaraManagedReference_StaticMesh.ush 
+
+Utility/Support code for using NiagaraSharedDataID_StaticMesh values to indirectly reference static mesh data from the managed static mesh buffers.
+
+=============================================================================*/
+
+#pragma once
+
+/** Counterpart to struct in NiagaraManagedReference_StaticMesh.h*/
+struct FNiagaraSharedDataGPU_StaticMesh
+{
+	int IDTag;
+	int DistanceFieldIndex;
+	int PrimitiveIndex;
+	int Padding;
+};
+
+StructuredBuffer<FNiagaraSharedDataGPU_StaticMesh> NiagaraSharedData_StaticMesh_Buffer;
+int NiagaraSharedData_StaticMesh_NumSlots;
+
+/** Resolves a mesh DF Reference into a DFIndex that can be used to access distance field data. */ 
+int ResolveNiagaraManagedReferenceID_StaticMesh(NiagaraSharedDataID_StaticMesh StaticMeshID)
+{
+	int Slot = StaticMeshID.Index;
+	int Tag = StaticMeshID.AcquireTag;
+ 	
+	int DFIndex = -1;
+	if (Slot >= 0 && Slot < NiagaraSharedData_StaticMesh_NumSlots)
+	{
+		FNiagaraSharedDataGPU_StaticMesh StaticMeshInfo = NiagaraSharedData_StaticMesh_Buffer[Slot];
+		if (StaticMeshInfo.IDTag == Tag)
+		{
+			DFIndex = StaticMeshInfo.DistanceFieldIndex;
+		}
+	}
+	return DFIndex;
+}
+
+/** 
+Returns closest point and normal on a mesh distance field to a given world space point.
+Duplicate implementation to NiagaraDistanceFieldQueries.ush but this handles LWC internally and takes in a tiled point directly from a Niagara Sim and outputs the correctly tiled point back ready for the Niagara Sim.
+*/
+void ComputeClosestPointMeshDistanceField_LWCTiled(int DFIndex, const float3 WorldPosition, const float GridSampledDistance, out float3 OutClosestPosition, out float3 OutClosestNormal, out float MaxEncodedDistance, out bool bNormalIsValid)
+{
+	FDFObjectData DFObjectData = LoadDFObjectData(DFIndex);
+	
+	const FLWCVector3 TiledLWCPos = MakeLWCVector3(GetEngineOwnerLWCTile(), WorldPosition);
+	const FDFVector3 LWCWorldPosition = DFFromTileOffset_Hack(TiledLWCPos);
+	
+	const float3 VolumePosition = DFMultiplyDemote(LWCWorldPosition, DFObjectData.WorldToVolume);
+	const float3 ClampedVolumePosition = clamp(VolumePosition, -DFObjectData.VolumePositionExtent, DFObjectData.VolumePositionExtent);
+	
+	uint NumMips = LoadDFAssetData(DFObjectData.AssetIndex, 0).NumMips;
+	FDFAssetData DFAssetData = LoadDFAssetData(DFObjectData.AssetIndex, NumMips - 1);
+
+	MaxEncodedDistance = (DFAssetData.DistanceFieldToVolumeScaleBias.x + DFAssetData.DistanceFieldToVolumeScaleBias.y) * DFObjectData.VolumeScale;
+		
+	float3 VolumeGradient = CalculateMeshSDFGradient(ClampedVolumePosition, DFAssetData);
+	float VolumeGradientLength = length(VolumeGradient);
+	float3 VolumeNormal = VolumeGradientLength > 0.0f ? VolumeGradient / VolumeGradientLength : 0;
+
+	// Transform by transposed inverse to handle non-uniform scaling
+	float4x4 WorldToVolume = DFHackToFloat(DFObjectData.WorldToVolume);
+	float3 WorldGradient = mul(VolumeNormal, transpose((float3x3) WorldToVolume));
+	float WorldGradientLength = length(WorldGradient);
+
+	// @note: we should include checking if we are within the range of max encoded distance, but this breaks backwards compatibility
+	bNormalIsValid = WorldGradientLength > 0.0f && (MaxEncodedDistance - abs(GridSampledDistance) > 1e-8);
+			
+	// output
+	if (bNormalIsValid)
+	{
+		OutClosestNormal = WorldGradient / WorldGradientLength;
+		
+		const FDFVector3 WorldClampedPosition = DFMultiply(ClampedVolumePosition, DFObjectData.VolumeToWorld);
+		const float3 TiledWorldClampedPosition = DFSubtractDemote(WorldClampedPosition, LWCGetTileOffset(TiledLWCPos));
+		OutClosestPosition = TiledWorldClampedPosition - OutClosestNormal * GridSampledDistance;
+	}
+	else
+	{
+		OutClosestPosition = WorldPosition;
+		OutClosestNormal = float3(0, 0, 1);
+	}
+}
\ No newline at end of file
diff --git a/Engine/Plugins/FX/Niagara/Source/Niagara/Internal/DataInterface/NiagaraDataInterfaceStaticMeshIndirect.h b/Engine/Plugins/FX/Niagara/Source/Niagara/Internal/DataInterface/NiagaraDataInterfaceStaticMeshIndirect.h
new file mode 100644
--- /dev/null
+++ b/Engine/Plugins/FX/Niagara/Source/Niagara/Internal/DataInterface/NiagaraDataInterfaceStaticMeshIndirect.h
@@ -0,0 +1,82 @@
+// Copyright Epic Games, Inc. All Rights Reserved.
+#pragma once
+
+#include "NiagaraDataInterface.h"
+#include "NiagaraCommon.h"
+#include "NiagaraDataInterfaceStaticMeshIndirect.generated.h"
+
+class UNiagaraSystem;
+
+USTRUCT()
+struct FNDIStaticMeshIndirectCompiledData
+{
+	GENERATED_BODY()
+	
+	/** True if this DI is used by a CPU emitter. */
+	UPROPERTY()
+	uint32 bUsedByCPU : 1 = false;
+	
+	/** True if this DI is used by a GPU emitter. */
+	UPROPERTY()
+	uint32 bUsedByGPU : 1 = false;
+	
+	/** True if this DI requires that it's referenced mesh has CPU Access Enabled. */
+	UPROPERTY()
+	uint32 bNeedsCpuAccess : 1 = false;
+
+	/** True if this DI requires static mesh distance field data. */
+	UPROPERTY()
+	uint32 bRequiresDistanceFields : 1 = false;
+	
+	void Init(const UNiagaraSystem* System, UNiagaraDataInterface* OwnerDI);
+};
+
+/** 
+Data Interface allowing limited access of static mesh info indirectly via an FNiagaraSharedDataID_StaticMesh. 
+Available data and features are far more limited than a regular Static Mesh DI but its indirect nature allows it to be used without direct binding to a single static mesh.
+*/
+UCLASS(EditInlineNew, Category = "Meshes", meta = (DisplayName = "Static Mesh Reader Indirect"), MinimalAPI)
+class UNiagaraDataInterfaceStaticMeshIndirect : public UNiagaraDataInterface
+{
+	GENERATED_UCLASS_BODY()
+
+public:
+
+	//~ UObject interface
+	NIAGARA_API virtual void PostInitProperties() override;
+	//~ UObject interface
+	
+	//~ UNiagaraDataInterface interface
+#if WITH_EDITORONLY_DATA
+	NIAGARA_API virtual void GetCommonHLSL(FString& OutHLSL) override;
+	NIAGARA_API virtual void PostCompile(const UNiagaraSystem& OwningSystem) override;
+#endif
+	virtual bool CanExecuteOnTarget(ENiagaraSimTarget Target) const override { return true; }
+	virtual bool RequiresEarlyViewData() const override;
+
+#if WITH_EDITORONLY_DATA
+	NIAGARA_API virtual bool AppendCompileHash(FNiagaraCompileHashVisitor* InVisitor) const override;
+	NIAGARA_API virtual void ModifyCompilationEnvironment(EShaderPlatform ShaderPlatform, struct FShaderCompilerEnvironment& OutEnvironment) const override;
+	NIAGARA_API virtual void GetParameterDefinitionHLSL(const FNiagaraDataInterfaceGPUParamInfo& ParamInfo, FString& OutHLSL) override;	
+	NIAGARA_API virtual bool GetFunctionHLSL(const FNiagaraDataInterfaceHlslGenerationContext& HlslGenContext, FString& OutHLSL) override;
+#endif
+	NIAGARA_API virtual void BuildShaderParameters(FNiagaraShaderParametersBuilder& ShaderParametersBuilder) const override;
+	NIAGARA_API virtual void SetShaderParameters(const FNiagaraDataInterfaceSetShaderParametersContext& Context) const override;
+
+	NIAGARA_API virtual bool Equals(const UNiagaraDataInterface* Other) const override;
+protected:
+#if WITH_EDITORONLY_DATA
+	NIAGARA_API virtual void GetFunctionsInternal(TArray<FNiagaraFunctionSignature>& OutFunctions) const override;
+	NIAGARA_API void GetDistanceFieldFunctions(TArray<FNiagaraFunctionSignature>& OutFunctions, const FNiagaraFunctionSignature& BaseSignature) const;
+#endif
+	NIAGARA_API virtual bool CopyToInternal(UNiagaraDataInterface* Destination) const override;
+public:
+	//~ UNiagaraDataInterface interface
+
+	const FNDIStaticMeshIndirectCompiledData& GetCompiledData()const { return CompiledData; }
+	
+private:
+
+ 	UPROPERTY()
+	FNDIStaticMeshIndirectCompiledData CompiledData;
+};
diff --git a/Engine/Plugins/FX/Niagara/Source/Niagara/Private/DataInterface/NiagaraDataInterfaceStaticMeshIndirect.cpp b/Engine/Plugins/FX/Niagara/Source/Niagara/Private/DataInterface/NiagaraDataInterfaceStaticMeshIndirect.cpp
new file mode 100644
--- /dev/null
+++ b/Engine/Plugins/FX/Niagara/Source/Niagara/Private/DataInterface/NiagaraDataInterfaceStaticMeshIndirect.cpp
@@ -0,0 +1,284 @@
+// Copyright Epic Games, Inc. All Rights Reserved.
+
+#include "DataInterface/NiagaraDataInterfaceStaticMeshIndirect.h"
+
+#include "NiagaraWorldManagerImpl.h"
+
+#include "NiagaraCompileHashVisitor.h"
+#include "NiagaraShaderParametersBuilder.h"
+#include "FXRenderingUtils.h"
+#include "NiagaraGpuComputeDispatchInterface.h"
+
+#include "NiagaraDataManager_StaticMesh.h"
+
+#include "NiagaraDataInterfaceUtilities.h"
+
+#include UE_INLINE_GENERATED_CPP_BY_NAME(NiagaraDataInterfaceStaticMeshIndirect)
+#define LOCTEXT_NAMESPACE "NiagaraDataInterfaceStaticMeshIndirect"
+
+namespace NDIStaticMeshIndirectLocal
+{
+	BEGIN_SHADER_PARAMETER_STRUCT(FShaderParameters,)
+	END_SHADER_PARAMETER_STRUCT()
+
+	static const TCHAR* TemplateShaderFile = TEXT("/Plugin/FX/Niagara/Private/NiagaraDataInterfaceStaticMeshIndirectTemplate.ush");
+
+	//////////////////////////////////////////////////////////////////////////
+	struct EDIFunctionVersion
+	{
+		enum Type
+		{
+			InitialVersion = 0,
+
+			VersionPlusOne,
+			LatestVersion = VersionPlusOne - 1
+		};
+	};
+
+	//////////////////////////////////////////////////////////////////////////
+	// Distance Field Functions
+	static const FName	QueryDistanceFieldName("QueryDistanceField");
+	static const FName	GetClosestPointOnMeshDistanceFieldName("GetClosestPointOnMeshDistanceField");
+
+	struct FRenderProxy : public FNiagaraDataInterfaceProxy
+	{
+		virtual void ConsumePerInstanceDataFromGameThread(void* FromGameThreadData, const FNiagaraSystemInstanceID& InstanceID) override{}
+		virtual int32 PerInstanceDataPassedToRenderThreadSize() const override{ return 0; }
+	};
+};
+
+//////////////////////////////////////////////////////////////////////////
+
+void FNDIStaticMeshIndirectCompiledData::Init(const UNiagaraSystem* System, UNiagaraDataInterface* OwnerDI)
+{
+	bUsedByCPU = false;
+	bUsedByGPU = false;
+	bNeedsCpuAccess = false;	
+	bRequiresDistanceFields = false;
+
+	//For every GPU script we iterate over the functions it calls and add each of them to the mapping.
+	//This will then be placed in a buffer for the RT to pass to the GPU so that each script can look up the correct function layout info.
+	auto HandleGpuFunc = [&](const UNiagaraScript* Script, const FNiagaraDataInterfaceGeneratedFunction& BindingInfo)
+	{
+		bUsedByGPU = true;
+
+		bool bIsDistanceFieldFunction = BindingInfo.DefinitionName == NDIStaticMeshIndirectLocal::QueryDistanceFieldName ||
+										BindingInfo.DefinitionName == NDIStaticMeshIndirectLocal::GetClosestPointOnMeshDistanceFieldName ||
+										BindingInfo.DefinitionName == NDIStaticMeshIndirectLocal::QueryDistanceFieldName;
+		bRequiresDistanceFields |= bIsDistanceFieldFunction;
+		 
+		return true;
+	};
+	FNiagaraDataInterfaceUtilities::ForEachGpuFunction(OwnerDI, System, HandleGpuFunc);
+}
+
+//////////////////////////////////////////////////////////////////////////
+
+UNiagaraDataInterfaceStaticMeshIndirect::UNiagaraDataInterfaceStaticMeshIndirect(FObjectInitializer const& ObjectInitializer)
+	: Super(ObjectInitializer)
+{
+	Proxy.Reset(new NDIStaticMeshIndirectLocal::FRenderProxy());
+}
+
+void UNiagaraDataInterfaceStaticMeshIndirect::PostInitProperties()
+{
+	Super::PostInitProperties();
+
+	// Can we register data interfaces as regular types and fold them into the FNiagaraVariable framework for UI and function calls etc?
+	if (HasAnyFlags(RF_ClassDefaultObject))
+	{
+		ENiagaraTypeRegistryFlags DIFlags =
+			ENiagaraTypeRegistryFlags::AllowAnyVariable |
+			ENiagaraTypeRegistryFlags::AllowParameter;
+		FNiagaraTypeRegistry::Register(FNiagaraTypeDefinition(GetClass()), DIFlags);
+	}
+}
+
+bool UNiagaraDataInterfaceStaticMeshIndirect::Equals(const UNiagaraDataInterface* Other) const
+{
+	if (!Super::Equals(Other))
+	{
+		return false;
+	}
+	const UNiagaraDataInterfaceStaticMeshIndirect* OtherTyped = CastChecked<const UNiagaraDataInterfaceStaticMeshIndirect>(Other);
+	return true;
+}
+
+bool UNiagaraDataInterfaceStaticMeshIndirect::CopyToInternal(UNiagaraDataInterface* Destination) const
+{
+	if (!Super::CopyToInternal(Destination))
+	{
+		return false;
+	}
+
+	UNiagaraDataInterfaceStaticMeshIndirect* OtherTyped = CastChecked<UNiagaraDataInterfaceStaticMeshIndirect>(Destination);
+	return true;
+}
+
+#if WITH_EDITORONLY_DATA
+void UNiagaraDataInterfaceStaticMeshIndirect::GetFunctionsInternal(TArray<FNiagaraFunctionSignature>& OutFunctions) const
+{
+	using namespace NDIStaticMeshIndirectLocal;
+
+	//OutFunctions.Reserve(OutFunctions.Num() + NumFunctions);
+
+	// Setup base signature
+	FNiagaraFunctionSignature BaseSignature;
+	BaseSignature.Inputs.Emplace(FNiagaraTypeDefinition(GetClass()), TEXT("StaticMeshIndirect"));
+	BaseSignature.bMemberFunction = true;
+	BaseSignature.bRequiresContext = false;
+	BaseSignature.FunctionVersion = EDIFunctionVersion::LatestVersion;
+
+	GetDistanceFieldFunctions(OutFunctions, BaseSignature);
+}
+
+void UNiagaraDataInterfaceStaticMeshIndirect::GetDistanceFieldFunctions(TArray<FNiagaraFunctionSignature>&OutFunctions, const FNiagaraFunctionSignature & BaseSignature) const
+{
+	using namespace NDIStaticMeshIndirectLocal;
+	
+	const FText DFIsValidText = LOCTEXT("DFIsValidTooltip","IsValid Returns true if the mesh ID is valid and the mesh has a valid Distance Field to sample.");
+	const FText DFNormalIsValidText = LOCTEXT("OutNormalIsValidTooltip","OutNormalIsValid Returns true the distance field was successfully sampled and the position and normal are valid.");
+	{
+		const FNiagaraVariable UseMaxDistanceVar(FNiagaraTypeDefinition::GetBoolDef(), TEXT("Use Max Distance"));
+		const FNiagaraVariable MaxDistanceVar(FNiagaraTypeDefinition::GetFloatDef(), TEXT("Max Distance"));
+
+		FNiagaraFunctionSignature& Sig = OutFunctions.Add_GetRef(BaseSignature);
+		Sig.Name = QueryDistanceFieldName;
+		Sig.bSupportsCPU = false;
+		Sig.Inputs.Emplace_GetRef(FNiagaraTypeDefinition::GetBoolDef(), TEXT("Execute")).SetValue(true);
+		Sig.Inputs.Emplace(FNiagaraTypeDefinition::GetPositionDef(), TEXT("World Position"));
+		Sig.Inputs.Emplace(FNiagaraTypeDefinition(FNiagaraSharedDataID_StaticMesh::StaticStruct()), TEXT("StaticMeshID"));
+		Sig.Inputs.Emplace(UseMaxDistanceVar);
+		Sig.Inputs.Emplace(MaxDistanceVar);
+		Sig.AddOutput(FNiagaraVariableBase(FNiagaraTypeDefinition::GetBoolDef(), TEXT("IsValid")), DFIsValidText);
+		Sig.Outputs.Emplace(FNiagaraTypeDefinition::GetFloatDef(), TEXT("Distance"));
+		
+#if WITH_EDITORONLY_DATA
+		Sig.Description = LOCTEXT("QueryDistanceFieldDescription", "Given a world position, this returns the value of the parented static mesh`s signed distance field");
+		Sig.InputDescriptions.Add(
+			UseMaxDistanceVar,
+			LOCTEXT("UseMaxDistanceOptimizationDescription", 
+			"This enables an optimization that will skip reading the SDF texture if the world position exceeds a provided max distance value.\n"
+			"This is a useful optimization if you only need accurate distance information within a certain distance threshold. \n"
+			"One example of a good usecase is if you only want to spawn particles inside the mesh (i.e. distance < 0.0), you should enable this with a MaxDistance of 0.0.")
+		);
+		Sig.InputDescriptions.Add(
+			MaxDistanceVar,
+			LOCTEXT("MaxDistanceDescription",
+			"Only used if UseMaxDistanceOptimization is enabled. This is the max distance from the static meshes` BOUNDING BOX at which we should query it's SDF.\n"
+			"This defaults to 0.0, which means it will only read the SDF if the world position is inside the meshes' bounding box")
+		);
+#endif
+	}
+
+	{
+		const FNiagaraVariable UseMaxDistanceVar(FNiagaraTypeDefinition::GetBoolDef(), TEXT("Use Max Distance"));
+		const FNiagaraVariable MaxDistanceVar(FNiagaraTypeDefinition::GetFloatDef(), TEXT("Max Distance"));
+
+		FNiagaraFunctionSignature& Sig = OutFunctions.Add_GetRef(BaseSignature);
+		Sig.Name = GetClosestPointOnMeshDistanceFieldName;
+		Sig.bSupportsCPU = false;
+		Sig.Inputs.Emplace_GetRef(FNiagaraTypeDefinition::GetBoolDef(), TEXT("Execute")).SetValue(true);
+		Sig.Inputs.Emplace(FNiagaraTypeDefinition::GetPositionDef(), TEXT("World Position"));
+		Sig.Inputs.Emplace(FNiagaraTypeDefinition(FNiagaraSharedDataID_StaticMesh::StaticStruct()), TEXT("StaticMeshID"));
+		Sig.Inputs.Emplace(UseMaxDistanceVar);
+		Sig.Inputs.Emplace(MaxDistanceVar);
+		Sig.AddOutput(FNiagaraVariableBase(FNiagaraTypeDefinition::GetBoolDef(), TEXT("IsValid")), DFIsValidText);
+		Sig.Outputs.Emplace(FNiagaraTypeDefinition::GetFloatDef(), TEXT("OutClosestDistance"));
+		Sig.Outputs.Emplace(FNiagaraTypeDefinition::GetPositionDef(), TEXT("OutClosestPosition"));
+		Sig.Outputs.Emplace(FNiagaraTypeDefinition::GetVec3Def(), TEXT("OutClosestNormal"));
+		Sig.Outputs.Emplace(FNiagaraTypeDefinition::GetFloatDef(), TEXT("OutMaxEncodedDistance"));
+		Sig.Outputs.Emplace(FNiagaraTypeDefinition::GetBoolDef(), TEXT("OutNormalIsValid"));
+
+#if WITH_EDITORONLY_DATA
+		Sig.Description = LOCTEXT("GetClosestPointOnMeshDistanceFieldDescription", "Given a world position, this returns the closest point on the mesh with it's distance and normal.");
+#endif
+	}
+}
+
+#endif //WITH_EDITORONLY_DATA
+
+bool UNiagaraDataInterfaceStaticMeshIndirect::RequiresEarlyViewData() const
+{
+	return GetCompiledData().bRequiresDistanceFields;
+}
+
+#if WITH_EDITORONLY_DATA
+
+void UNiagaraDataInterfaceStaticMeshIndirect::PostCompile(const UNiagaraSystem& OwningSystem)
+{
+	CompiledData.Init(&OwningSystem, this);
+}
+
+void UNiagaraDataInterfaceStaticMeshIndirect::GetCommonHLSL(FString& OutHLSL)
+{
+	OutHLSL += TEXT("#include \"/Engine/Private/DistanceFieldLightingShared.ush\"\n");
+	OutHLSL += TEXT("#include \"/Engine/Private/MeshDistanceFieldCommon.ush\"\n");
+}
+
+bool UNiagaraDataInterfaceStaticMeshIndirect::AppendCompileHash(FNiagaraCompileHashVisitor* InVisitor) const
+{
+	bool bSuccess = Super::AppendCompileHash(InVisitor);
+	bSuccess &= InVisitor->UpdateShaderFile(NDIStaticMeshIndirectLocal::TemplateShaderFile);
+	bSuccess &= InVisitor->UpdateShaderParameters<NDIStaticMeshIndirectLocal::FShaderParameters>();
+	bSuccess &= InVisitor->UpdateShaderParameters<FNiagaraDataManager_StaticMesh_Parameters>();
+	return bSuccess;
+}
+
+void UNiagaraDataInterfaceStaticMeshIndirect::ModifyCompilationEnvironment(EShaderPlatform ShaderPlatform, struct FShaderCompilerEnvironment& OutEnvironment) const
+{
+	Super::ModifyCompilationEnvironment(ShaderPlatform, OutEnvironment);
+}
+
+void UNiagaraDataInterfaceStaticMeshIndirect::GetParameterDefinitionHLSL(const FNiagaraDataInterfaceGPUParamInfo& ParamInfo, FString& OutHLSL)
+{
+	const TMap<FString, FStringFormatArg> TemplateArgs =
+	{
+		{TEXT("ParameterName"),	ParamInfo.DataInterfaceHLSLSymbol},
+	};
+	AppendTemplateHLSL(OutHLSL, NDIStaticMeshIndirectLocal::TemplateShaderFile, TemplateArgs);
+}
+
+bool UNiagaraDataInterfaceStaticMeshIndirect::GetFunctionHLSL(const FNiagaraDataInterfaceHlslGenerationContext& HlslGenContext, FString& OutHLSL)
+{
+	return	HlslGenContext.GetFunctionInfo().DefinitionName == NDIStaticMeshIndirectLocal::GetClosestPointOnMeshDistanceFieldName ||
+		HlslGenContext.GetFunctionInfo().DefinitionName == NDIStaticMeshIndirectLocal::QueryDistanceFieldName;
+}
+#endif
+
+void UNiagaraDataInterfaceStaticMeshIndirect::BuildShaderParameters(FNiagaraShaderParametersBuilder& ShaderParametersBuilder) const
+{
+	ShaderParametersBuilder.AddNestedStruct<NDIStaticMeshIndirectLocal::FShaderParameters>();
+	ShaderParametersBuilder.AddIncludedStruct(UE::FXRenderingUtils::DistanceFields::GetObjectBufferParametersMetadata());
+	ShaderParametersBuilder.AddIncludedStruct(UE::FXRenderingUtils::DistanceFields::GetAtlasParametersMetadata());
+	ShaderParametersBuilder.AddIncludedStruct(FNiagaraDataManager_StaticMesh_RTProxy::GetParametersMetadata());
+}
+
+void UNiagaraDataInterfaceStaticMeshIndirect::SetShaderParameters(const FNiagaraDataInterfaceSetShaderParametersContext& Context) const
+{
+	NDIStaticMeshIndirectLocal::FRenderProxy& DIProxy = Context.GetProxy<NDIStaticMeshIndirectLocal::FRenderProxy>();
+	NDIStaticMeshIndirectLocal::FShaderParameters* ShaderParameters = Context.GetParameterNestedStruct<NDIStaticMeshIndirectLocal::FShaderParameters>();
+
+	// Bind Mesh Distance Field Data
+	const FShaderParametersMetadata* ObjectBufferParametersMetadata = UE::FXRenderingUtils::DistanceFields::GetObjectBufferParametersMetadata();
+	const FShaderParametersMetadata* AtlasParametersMetadata = UE::FXRenderingUtils::DistanceFields::GetAtlasParametersMetadata();
+
+	uint8* ShaderDistanceFieldObjectParameters = Context.GetParameterIncludedStruct(ObjectBufferParametersMetadata);
+	uint8* ShaderDistanceFieldAtlasParameters = Context.GetParameterIncludedStruct(AtlasParametersMetadata);
+
+	const bool bDistanceFieldDataBound =
+		Context.IsStructBound(ShaderDistanceFieldObjectParameters, ObjectBufferParametersMetadata) ||
+		Context.IsStructBound(ShaderDistanceFieldAtlasParameters, AtlasParametersMetadata);
+
+	if (bDistanceFieldDataBound)
+	{
+		TConstStridedView<FSceneView> SimulationSceneViews = Context.GetComputeDispatchInterface().GetSimulationSceneViews();
+		const FSceneView* PrimaryView = SimulationSceneViews.Num() > 0 ? &SimulationSceneViews[0] : nullptr;
+		UE::FXRenderingUtils::DistanceFields::SetupObjectBufferParameters(Context.GetGraphBuilder(), ShaderDistanceFieldObjectParameters, PrimaryView);
+		UE::FXRenderingUtils::DistanceFields::SetupAtlasParameters(Context.GetGraphBuilder(), ShaderDistanceFieldAtlasParameters, PrimaryView);
+	}
+
+	FNiagaraDataManager_StaticMesh_RTProxy::SetShaderParameters(Context);	
+}
+
+#undef LOCTEXT_NAMESPACE
diff --git a/Engine/Plugins/FX/Niagara/Source/Niagara/Private/NiagaraDataChannelAccessor.cpp b/Engine/Plugins/FX/Niagara/Source/Niagara/Private/NiagaraDataChannelAccessor.cpp
--- a/Engine/Plugins/FX/Niagara/Source/Niagara/Private/NiagaraDataChannelAccessor.cpp
+++ b/Engine/Plugins/FX/Niagara/Source/Niagara/Private/NiagaraDataChannelAccessor.cpp
@@ -342,6 +342,9 @@ void FNDCVarAccessorBase::Init(const UNiagaraDataChannel* DataChannel)
 void FNDCAccessorBase::Init(const UNiagaraDataChannel* DataChannel)
 {
 	check(DataChannel);
+	
+	CachedLayout = DataChannel->GetLayoutInfo();
+
 	TArray<FNiagaraVariableBase> MissingVars;
 	for (FNDCVarAccessorBase* Var : VariableAccessors)
 	{
diff --git a/Engine/Plugins/FX/Niagara/Source/Niagara/Private/NiagaraDataChannelHandler.cpp b/Engine/Plugins/FX/Niagara/Source/Niagara/Private/NiagaraDataChannelHandler.cpp
--- a/Engine/Plugins/FX/Niagara/Source/Niagara/Private/NiagaraDataChannelHandler.cpp
+++ b/Engine/Plugins/FX/Niagara/Source/Niagara/Private/NiagaraDataChannelHandler.cpp
@@ -27,6 +27,8 @@ void UNiagaraDataChannelHandler::Init(const UNiagaraDataChannel* InChannel)
 	check(InChannel);
 	DataChannel = InChannel;
 	AccessContextType = InChannel->GetAccessContextType();
+
+	CachedLayoutInfo = DataChannel->GetLayoutInfo();
 }
 
 void UNiagaraDataChannelHandler::Cleanup()
@@ -73,6 +75,18 @@ void UNiagaraDataChannelHandler::Tick(float DeltaTime, ETickingGroup TickGroup)
 	CurrentTG = TickGroup;
 }
 
+void UNiagaraDataChannelHandler::OnDataChannelLayoutChanged()
+{
+	if(const UNiagaraDataChannel* NDC = DataChannel.Get())
+	{
+		CachedLayoutInfo = NDC->GetLayoutInfo();
+	}
+	else
+	{
+		CachedLayoutInfo.Reset();
+	}
+}
+
 UNiagaraDataChannelWriter* UNiagaraDataChannelHandler::GetDataChannelWriter()
 {
 	if(Writer == nullptr)
diff --git a/Engine/Plugins/FX/Niagara/Source/Niagara/Private/NiagaraDataChannelLayoutInfo.cpp b/Engine/Plugins/FX/Niagara/Source/Niagara/Private/NiagaraDataChannelLayoutInfo.cpp
--- a/Engine/Plugins/FX/Niagara/Source/Niagara/Private/NiagaraDataChannelLayoutInfo.cpp
+++ b/Engine/Plugins/FX/Niagara/Source/Niagara/Private/NiagaraDataChannelLayoutInfo.cpp
@@ -8,7 +8,7 @@
 void FNiagaraDataChannelGameDataLayout::Init(TConstArrayView<FNiagaraDataChannelVariable> Variables)
 {
 	VariableIndices.Reset();
-	LwcConverters.Reserve(Variables.Num());
+	LwcConverters.Reset(Variables.Num());
 	for (const FNiagaraDataChannelVariable& Var : Variables)
 	{
 		//Sigh.
diff --git a/Engine/Plugins/FX/Niagara/Source/Niagara/Private/NiagaraDataChannelManager.cpp b/Engine/Plugins/FX/Niagara/Source/Niagara/Private/NiagaraDataChannelManager.cpp
--- a/Engine/Plugins/FX/Niagara/Source/Niagara/Private/NiagaraDataChannelManager.cpp
+++ b/Engine/Plugins/FX/Niagara/Source/Niagara/Private/NiagaraDataChannelManager.cpp
@@ -38,29 +38,38 @@ void FNiagaraDataChannelManager::RefreshDataChannels()
 		return;
 	}
 
-	UNiagaraDataChannel::ForEachDataChannel([&](UNiagaraDataChannel* DataChannel)
+	if(GNDCAllowLazyHandlerInit == 0)
 	{
-		InitDataChannel(DataChannel, false);
-	});
+		UNiagaraDataChannel::ForEachDataChannel([&](UNiagaraDataChannel* DataChannel)
+		{
+			InitDataChannel(DataChannel, false);
+		});
+	}
 }
 
 void FNiagaraDataChannelManager::Init()
 {
 	bIsInitialized = true;
 	bIsCleanedUp = false;
-
-	//Initialize any existing data channels, more may be initialized later as they are loaded.
-	UNiagaraDataChannel::ForEachDataChannel([&](UNiagaraDataChannel* DataChannel)
+	
+	if(GNDCAllowLazyHandlerInit == 0)
 	{
-		InitDataChannel(DataChannel, true);
-	});
+		//Initialize any existing data channels, more may be initialized later as they are loaded.
+		UNiagaraDataChannel::ForEachDataChannel([&](UNiagaraDataChannel* DataChannel)
+		{
+			InitDataChannel(DataChannel, true);
+		});
+	}
 }
 
 void FNiagaraDataChannelManager::Cleanup()
 {
 	for (auto& ChannelPair : Channels)
 	{
-		ChannelPair.Value->Cleanup();
+		if(UNiagaraDataChannelHandler* ChannelHandler = ChannelPair.Value.Get())
+		{
+			ChannelHandler->Cleanup();
+		}
 	}
 	Channels.Empty();
 	bIsCleanedUp = true;
diff --git a/Engine/Plugins/FX/Niagara/Source/Niagara/Private/NiagaraDataChannel_Map.cpp b/Engine/Plugins/FX/Niagara/Source/Niagara/Private/NiagaraDataChannel_Map.cpp
--- a/Engine/Plugins/FX/Niagara/Source/Niagara/Private/NiagaraDataChannel_Map.cpp
+++ b/Engine/Plugins/FX/Niagara/Source/Niagara/Private/NiagaraDataChannel_Map.cpp
@@ -71,7 +71,7 @@ void UNiagaraDataChannelHandler_MapBase::BeginFrame(float DeltaTime)
 	for(auto It = ActiveEntries.CreateIterator(); It; ++It)
 	{
 		FNDCMapEntry& Entry = EntryPool[It.Value()];
-		 if (Entry.Get().BeginFrame(DeltaTime, OwningWorld, It.Key()) == false)
+		if (Entry.Get().BeginFrame(DeltaTime, OwningWorld, It.Key()) == false)
  		{
  			//This entry is done with so reset it and return to the free entries.
  
@@ -146,7 +146,12 @@ void UNiagaraDataChannelHandler_MapBase::Tick(float DeltaTime, ETickingGroup Tic
 FNiagaraDataChannelDataPtr UNiagaraDataChannelHandler_MapBase::FindData(FNDCAccessContextInst& AccessContext, ENiagaraResourceAccess AccessType)
 {
 	FNDCMapEntry& Entry = FindOrAddEntry(AccessContext);
-	return Entry.Get().GetData();
+
+	if(Entry.IsValid())
+	{
+		return Entry.Get().GetData();
+	}
+	return nullptr;
 }
 
 UObject* UNiagaraDataChannelHandler_MapBase::GetSystemToSpawn(FNDCAccessContextInst& AccessContext)const
@@ -177,79 +182,102 @@ FNDCMapEntry& UNiagaraDataChannelHandler_MapBase::FindOrAddEntry(FNDCAccessConte
 	//Clear out any transient output data from our context.
 	//TODO; Should possibly use virtual or iterate over properties to clear out all transient output properties here.
 	TypedContext.SpawnedSystems.Reset();
-
+	FNDCMapEntry* Ret = nullptr;
+	int32 NumSpawnedSystems = 0;
 	if (TypedContext.bIsAutoLinkingSystem)
 	{
 		//Only Niagara Components spawned by this NDC should be setting this flag. It is purely to allow auto-linking back to the originating NDC data without the need for generating a key.
-		if(UNiagaraComponent* OwnerComp = Cast<UNiagaraComponent>(TypedContext.GetOwner()))
+		if (UNiagaraComponent* OwnerComp = Cast<UNiagaraComponent>(TypedContext.GetOwner()))
 		{
 			if (CurrentlyInitializingEntry)
 			{
 				if (CurrentlyInitializingEntry->Get().SpawnedComponents.Contains(OwnerComp))
 				{
 					//The currently initializing entry spawned this component so just return this entry.
-					return *CurrentlyInitializingEntry;
+					Ret = CurrentlyInitializingEntry;
 				}
 			}
 			else
 			{
 				//Even if we're not currently activating an entry, this might be a handler system.
 				//Use the reverse mapping to find it without needing to do a full key gen.
-				if(int32* FoundEntryIdx = SpawnedComponentsToActiveEntry.Find(OwnerComp))
+				if (int32* FoundEntryIdx = SpawnedComponentsToActiveEntry.Find(OwnerComp))
 				{
 					check(EntryPool[*FoundEntryIdx].Get().SpawnedComponents.Contains(OwnerComp));
-					return EntryPool[*FoundEntryIdx];
+					Ret = &EntryPool[*FoundEntryIdx];
 				}
 			}
 		}
 	}
-
-	//Generate the key for this Access Context. Each child class will implement this as needed.
-	FNDCMapKey Key;
-	FNDCMapKeyWriter KeyWriter(Key);
-	GenerateKey(AccessContext, KeyWriter);
-
-	int32* ActiveEntry = ActiveEntries.Find(Key);
-	if(ActiveEntry)
-	{
-		return EntryPool[*ActiveEntry];
-	}
 	else
 	{
-		int32 NewEntryIdx = INDEX_NONE;
-		if (FreeEntries.Num() > 0)
-		{
-			NewEntryIdx = FreeEntries.Pop();
+		//Generate the key for this Access Context. Each child class will implement this as needed.
+		FNDCMapKey Key;
+		FNDCMapKeyWriter KeyWriter(Key);
+		GenerateKey(AccessContext, KeyWriter);
+
+		int32* ActiveEntry = ActiveEntries.Find(Key);
+		if (ActiveEntry)
+		{					
+			Ret = &EntryPool[*ActiveEntry];
 		}
-		else
+		else if (!TypedContext.bIsAutoLinkingSystem)//Don't create new entires for auto linking niagra systems.
 		{
-			NewEntryIdx = EntryPool.Num();
-			EntryPool.Emplace(GetMapEntryType());
-		}
+			int32 NewEntryIdx = INDEX_NONE;
+			if (FreeEntries.Num() > 0)
+			{
+				NewEntryIdx = FreeEntries.Pop();
+			}
+			else
+			{
+				NewEntryIdx = EntryPool.Num();
+				EntryPool.Emplace(GetMapEntryType(), NewEntryIdx);
+			}
 
-		check(NewEntryIdx != INDEX_NONE);
-		FNDCMapEntry* RetEntry = &EntryPool[NewEntryIdx];
-		check(RetEntry);
-		ActiveEntries.Add(Key) = NewEntryIdx;
+			check(NewEntryIdx != INDEX_NONE);
+			FNDCMapEntry* RetEntry = &EntryPool[NewEntryIdx];
+			check(RetEntry);
+			ActiveEntries.Add(Key) = NewEntryIdx;
 
-		//Init the new entry.
-		//Cache off the currently initializing entry so we can retreive it easily should the entry spawn any systems that re-enter FindData() and need to quickly find their spawning NDC data.
-		CurrentlyInitializingEntry = RetEntry;
+			//Init the new entry.
+			//Cache off the currently initializing entry so we can retreive it easily should the entry spawn any systems that re-enter FindData() and need to quickly find their spawning NDC data.
+			CurrentlyInitializingEntry = RetEntry;
 
-		RetEntry->Get().Init(AccessContext, this, Key);
+			RetEntry->Get().Init(AccessContext, this, Key);
 
-		CurrentlyInitializingEntry = nullptr;
+			NumSpawnedSystems = TypedContext.SpawnedSystems.Num();
+			CurrentlyInitializingEntry = nullptr;
 
-		//Add any spawned components to the reverse lookup map.
-		for (const FNDCSpawnedSystemRef& WeakSpawned : TypedContext.SpawnedSystems)
-		{
-			if (UNiagaraComponent* Spawned = WeakSpawned.Get())
+			//Add any spawned components to the reverse lookup map.
+			for (const FNDCSpawnedSystemRef& WeakSpawned : TypedContext.SpawnedSystems)
 			{
-				SpawnedComponentsToActiveEntry.Add(Spawned) = NewEntryIdx;
+				if (UNiagaraComponent* Spawned = WeakSpawned.Get())
+				{
+					SpawnedComponentsToActiveEntry.Add(Spawned) = NewEntryIdx;
+				}
 			}
+			Ret = RetEntry;
 		}
-		return *RetEntry;
 	}
+	
+	if(Ret)
+	{
+		//Return any existing systems if needed.
+		FNDCMapEntryBase& TypedEntry = Ret->Get<FNDCMapEntryBase>();
+		if(TypedContext.bReturnExistingSystems)
+		{
+			//Any newly spawned systems are added in entry init. If we need all systems, add the rest now.
+			const TArray<TObjectPtr<UNiagaraComponent>>& EntrySystems = TypedEntry.GetSpawnedSystems();
+			for(int32 i=0; i < EntrySystems.Num() - NumSpawnedSystems; ++i)
+			{
+				TypedContext.SpawnedSystems.Add(EntrySystems[i].Get());
+			}			
+		}
+		return *Ret;
+	}
+
+	static FNDCMapEntry Dummy;
+	return Dummy;
 }
 
 //////////////////////////////////////////////////////////////////////////
diff --git a/Engine/Plugins/FX/Niagara/Source/Niagara/Private/NiagaraDataManager.cpp b/Engine/Plugins/FX/Niagara/Source/Niagara/Private/NiagaraDataManager.cpp
new file mode 100644
--- /dev/null
+++ b/Engine/Plugins/FX/Niagara/Source/Niagara/Private/NiagaraDataManager.cpp
@@ -0,0 +1,68 @@
+// Copyright Epic Games, Inc. All Rights Reserved.
+
+#include "NiagaraDataManager.h"
+
+#include "NiagaraWorldManager.h"
+#include "NiagaraGpuComputeDispatchInterface.h"
+#include "NiagaraComponent.h"
+
+
+void UNiagaraDataManager::Init(FNiagaraWorldManager* InWorldManager)
+{
+	check(WorldManager == nullptr); //Called initialize multiple times.
+	check(InWorldManager);
+	WorldManager = InWorldManager;
+	ComputeDispatchInterface = FNiagaraGpuComputeDispatchInterface::Get(WorldManager->GetWorld());
+}
+
+void UNiagaraDataManager::BeginFrame()
+{
+
+}
+
+void UNiagaraDataManager::EndFrame()
+{
+	//Remove component to ref map entry.
+	for (auto ComponentToDataIt = NiagaraComponentToSharedData.CreateIterator(); ComponentToDataIt; ++ComponentToDataIt)
+	{
+		UNiagaraComponent* Comp = ComponentToDataIt.Key().Get();
+		TArray<FNiagaraSharedDataPtr>& ComponentSharedData = ComponentToDataIt.Value();
+
+		if(IsValid(Comp) == false)
+		{
+			ComponentSharedData.Empty();
+		}
+
+		for(auto DataIt = ComponentSharedData.CreateIterator(); DataIt; ++DataIt)
+		{
+			FNiagaraSharedDataPtr& Data = *DataIt;
+			if(!Data.IsValid() || Data->bInvalidated)
+			{
+				DataIt.RemoveCurrentSwap();
+			}
+		}
+
+		if(ComponentSharedData.Num() == 0)
+		{
+			if(IsValid(Comp))
+			{
+				Comp->OnSystemFinished.RemoveDynamic(this, &UNiagaraDataManager::OnReferencingSystemFinished);
+			}
+			ComponentToDataIt.RemoveCurrent();
+		}
+	}
+}
+
+void UNiagaraDataManager::RegisterReferencingSystem(UNiagaraComponent* Component, FNiagaraSharedDataPtr Data)
+{
+	check(Component && Data.IsValid());
+	NiagaraComponentToSharedData.FindOrAdd(Component).AddUnique(Data);
+	Component->OnSystemFinished.AddUniqueDynamic(this, &UNiagaraDataManager::OnReferencingSystemFinished);
+}
+
+void UNiagaraDataManager::OnReferencingSystemFinished(UNiagaraComponent* Component)
+{
+	NiagaraComponentToSharedData.Remove(Component);
+	Component->OnSystemFinished.RemoveDynamic(this, &UNiagaraDataManager::OnReferencingSystemFinished);
+}
+
diff --git a/Engine/Plugins/FX/Niagara/Source/Niagara/Private/NiagaraDataManager_StaticMesh.cpp b/Engine/Plugins/FX/Niagara/Source/Niagara/Private/NiagaraDataManager_StaticMesh.cpp
new file mode 100644
--- /dev/null
+++ b/Engine/Plugins/FX/Niagara/Source/Niagara/Private/NiagaraDataManager_StaticMesh.cpp
@@ -0,0 +1,221 @@
+// Copyright Epic Games, Inc. All Rights Reserved.
+#include "NiagaraDataManager_StaticMesh.h"
+
+#include "PrimitiveSceneInfo.h"
+#include "NiagaraGpuComputeDispatchInterface.h"
+#include "NiagaraDataInterface.h"
+#include "RenderGraphUtils.h"
+#include "SystemTextures.h"
+#include "Components/StaticMeshComponent.h"
+#include "SceneInterface.h"
+
+//////////////////////////////////////////////////////////////////////////
+// UNiagaraDataManager_StaticMesh
+// Allows indirect referencing of static mesh data from Niagara via a persistent ID.
+// 
+
+UNiagaraDataManager_StaticMesh::UNiagaraDataManager_StaticMesh(FObjectInitializer const& ObjectInitializer)
+	: Super(ObjectInitializer)
+{
+	FNiagaraTypeRegistry::Register(FNiagaraTypeDefinition(FNiagaraSharedDataID_StaticMesh::StaticStruct()), ENiagaraTypeRegistryFlags::AllowParameter | ENiagaraTypeRegistryFlags::AllowAnyVariable);
+}
+
+void UNiagaraDataManager_StaticMesh::EndFrame()
+{
+	++IDAcquireTag;
+
+	TArray<FPrimitiveComponentId> ToRemove;
+	int32 MaxSlotsUsed = -1;
+	//Remove and invalidate any invalid/refs
+	for (auto It = TrackedComponents.CreateIterator(); It; ++It)
+	{
+		TWeakObjectPtr<UStaticMeshComponent> WeakComp = It.Key();
+		FNiagaraSharedData_StaticMeshPtr& MeshData = It.Value();
+		
+		check(MeshData.IsValid());
+		
+		bool bRefIsValid = WeakComp.IsValid() && !MeshData.IsUnique();
+		int32 Slot = MeshData->ID.Index;
+		if (bRefIsValid)
+		{
+			MaxSlotsUsed = FMath::Max(Slot, MaxSlotsUsed);
+			continue;
+		}
+
+		//Signal to anyone else holding this ref that it should be removed.
+		MeshData->bInvalidated = true;
+
+		ToRemove.Add(MeshData->PrimitiveId);
+		FreeSlots.Add(Slot);		
+		It.RemoveCurrent();
+	}
+
+	Super::EndFrame();
+
+	NumSlots = MaxSlotsUsed + 1;
+
+	if (ToRemove.Num() > 0)
+	{
+		ENQUEUE_RENDER_COMMAND(RemovePrimIDtoMeshID)(
+			[RTDispatchInterface = ComputeDispatchInterface, RTToRemove = MoveTemp(ToRemove), RTNumSlots = NumSlots](FRHICommandListImmediate& CmdList)
+			{
+				if (RTDispatchInterface)
+				{
+					FNiagaraDataManager_StaticMesh_RTProxy& Manager = RTDispatchInterface->GetOrCreateDataManager<FNiagaraDataManager_StaticMesh_RTProxy>();
+					for (const FPrimitiveComponentId& PrimitiveIdToRemove : RTToRemove)
+					{
+						Manager.Remove(PrimitiveIdToRemove);
+					}
+					Manager.SetNumSlots(RTNumSlots);
+				}
+			}
+		);
+	}
+}
+
+FNiagaraSharedData_StaticMeshPtr UNiagaraDataManager_StaticMesh::GetOrCreateSharedMeshData(UStaticMeshComponent* ComponentToReference)
+{
+	if(FNiagaraSharedData_StaticMeshPtr* FoundData = TrackedComponents.Find(ComponentToReference))
+	{
+		return *FoundData;
+	}
+
+	FNiagaraSharedData_StaticMeshPtr& NewMeshData = TrackedComponents.Add(ComponentToReference);
+	NewMeshData = MakeShared<FNiagaraSharedData_StaticMesh>();
+	NewMeshData->PrimitiveId = ComponentToReference->GetPrimitiveSceneId();
+	NewMeshData->ID.AcquireTag = IDAcquireTag;
+	NewMeshData->SourceComponent = ComponentToReference;
+
+	int32 IDIndex = NumSlots;
+	if(FreeSlots.Num() > 0)
+	{
+		IDIndex = FreeSlots.Pop();
+	}
+
+	NewMeshData->ID.Index = IDIndex;
+	NumSlots = FMath::Max(IDIndex + 1, NumSlots);
+
+	ENQUEUE_RENDER_COMMAND(AddPrimIDtoMeshID)(
+		[RTDispatchInterface = ComputeDispatchInterface, RTPrimId = NewMeshData->PrimitiveId, RTRefID = NewMeshData->ID, RTNumSlots = NumSlots](FRHICommandListImmediate& CmdList)
+		{
+			if (RTDispatchInterface)
+			{
+				FNiagaraDataManager_StaticMesh_RTProxy& Manager = RTDispatchInterface->GetOrCreateDataManager<FNiagaraDataManager_StaticMesh_RTProxy>();
+				Manager.Add(RTPrimId, RTRefID);
+				Manager.SetNumSlots(RTNumSlots);
+			}
+		}
+	);
+
+	return NewMeshData;
+}
+
+//////////////////////////////////////////////////////////////////////////
+// FNiagaraDataManager_StaticMesh_RTProxy
+// Manager class on the GPU for static mesh data. Effectively an RT Proxy for the GT side manager.
+
+void FNiagaraDataManager_StaticMesh_RTProxy::Add(FPrimitiveComponentId PrimitiveId, FNiagaraSharedDataID_StaticMesh MeshID)
+{
+	PrimitiveToMeshID.Add(PrimitiveId) = MeshID;
+}
+void FNiagaraDataManager_StaticMesh_RTProxy::Remove(FPrimitiveComponentId PrimitiveId)
+{
+	PrimitiveToMeshID.Remove(PrimitiveId);
+}
+	
+void FNiagaraDataManager_StaticMesh_RTProxy::BeginFrame(FNiagaraGpuComputeDataManagerTickContext& Context)
+{
+
+}
+
+void FNiagaraDataManager_StaticMesh_RTProxy::PreTick(FNiagaraGpuComputeDataManagerTickContext& Context, ENiagaraGpuComputeTickStage::Type TickStage)
+{
+	if(MeshDataBuffer == nullptr)
+	{
+		static_assert(sizeof(FNiagaraSharedDataGPU_StaticMesh) == 16u);
+		//TODO: If we have some hooks into add/remove primitive and can be sure if/when the DF indices can change then we could keep this data persistent until it's dirty.	
+		//NOTE: This building must be done after the DF Scene has been updated so shouldn't be done early from PreInitViews etc.
+		const FSceneInterface* Scene = GetOwnerInterface()->GetSceneInterface();
+		check(Scene);
+
+		TArray<FNiagaraSharedDataGPU_StaticMesh> Buffer;
+		Buffer.SetNum(NumSlots);
+		for (auto It = PrimitiveToMeshID.CreateConstIterator(); It; ++It)
+		{
+			FPrimitiveComponentId PrimitiveId = It.Key();
+			const FNiagaraSharedDataID_StaticMesh& StaticMeshID = It.Value();
+			if (const FPrimitiveSceneInfo* PrimitiveSceneInfo = Scene ? Scene->GetPrimitiveSceneInfo(PrimitiveId) : nullptr)
+			{
+				int32 Slot = StaticMeshID.Index;
+				
+				if(Buffer.IsValidIndex(Slot))
+				{
+					Buffer[Slot].IDTag = StaticMeshID.AcquireTag;
+
+					Buffer[Slot].DistanceFieldIndex = PrimitiveSceneInfo->DistanceFieldInstanceIndices.Num() > 0 ? PrimitiveSceneInfo->DistanceFieldInstanceIndices[0] : INDEX_NONE;
+				
+					//TODO: if we bind GPU scene buffers we can access a bunch of data.
+					Buffer[Slot].PrimitiveIndex = PrimitiveSceneInfo->GetIndex();
+
+					//TODO: Could pass bindless handles to vertex data etc...
+				
+					//TODO: Can pass in bounds here if we want but could be better to access via GPU scene info above?
+	// 				FBoxSphereBounds Bounds = PrimitiveSceneInfo->Proxy->GetBounds();
+	// 				Buffer[Slot].BoundsRadius = Bounds.SphereRadius;
+	// 				Buffer[Slot].BoundsExtents = FVector3f(Bounds.BoxExtent);
+	// 				
+	// 				FVector3f LWCTile  = FLargeWorldRenderScalar::GetTileFor(Bounds.Origin);
+	// 				FVector3f SimLocalSWC = FVector3f(Bounds.Origin - FVector(LWCTile) * FLargeWorldRenderScalar::GetTileSize());
+	// 				Buffer[Slot].BoundsOrigin = SimLocalSWC;
+	// 				Buffer[Slot].BoundsOriginTile = LWCTile;
+				}
+			}
+		}
+
+		MeshDataBuffer = CreateStructuredBuffer(Context.GetGraphBuilder(), TEXT("NiagaraSharedMeshDataBuffer"), MoveTemp(Buffer));
+		MeshDataBufferSRV = Context.GetGraphBuilder().CreateSRV(MeshDataBuffer);		
+	}
+}
+
+void FNiagaraDataManager_StaticMesh_RTProxy::EndFrame(FNiagaraGpuComputeDataManagerTickContext& Context)
+{
+	MeshDataBufferSRV = nullptr;
+	MeshDataBuffer = nullptr;
+}
+
+FRDGBufferSRVRef FNiagaraDataManager_StaticMesh_RTProxy::GetGPUBufferSRV()const
+{
+	return MeshDataBufferSRV;
+}
+
+const FShaderParametersMetadata* FNiagaraDataManager_StaticMesh_RTProxy::GetParametersMetadata()
+{
+	return TShaderParameterStructTypeInfo<FNiagaraDataManager_StaticMesh_Parameters>::GetStructMetadata();
+}
+
+void FNiagaraDataManager_StaticMesh_RTProxy::SetShaderParameters(const FNiagaraDataInterfaceSetShaderParametersContext& Context)
+{
+	const FShaderParametersMetadata* ParametersMetadata = GetParametersMetadata();
+	if(uint8* ShaderParameters = Context.GetParameterIncludedStruct(ParametersMetadata))
+	{
+		if (Context.IsStructBound(ShaderParameters, ParametersMetadata))
+		{
+			FNiagaraDataManager_StaticMesh_Parameters* TypedParameters = reinterpret_cast<FNiagaraDataManager_StaticMesh_Parameters*>(Context.GetParameterIncludedStruct(ParametersMetadata));
+			
+			if (const FNiagaraDataManager_StaticMesh_RTProxy* MeshDataManager = Context.GetComputeDispatchInterface().GetDataManager<FNiagaraDataManager_StaticMesh_RTProxy>())
+			{
+				TypedParameters->NiagaraSharedData_StaticMesh_Buffer = MeshDataManager->GetGPUBufferSRV();
+				TypedParameters->NiagaraSharedData_StaticMesh_NumSlots = MeshDataManager->GetNumSlots();
+			}
+			else
+			{			
+				TArray<FNiagaraSharedDataGPU_StaticMesh> Buffer;
+				Buffer.SetNum(1);
+				FRDGBufferRef MeshDataBuffer = CreateStructuredBuffer(Context.GetGraphBuilder(), TEXT("DummyNiagaraSharedMeshDataBuffer"), Buffer);;
+				FRDGBufferSRVRef MeshDataBufferSRV = Context.GetGraphBuilder().CreateSRV(MeshDataBuffer);
+				TypedParameters->NiagaraSharedData_StaticMesh_Buffer = MeshDataBufferSRV;
+				TypedParameters->NiagaraSharedData_StaticMesh_NumSlots = 0;
+			}
+		}
+	}
+}
\ No newline at end of file
diff --git a/Engine/Plugins/FX/Niagara/Source/Niagara/Private/NiagaraGpuComputeDispatch.cpp b/Engine/Plugins/FX/Niagara/Source/Niagara/Private/NiagaraGpuComputeDispatch.cpp
--- a/Engine/Plugins/FX/Niagara/Source/Niagara/Private/NiagaraGpuComputeDispatch.cpp
+++ b/Engine/Plugins/FX/Niagara/Source/Niagara/Private/NiagaraGpuComputeDispatch.cpp
@@ -1171,6 +1171,15 @@ void FNiagaraGpuComputeDispatch::ExecuteTicks(FRDGBuilder& GraphBuilder, TConstS
 	}
 
 	SubstratePublicGlobalUniformParams = ::Substrate::GetPublicGlobalUniformBuffer(GraphBuilder, *GetScene());
+	
+	FNiagaraGpuComputeDataManagerTickContext Context(GraphBuilder);
+	for (auto& DataManager : GpuDataManagers)
+	{
+		if (DataManager.Value)
+		{
+			DataManager.Value->PreTick(Context, TickStage);
+		}
+	}
 
 	// Loop over dispatches
 	for ( const FNiagaraGpuDispatchGroup& DispatchGroup : DispatchList.DispatchGroups )
@@ -2034,6 +2043,15 @@ void FNiagaraGpuComputeDispatch::PreInitViews(FRDGBuilder& GraphBuilder, bool bA
 					It.RemoveCurrentSwap();
 				}
 			}
+			
+			FNiagaraGpuComputeDataManagerTickContext Context(GraphBuilder);
+			for (auto& DataManager : GpuDataManagers)
+			{
+				if (DataManager.Value)
+				{
+					DataManager.Value->BeginFrame(Context);
+				}
+			}
 
 			UpdateInstanceCountManager(GraphBuilder.RHICmdList);
 			PrepareAllTicks(GraphBuilder.RHICmdList);
@@ -2192,7 +2210,7 @@ void FNiagaraGpuComputeDispatch::PostRenderOpaque(FRDGBuilder& GraphBuilder, TCo
 	}
 
 	GraphBuilder.AddPostExecuteCallback(
-		[this, &RHICmdList=GraphBuilder.RHICmdList, bAllowGPUParticleUpdate]
+		[this, &RHICmdList=GraphBuilder.RHICmdList, &GraphBuilder, bAllowGPUParticleUpdate]
 		{
 		#if WITH_NIAGARA_GPU_PROFILER
 			GPUProfilerPtr->EndFrame(RHICmdList);
@@ -2216,6 +2234,15 @@ void FNiagaraGpuComputeDispatch::PostRenderOpaque(FRDGBuilder& GraphBuilder, TCo
 					}
 				}
 			}
+			
+			FNiagaraGpuComputeDataManagerTickContext Context(GraphBuilder);
+			for (auto& DataManager : GpuDataManagers)
+			{
+				if (DataManager.Value)
+				{
+					DataManager.Value->EndFrame(Context);
+				}
+			}
 		}
 	);
 
diff --git a/Engine/Plugins/FX/Niagara/Source/Niagara/Private/NiagaraWorldManager.cpp b/Engine/Plugins/FX/Niagara/Source/Niagara/Private/NiagaraWorldManager.cpp
--- a/Engine/Plugins/FX/Niagara/Source/Niagara/Private/NiagaraWorldManager.cpp
+++ b/Engine/Plugins/FX/Niagara/Source/Niagara/Private/NiagaraWorldManager.cpp
@@ -600,6 +600,8 @@ void FNiagaraWorldManager::AddReferencedObjects(FReferenceCollector& Collector)
 	Collector.AddReferencedObjects(ReferencedObjects);
 
 	DataChannelManager->AddReferencedObjects(Collector);
+
+	Collector.AddReferencedObjects(DataManagers);
 }
 
 FString FNiagaraWorldManager::GetReferencerName() const
@@ -842,6 +844,8 @@ void FNiagaraWorldManager::CleanupTransientObjectsAndManagers()
 	CullProxyMap.Empty();
 
 	ScalabilityManagers.Empty();
+
+	DataManagers.Empty();
 }
 
 void FNiagaraWorldManager::OnWorldBeginTearDown()
@@ -927,6 +931,8 @@ void FNiagaraWorldManager::LegacyOnWorldBeginTearDown()
 		ComponentPool->Cleanup(World);
 	}
 	DataChannelManager->Cleanup();
+	
+	DataManagers.Empty();
 }
 
 void FNiagaraWorldManager::LegacyOnWorldCleanup(bool bSessionEnded, bool bCleanupResources)
@@ -940,6 +946,8 @@ void FNiagaraWorldManager::LegacyOnWorldCleanup(bool bSessionEnded, bool bCleanu
 	}
 
 	DataChannelManager->Cleanup();
+	
+	DataManagers.Empty();
 
 	for (int TG = 0; TG < NiagaraNumTickGroups; ++TG)
 	{
@@ -984,6 +992,8 @@ void FNiagaraWorldManager::LegacyOnPostWorldCleanup(bool bSessionEnded, bool bCl
 	}
 
 	DataChannelManager->Cleanup();
+	
+	DataManagers.Empty();
 }
 
 //////////////////////////////////////////////////////////////////////////
@@ -1318,6 +1328,15 @@ void FNiagaraWorldManager::TickStart(float DeltaSeconds)
 
 	DataChannelManager->BeginFrame(DeltaSeconds);
 
+	for(auto& Pair : DataManagers)
+	{
+		UNiagaraDataManager* Manager = Pair.Value.Get();
+		if(Manager)
+		{
+			Manager->BeginFrame();
+		}
+	}
+
 #if WITH_NIAGARA_LEAK_DETECTOR
 	if (ComponentLeakDetector.IsValid() && HasActiveWorld())
 	{
@@ -1360,6 +1379,15 @@ void FNiagaraWorldManager::TickEnd(float DeltaSeconds)
 	HandleCSVStats(DeltaSeconds);
 
 	DataChannelManager->EndFrame(DeltaSeconds);
+	
+	for(auto& Pair : DataManagers)
+	{
+		UNiagaraDataManager* Manager = Pair.Value.Get();
+		if(Manager)
+		{
+			Manager->EndFrame();
+		}
+	}
 }
 
 void FNiagaraWorldManager::PreActorTick(ELevelTick InLevelTick, float InDeltaSeconds)
@@ -1640,6 +1668,15 @@ void FNiagaraWorldManager::ExecutePreTickWork(ETickingGroup TickGroup, float Del
 
 	//Tick DataChannel Manager.
 	DataChannelManager->Tick(DeltaSeconds, TickGroup);
+	
+	for(auto& Pair : DataManagers)
+	{
+		UNiagaraDataManager* Manager = Pair.Value.Get();
+		if(Manager)
+		{
+			Manager->PreTick(TickGroup, DeltaSeconds);
+		}
+	}
 
 	// We do book keeping in the first tick group
 	if ( TickGroup == NiagaraFirstTickGroup )
diff --git a/Engine/Plugins/FX/Niagara/Source/Niagara/Public/NiagaraDataChannelAccessContext.h b/Engine/Plugins/FX/Niagara/Source/Niagara/Public/NiagaraDataChannelAccessContext.h
--- a/Engine/Plugins/FX/Niagara/Source/Niagara/Public/NiagaraDataChannelAccessContext.h
+++ b/Engine/Plugins/FX/Niagara/Source/Niagara/Public/NiagaraDataChannelAccessContext.h
@@ -153,13 +153,23 @@ struct FNDCAccessContext : public FNDCAccessContextBase
 	/** If true, we'll override the system to spawn for this NDC write. If false, we'll use the system defined in the NDC asset, if there is one. */
 	UPROPERTY(BlueprintReadWrite, Category = "Input", meta = (InlineEditCondition, NDCAccessContextInput))
 	uint32 bOverrideSystemToSpawn : 1 = false;
+	
+	/** If true, SpawnedSystems will be filled with both existing and newly spawned systems. If false, SpawnedSystems will contain only newly spawned systems. */
+	UPROPERTY(BlueprintReadWrite, Category = "Input", meta = (NDCAccessContextInput))
+	uint32 bReturnExistingSystems : 1 = false;
 
 	UPROPERTY(EditAnywhere, BlueprintReadWrite, Category = "Input", meta = (EditCondition="bOverrideSystemToSpawn", NDCAccessContextInput, AllowedClasses = "/Script/Niagara.NiagaraSystem,/Script/Niagara.NiagaraSystemCollection"))
 	TObjectPtr<UObject> SystemToSpawn;
 
-	/** The Niagara System spawned in response to this access, if any. Allows callers to set parameters etc on this system. It is not safe to store and use this component reference later. It can be reclaimed by the system at any time. */
+	/** 
+	The Niagara System spawned in response to this access, if any. Allows callers to set parameters etc on this system. It is not safe to store and use this component reference later. It can be reclaimed by the system at any time.
+	If bReturnExistingSystems is set, this will contain both new and existing systems matching this access context.
+	*/
 	UPROPERTY(BlueprintReadWrite, Transient, Category = "Output", meta = (NDCAccessContextOutput, ShowOnlyInnerProperties))
 	TArray<FNDCSpawnedSystemRef> SpawnedSystems;
+	
+	template<typename TAction>
+	void ForEachSpawnedSystem(TAction Func) const;
 };
 
 /** Special legacy access context type. Used only to facilitate legacy NDCs to be usable by the new API. */
diff --git a/Engine/Plugins/FX/Niagara/Source/Niagara/Public/NiagaraDataChannelAccessContextImpl.h b/Engine/Plugins/FX/Niagara/Source/Niagara/Public/NiagaraDataChannelAccessContextImpl.h
new file mode 100644
--- /dev/null
+++ b/Engine/Plugins/FX/Niagara/Source/Niagara/Public/NiagaraDataChannelAccessContextImpl.h
@@ -0,0 +1,17 @@
+// Copyright Epic Games, Inc. All Rights Reserved.
+
+#pragma once
+
+#include "NiagaraDataChannelAccessContext.h"
+
+template<typename TAction>
+void FNDCAccessContext::ForEachSpawnedSystem(TAction Func) const
+{
+	for ( const FNDCSpawnedSystemRef& SysRef : SpawnedSystems )
+	{
+		if (UNiagaraComponent* Comp  = SysRef.Get())
+		{
+			Func(Comp);
+		}
+	}
+}
\ No newline at end of file
diff --git a/Engine/Plugins/FX/Niagara/Source/Niagara/Public/NiagaraDataChannelAccessor.h b/Engine/Plugins/FX/Niagara/Source/Niagara/Public/NiagaraDataChannelAccessor.h
--- a/Engine/Plugins/FX/Niagara/Source/Niagara/Public/NiagaraDataChannelAccessor.h
+++ b/Engine/Plugins/FX/Niagara/Source/Niagara/Public/NiagaraDataChannelAccessor.h
@@ -272,6 +272,7 @@ protected:
 public:
 
 	bool IsInitialized()const { return CachedLayout.IsValid() && !CachedLayout.IsUnique(); }
+	bool IsInitialized(UNiagaraDataChannel* NDC)const { return CachedLayout.IsValid() && !CachedLayout.IsUnique() && NDC && CachedLayout == NDC->GetLayoutInfo(); }
 
 	//Initialize the writer and update cached layout information.
 	NIAGARA_API void Init(const UNiagaraDataChannel* DataChannel);
@@ -470,4 +471,65 @@ FNDCVarReader<VarType> VarName##Reader = FNDCVarReader<VarType>(*this, #VarName,
 bool Read##VarName(int32 Index, TNiagaraTypeHelper<VarType>::TLWCType& OutData){ return VarName##Reader.Read(Data, Index, OutData, bPreviousFrame); }\
 const FNDCVarReader<VarType>& Get##VarName##Reader()const { return VarName##Reader; }
 
+struct FNDCPackedData
+{
+	int32 GetData() { return static_cast<int32>(PackedData); }//convert to int32 for passing into Niagara.
+
+protected:
+	uint32 PackedData = 0;
+
+	template <typename T, uint32 Shift, uint32 NumBits>
+	inline void Set(T Value)
+	{
+		static_assert(std::is_integral_v<T>, "Must use integral types for NDC packed data.");
+		static_assert(!std::is_signed_v<T>, "Signed values not supported currently but could be.");
+		static_assert(Shift < 32, "Shift cannot exceeded size of packed uint32.");
+		static_assert(NumBits > 0, "Cannot pack zero bits.");
+		static_assert(NumBits < 32, "Cannot pack more bits that the underlying uint32.");		
+		static_assert(Shift + NumBits < 32, "Cannot pack more bits that the underlying uint32.");
+
+		uint32 V;
+		if constexpr(std::is_same_v<T,bool>)
+		{
+			static_assert(NumBits == 1, "Bools should be packed into one bit.");
+			V = !!Value;
+		}
+		else
+		{
+			V = Value;
+		}
+
+		uint32 Mask = (1u << NumBits) - 1u;
+		PackedData &= ~(Mask << Shift);
+		PackedData |= (V & Mask) << Shift; 
+	}
+
+	template<typename T, uint32 Shift, uint32 NumBits>
+	inline T Get()
+	{
+		static_assert(std::is_integral_v<T>, "Must use integral types for NDC packed data.");
+		static_assert(!std::is_signed_v<T>, "Signed values not supported currently but could be.");
+		static_assert(Shift < 32, "Shift cannot exceeded size of packed uint32.");
+		static_assert(NumBits > 0, "Cannot pack zero bits.");
+		static_assert(NumBits < 32, "Cannot pack more bits that the underlying uint32.");		
+		static_assert(Shift + NumBits < 32, "Cannot pack more bits that the underlying uint32.");
+
+		uint32 Mask = (1u << NumBits) - 1u;
+		uint32 V = (PackedData >> Shift) & Mask;
+		if constexpr(std::is_same_v<T,bool>)
+		{
+			static_assert(NumBits == 1, "Bools should be packed into one bit.");
+			return V != 0;
+		}
+		else
+		{
+			return V;
+		}
+	}
+};
+
+#define NDCVarPackedValue(Name, Type, Shift, NumBits) \
+void Set_##Name(Type Value){ Set<Type, Shift, NumBits>(Value); }\
+Type Get_##Name(){ return Get<Type, Shift, NumBits>(); }
+
 #undef UE_API
diff --git a/Engine/Plugins/FX/Niagara/Source/Niagara/Public/NiagaraDataChannelHandler.h b/Engine/Plugins/FX/Niagara/Source/Niagara/Public/NiagaraDataChannelHandler.h
--- a/Engine/Plugins/FX/Niagara/Source/Niagara/Public/NiagaraDataChannelHandler.h
+++ b/Engine/Plugins/FX/Niagara/Source/Niagara/Public/NiagaraDataChannelHandler.h
@@ -58,6 +58,8 @@ public:
 		check(Ret);
 		return Ret;
 	}
+	
+	NIAGARA_API virtual void OnDataChannelLayoutChanged();
 
 	UFUNCTION(BlueprintCallable, Category="Data Channel")
 	NIAGARA_API UNiagaraDataChannelWriter* GetDataChannelWriter();
@@ -121,6 +123,7 @@ protected:
 
 	FRenderCommandFence RTFence;
 
+	FNiagaraDataChannelLayoutInfoPtr CachedLayoutInfo;
 private:
 
 	struct FChannelSubscription
diff --git a/Engine/Plugins/FX/Niagara/Source/Niagara/Public/NiagaraDataChannel_Map.h b/Engine/Plugins/FX/Niagara/Source/Niagara/Public/NiagaraDataChannel_Map.h
--- a/Engine/Plugins/FX/Niagara/Source/Niagara/Public/NiagaraDataChannel_Map.h
+++ b/Engine/Plugins/FX/Niagara/Source/Niagara/Public/NiagaraDataChannel_Map.h
@@ -84,6 +84,7 @@ struct FNDCMapEntryBase
 
 	[[nodiscard]] float GetLastUsedTime()const{ return LastUsedTime; }
 
+	const TArray<TObjectPtr<UNiagaraComponent>>& GetSpawnedSystems() { return SpawnedComponents; }
 protected:
 
 	/** The owning handler for this entry. */
@@ -127,7 +128,8 @@ struct FNDCMapEntry
 {
 	GENERATED_BODY()
 	FNDCMapEntry() = default;
-	[[nodiscard]] explicit FNDCMapEntry(TNDCMapEntryType EntryType)
+	[[nodiscard]] explicit FNDCMapEntry(TNDCMapEntryType EntryType, int32 InIndex)
+	: Index(InIndex)
 	{
 		EntryData.InitializeAs(EntryType.Get());
 	}
@@ -138,12 +140,16 @@ struct FNDCMapEntry
 	template<typename T=FNDCMapEntryBase>
 	[[nodiscard]] const T& Get()const { return EntryData.Get<T>(); }
 
-	[[nodiscard]] bool IsValid()const { return EntryData.IsValid(); }
+	[[nodiscard]] int32 GetIndex()const { return Index; }
 
+	[[nodiscard]] bool IsValid()const { return EntryData.IsValid(); }
 private:
 
 	UPROPERTY()
 	FInstancedStruct EntryData;
+
+	/** Index of this entry in the owning pool. */
+	int32 Index = INDEX_NONE;
 };
 
 
diff --git a/Engine/Plugins/FX/Niagara/Source/Niagara/Public/NiagaraDataManager.h b/Engine/Plugins/FX/Niagara/Source/Niagara/Public/NiagaraDataManager.h
new file mode 100644
--- /dev/null
+++ b/Engine/Plugins/FX/Niagara/Source/Niagara/Public/NiagaraDataManager.h
@@ -0,0 +1,74 @@
+// Copyright Epic Games, Inc. All Rights Reserved.
+
+#pragma once
+
+#include "CoreMinimal.h"
+#include "NiagaraTypes.h"
+#include "NiagaraGpuComputeDataManager.h"
+#include "NiagaraDataManager.generated.h"
+
+class FNiagaraWorldManager;
+class FNiagaraGpuComputeDispatchInterface;
+
+/** 
+Base class for generic data items held by data managers.
+Niagara Systems and NDCs will hold onto these when they are being used.
+Allows Data Managers to ensure lifetime of their internal data that may be being reference inside Niagara and by external users.
+*/
+struct FNiagaraSharedData : public TSharedFromThis<FNiagaraSharedData>
+{
+	virtual ~FNiagaraSharedData(){}
+
+	//Flag to mark this ref as invalid and it should be removed.
+	uint8 bInvalidated : 1 = false;
+};
+
+typedef TSharedPtr<FNiagaraSharedData> FNiagaraSharedDataPtr;
+
+/** 
+A generic manager class holding some data requiring management or marshaling for Niagara.
+A counterpart for Gpu accessible data is FNiagaraGpuComputeDataManager.
+Some cases may utilize a UNiagaraDataManager and a FNiagaraGpuComputeDataManager together.
+*/
+UCLASS()
+class UNiagaraDataManager : public UObject
+{
+	GENERATED_BODY()
+
+public:
+
+	virtual void Init(FNiagaraWorldManager* WorldMan);		
+	virtual void BeginFrame();
+	virtual void EndFrame();
+	virtual void PreTick(ETickingGroup TickGroup, float DeltaSeconds){}
+
+	/** Track the connection between a NiagaraComponent and a FNiagaraSharedDataPtr. The data will remain valid at least until the valid component completes. */
+	NIAGARA_API void RegisterReferencingSystem(UNiagaraComponent* Component, FNiagaraSharedDataPtr SharedData);
+	
+	/** Callback when a Niagara component completes so we can remove any tracked data associated with it. */
+	UFUNCTION()
+	void OnReferencingSystemFinished(UNiagaraComponent* Component);
+
+protected:
+	
+	FNiagaraWorldManager* WorldManager = nullptr;
+	FNiagaraGpuComputeDispatchInterface* ComputeDispatchInterface = nullptr;
+
+	/** 
+	Map of active Niagara Components referencing FNiagaraSharedDataPtrs.
+	Allows us to maintain the lifetime of a FNiagaraSharedDataPtr until a referencing Niagara Component completes.
+	*/
+	TMap<TWeakObjectPtr<UNiagaraComponent>, TArray<FNiagaraSharedDataPtr>> NiagaraComponentToSharedData;
+};
+
+struct FNiagaraDataInterfaceSetShaderParametersContext;
+typedef FNiagaraDataInterfaceSetShaderParametersContext FNiagaraDataManagerSetShaderParametersContext;
+
+/** Base class for RT Proxies of data managers. */
+struct FNiagaraDataManager_RTProxy : public FNiagaraGpuComputeDataManager
+{
+	FNiagaraDataManager_RTProxy(FNiagaraGpuComputeDispatchInterface* InOwnerInterface)
+	: FNiagaraGpuComputeDataManager(InOwnerInterface)
+	{
+	}
+};
\ No newline at end of file
diff --git a/Engine/Plugins/FX/Niagara/Source/Niagara/Public/NiagaraDataManager_StaticMesh.h b/Engine/Plugins/FX/Niagara/Source/Niagara/Public/NiagaraDataManager_StaticMesh.h
new file mode 100644
--- /dev/null
+++ b/Engine/Plugins/FX/Niagara/Source/Niagara/Public/NiagaraDataManager_StaticMesh.h
@@ -0,0 +1,123 @@
+// Copyright Epic Games, Inc. All Rights Reserved.
+
+#pragma once
+
+#include "CoreMinimal.h"
+#include "NiagaraTypes.h"
+#include "MeshPassProcessor.h"
+
+#include "NiagaraDataManager.h"
+#include "NiagaraDataManager_StaticMesh.generated.h"
+
+
+class UPrimitiveComponent;
+struct FNiagaraDataInterfaceSetShaderParametersContext;
+
+/** Persistent ID allowing Niagara particles access to the static mesh info provided by the manager. */
+USTRUCT(BlueprintType)
+struct FNiagaraSharedDataID_StaticMesh : public FNiagaraID
+{
+	GENERATED_BODY()
+};
+
+//Shared data about a static mesh we are tracking for indirect access in Niagara.
+struct FNiagaraSharedData_StaticMesh : public FNiagaraSharedData
+{
+	TWeakObjectPtr<UStaticMeshComponent> SourceComponent;
+	
+	FPrimitiveComponentId PrimitiveId;
+
+	//The mesh reference generated for this primitive. 
+	FNiagaraSharedDataID_StaticMesh ID;
+};
+
+typedef TSharedPtr<FNiagaraSharedData_StaticMesh> FNiagaraSharedData_StaticMeshPtr;
+
+UCLASS()
+class UNiagaraDataManager_StaticMesh : public UNiagaraDataManager
+{
+	GENERATED_BODY()
+
+public:
+	
+	UNiagaraDataManager_StaticMesh(FObjectInitializer const& ObjectInitializer);
+
+	//UNiagaraDataManager Interface
+	virtual void EndFrame()override;
+	//END UNiagaraDataManager Interface
+	
+	/** 
+	Finds or creates shared data tracking info about the given mesh component for use later in Niagara.
+	Once no-one is holding this shared data pointer, the data will be cleaned up and invalid to access on the GPU.
+	Access is still safe but will fail.
+	You can have Niagara Components hold the reference util they complese using RegisterReferencingSystem().
+	*/
+	[[nodiscard]] NIAGARA_API FNiagaraSharedData_StaticMeshPtr GetOrCreateSharedMeshData(UStaticMeshComponent* ComponentToReference);	
+
+private:
+	
+	TMap<TWeakObjectPtr<UStaticMeshComponent>, FNiagaraSharedData_StaticMeshPtr> TrackedComponents;
+	
+	/** Currently free slots in the MeshDataBuffer. */
+	TArray<int32> FreeSlots;
+	int32 NumSlots = 0;
+	int32 IDAcquireTag = 0;
+};
+
+/** 
+Render thread proxy for UNiagaraDataMangaer_StaticMesh.
+Tracks link between PrimitiveIds for tracked static mesh components and associated FNiagaraSharedDataID_StaticMesh.
+The manager generates a buffer each frame allowing particles to use a FNiagaraSharedDataID_StaticMesh to access info about each static mesh component.
+*/
+struct FNiagaraDataManager_StaticMesh_RTProxy final : public FNiagaraDataManager_RTProxy
+{	
+	static FName GetManagerName()
+	{
+		static FName ManagerName("FNiagaraDataManager_StaticMesh_RTProxy");
+		return ManagerName;
+	}
+	
+	FNiagaraDataManager_StaticMesh_RTProxy(FNiagaraGpuComputeDispatchInterface* InOwnerInterface)
+	: FNiagaraDataManager_RTProxy(InOwnerInterface)
+	{
+	}
+	
+	//FNiagaraDataManager_RTProxy Interface
+	virtual void BeginFrame(FNiagaraGpuComputeDataManagerTickContext& Context) override;
+	virtual void EndFrame(FNiagaraGpuComputeDataManagerTickContext& Context) override;
+	virtual void PreTick(FNiagaraGpuComputeDataManagerTickContext& Context, ENiagaraGpuComputeTickStage::Type TickStage) override;	
+	//END FNiagaraDataManager_RTProxy Interface
+	
+	void Add(FPrimitiveComponentId PrimitiveId, FNiagaraSharedDataID_StaticMesh MeshID);
+	void Remove(FPrimitiveComponentId PrimitiveId);
+	void SetNumSlots(int32 InNumSlots){ NumSlots = InNumSlots; }
+
+	[[nodiscard]] FRDGBufferSRVRef GetGPUBufferSRV()const;
+	[[nodiscard]] int32 GetNumSlots()const{ return NumSlots; }
+	
+	[[nodiscard]] static const FShaderParametersMetadata* GetParametersMetadata();
+	static void SetShaderParameters(const FNiagaraDataManagerSetShaderParametersContext& Context);
+
+private:
+
+	TMap<FPrimitiveComponentId, FNiagaraSharedDataID_StaticMesh> PrimitiveToMeshID;
+	int32 NumSlots = 0;
+
+	FRDGBufferRef MeshDataBuffer = nullptr;
+	FRDGBufferSRVRef MeshDataBufferSRV = nullptr;
+};
+
+/** Counterpart to struct in NiagaraManagedReference_StaticMesh.ush */
+struct FNiagaraSharedDataGPU_StaticMesh
+{
+	int32 IDTag = INDEX_NONE;
+	int32 DistanceFieldIndex = INDEX_NONE;
+	int32 PrimitiveIndex = INDEX_NONE;
+	int32 Padding = INDEX_NONE;
+};
+
+BEGIN_SHADER_PARAMETER_STRUCT(FNiagaraDataManager_StaticMesh_Parameters, )
+	SHADER_PARAMETER_RDG_BUFFER_SRV(StructuredBuffer<FNiagaraSharedDataGPU_StaticMesh>, NiagaraSharedData_StaticMesh_Buffer)
+	SHADER_PARAMETER(int32, NiagaraSharedData_StaticMesh_NumSlots)
+END_SHADER_PARAMETER_STRUCT()
+
diff --git a/Engine/Plugins/FX/Niagara/Source/Niagara/Public/NiagaraGpuComputeDataManager.h b/Engine/Plugins/FX/Niagara/Source/Niagara/Public/NiagaraGpuComputeDataManager.h
--- a/Engine/Plugins/FX/Niagara/Source/Niagara/Public/NiagaraGpuComputeDataManager.h
+++ b/Engine/Plugins/FX/Niagara/Source/Niagara/Public/NiagaraGpuComputeDataManager.h
@@ -5,6 +5,20 @@
 #include "NiagaraCommon.h"
 
 class FNiagaraGpuComputeDispatchInterface;
+class FRDGBuilder;
+
+struct FNiagaraGpuComputeDataManagerTickContext
+{
+	FNiagaraGpuComputeDataManagerTickContext(FRDGBuilder& InGraphBuilder)
+		: GraphBuilder(InGraphBuilder)
+	{
+	}
+
+	FRDGBuilder& GetGraphBuilder() const { return GraphBuilder; }
+
+private:
+	FRDGBuilder& GraphBuilder;
+};
 
 // Abstract class for managing GPU data for Niagara
 // Once the manager is created it will last for the lifetime of the owner dispatch interface
@@ -28,6 +42,10 @@ public:
 	//	return ManagerName;
 	//}
 
+	virtual void BeginFrame(FNiagaraGpuComputeDataManagerTickContext& Context){}
+	virtual void EndFrame(FNiagaraGpuComputeDataManagerTickContext& Context){}
+	virtual void PreTick(FNiagaraGpuComputeDataManagerTickContext& Context, ENiagaraGpuComputeTickStage::Type TickStage){}
+
 private:
 	FNiagaraGpuComputeDispatchInterface* InternalOwnerInterface = nullptr;
 };
diff --git a/Engine/Plugins/FX/Niagara/Source/Niagara/Public/NiagaraGpuComputeDispatchInterface.h b/Engine/Plugins/FX/Niagara/Source/Niagara/Public/NiagaraGpuComputeDispatchInterface.h
--- a/Engine/Plugins/FX/Niagara/Source/Niagara/Public/NiagaraGpuComputeDispatchInterface.h
+++ b/Engine/Plugins/FX/Niagara/Source/Niagara/Public/NiagaraGpuComputeDispatchInterface.h
@@ -69,6 +69,23 @@ public:
 	UE_DEPRECATED(5.4, "AddSortedGPUSimulation requires an RHI command list")
 	virtual bool AddSortedGPUSimulation(struct FNiagaraGPUSortInfo& SortInfo) final { return false; }
 
+	template<typename TManager>
+	const TManager* GetDataManager()const
+	{
+		check(IsInParallelRenderingThread());
+
+		UE::TScopeLock ScopeLock(ComputeManagerGuard);
+		const FName ManagerName = TManager::GetManagerName();
+		for (auto& DataManager : GpuDataManagers)
+		{
+			if (DataManager.Key == ManagerName)
+			{
+				return static_cast<const TManager*>(DataManager.Value.Get());
+			}
+		}
+		return nullptr;
+	}
+
 	/** Get or create the a data manager, must be done on the rendering thread only. */
 	template<typename TManager>
 	TManager& GetOrCreateDataManager()
@@ -232,5 +249,5 @@ protected:
 	FOnPostPreRenderEvent					OnPreRenderEvent;
 	FOnPostPreRenderEvent					OnPostRenderEvent;
 
-	UE::FMutex								ComputeManagerGuard;
+	mutable UE::FMutex						ComputeManagerGuard;
 };
diff --git a/Engine/Plugins/FX/Niagara/Source/Niagara/Public/NiagaraWorldManager.h b/Engine/Plugins/FX/Niagara/Source/Niagara/Public/NiagaraWorldManager.h
--- a/Engine/Plugins/FX/Niagara/Source/Niagara/Public/NiagaraWorldManager.h
+++ b/Engine/Plugins/FX/Niagara/Source/Niagara/Public/NiagaraWorldManager.h
@@ -21,6 +21,7 @@
 #include "NiagaraScalabilityManager.h"
 #include "NiagaraDebuggerCommon.h"
 #include "NiagaraDeferredMethodQueue.h"
+#include "NiagaraDataManager.h"
 
 #include "NiagaraWorldManager.generated.h"
 
@@ -316,6 +317,9 @@ public:
 	NIAGARA_API static void SetScalabilityCullingMode(ENiagaraScalabilityCullingMode NewMode);
 	static ENiagaraScalabilityCullingMode GetScalabilityCullingMode() { return ScalabilityCullingMode; }
 
+	template<typename T>
+	T& GetOrCreateDataManager();
+
 	FNiagaraDataChannelManager& GetDataChannelManager(){ return *DataChannelManager; }
 	NIAGARA_API void InitDataChannel(const UNiagaraDataChannel* Channel, bool bForce);
 	NIAGARA_API void RemoveDataChannel(const UNiagaraDataChannel* Channel);
@@ -439,6 +443,8 @@ private:
 
 	TMap<TObjectPtr<UNiagaraEffectType>, FNiagaraScalabilityManager> ScalabilityManagers;
 
+	TMap<FName, TObjectPtr<UNiagaraDataManager>> DataManagers;
+
 	FNiagaraDeferredMethodQueue DeferredMethods;
 
 	/** True if the app has focus. We prevent some culling if the app doesn't have focus as it can interfere. */
diff --git a/Engine/Plugins/FX/Niagara/Source/Niagara/Public/NiagaraWorldManagerImpl.h b/Engine/Plugins/FX/Niagara/Source/Niagara/Public/NiagaraWorldManagerImpl.h
new file mode 100644
--- /dev/null
+++ b/Engine/Plugins/FX/Niagara/Source/Niagara/Public/NiagaraWorldManagerImpl.h
@@ -0,0 +1,27 @@
+// Copyright Epic Games, Inc. All Rights Reserved.
+
+#pragma once
+
+#include "NiagaraWorldManager.h"
+
+template<typename T>
+T& FNiagaraWorldManager::GetOrCreateDataManager()
+{
+	FName ManagerName = T::StaticClass()->GetFName();
+	T* Ret = nullptr;
+	if(TObjectPtr<UNiagaraDataManager>* Found = DataManagers.Find(ManagerName))
+	{
+		Ret = CastChecked<T>(Found->Get());
+	}
+
+	if(Ret == nullptr)
+	{
+		TObjectPtr<UNiagaraDataManager>& NewMan = DataManagers.Add(ManagerName);
+		NewMan = NewObject<T>(GetTransientPackage(), ManagerName);
+		NewMan->Init(this);
+		Ret = CastChecked<T>(NewMan.Get());
+	}
+
+	check(Ret);
+	return *Ret;
+}
\ No newline at end of file
