Fixed static analysis warning.

diff --git a/Engine/Plugins/FX/Niagara/Source/Niagara/Private/DataInterface/NiagaraDataInterfaceStaticMeshIndirect.cpp b/Engine/Plugins/FX/Niagara/Source/Niagara/Private/DataInterface/NiagaraDataInterfaceStaticMeshIndirect.cpp
--- a/Engine/Plugins/FX/Niagara/Source/Niagara/Private/DataInterface/NiagaraDataInterfaceStaticMeshIndirect.cpp
+++ b/Engine/Plugins/FX/Niagara/Source/Niagara/Private/DataInterface/NiagaraDataInterfaceStaticMeshIndirect.cpp
@@ -63,8 +63,7 @@ void FNDIStaticMeshIndirectCompiledData::Init(const UNiagaraSystem* System, UNia
 		bUsedByGPU = true;
 
 		bool bIsDistanceFieldFunction = BindingInfo.DefinitionName == NDIStaticMeshIndirectLocal::QueryDistanceFieldName ||
-										BindingInfo.DefinitionName == NDIStaticMeshIndirectLocal::GetClosestPointOnMeshDistanceFieldName ||
-										BindingInfo.DefinitionName == NDIStaticMeshIndirectLocal::QueryDistanceFieldName;
+										BindingInfo.DefinitionName == NDIStaticMeshIndirectLocal::GetClosestPointOnMeshDistanceFieldName;
 		bRequiresDistanceFields |= bIsDistanceFieldFunction;
 		 
 		return true;
