Add log contexts to async tasks

#Synced-CL Unknown

diff --git a/Engine/Plugins/FX/Niagara/Source/NiagaraEditor/Private/NiagaraCompilationTasks.cpp b/Engine/Plugins/FX/Niagara/Source/NiagaraEditor/Private/NiagaraCompilationTasks.cpp
--- a/Engine/Plugins/FX/Niagara/Source/NiagaraEditor/Private/NiagaraCompilationTasks.cpp
+++ b/Engine/Plugins/FX/Niagara/Source/NiagaraEditor/Private/NiagaraCompilationTasks.cpp
@@ -6,6 +6,7 @@
 #include "AssetCompilingManager.h"
 #include "DataDrivenShaderPlatformInfo.h"
 #include "Interfaces/ITargetPlatformManagerModule.h"
+#include "Logging/StructuredLog.h"
 #include "Misc/PathViews.h"
 #include "NiagaraCompilationTypes.h"
 #include "NiagaraDigestDatabase.h"
@@ -1795,6 +1796,7 @@ void FNiagaraSystemCompilationTask::AddScript(int32 SourceEmitterIndex, UNiagara
 
 void FNiagaraSystemCompilationTask::Tick()
 {
+	UE_LOG_CONTEXT("NiagaraSystem", SystemInfo.GetSystemPath());
 	switch (CompilationState)
 	{
 		case EState::WaitingForProcessing:
@@ -2651,6 +2653,7 @@ void FNiagaraSystemCompilationTask::IssueCompilationTasks()
 		FTask PrecompileTask = Launch(UE_SOURCE_LOCATION, [Self = AsShared()]
 		{
 			COOK_STAT(auto Timer = NiagaraSystemCookStats::UsageStats.TimeSyncWork(); Timer.TrackCyclesOnly(););
+			UE_LOG_CONTEXT("NiagaraSystem", Self->SystemInfo.GetSystemPath());
 
 			Self->Precompile();
 		});
@@ -2663,6 +2666,7 @@ void FNiagaraSystemCompilationTask::IssueCompilationTasks()
 				FTask InstantiateGraphTask = Launch(UE_SOURCE_LOCATION, [Self = AsShared(), &GroupInfo]
 				{
 					COOK_STAT(auto Timer = NiagaraSystemCookStats::UsageStats.TimeSyncWork(); Timer.TrackCyclesOnly(););
+					UE_LOG_CONTEXT("NiagaraSystem", Self->SystemInfo.GetSystemPath());
 
 					GroupInfo.InstantiateCompileGraph(*Self);
 				}, PrecompileTask);
@@ -2680,6 +2684,7 @@ void FNiagaraSystemCompilationTask::IssueCompilationTasks()
 							COOK_STAT(auto Timer = NiagaraSystemCookStats::UsageStats.TimeSyncWork(); Timer.TrackCyclesOnly(););
 
 							FNiagaraSystemCompilationTask* TaskPtr = &Self.Get();
+							UE_LOG_CONTEXT("NiagaraSystem", TaskPtr->SystemInfo.GetSystemPath());
 
 							CompileTask.CollectNamedDataInterfaces(TaskPtr, GroupInfo);
 							switch (CompileTask.ScriptCompileType)
@@ -2714,6 +2719,7 @@ void FNiagaraSystemCompilationTask::IssueCompilationTasks()
 						Launch(UE_SOURCE_LOCATION, [Self = AsShared(), &GroupInfo, &CompileTask]
 						{
 							COOK_STAT(auto Timer = NiagaraSystemCookStats::UsageStats.TimeSyncWork(); Timer.TrackCyclesOnly(););
+							UE_LOG_CONTEXT("NiagaraSystem", Self->SystemInfo.GetSystemPath());
 
 							CompileTask.ProcessCompilation(&Self.Get(), GroupInfo);
 						}, *CompileTask.CompileResultsReadyEvent);
@@ -2732,6 +2738,7 @@ void FNiagaraSystemCompilationTask::IssueCompilationTasks()
 
 	Launch(UE_SOURCE_LOCATION, [Self = AsShared()]
 	{
+		UE_LOG_CONTEXT("NiagaraSystem", Self->SystemInfo.GetSystemPath());
 		Self->IssuePostResultsProcessedTasks();
 	}, CompilationTasks);
 }
diff --git a/Engine/Plugins/FX/Niagara/Source/NiagaraEditor/Private/NiagaraCompilationTasks.h b/Engine/Plugins/FX/Niagara/Source/NiagaraEditor/Private/NiagaraCompilationTasks.h
--- a/Engine/Plugins/FX/Niagara/Source/NiagaraEditor/Private/NiagaraCompilationTasks.h
+++ b/Engine/Plugins/FX/Niagara/Source/NiagaraEditor/Private/NiagaraCompilationTasks.h
@@ -120,6 +120,11 @@ struct FNiagaraSystemCompilationTask : public TSharedFromThis<FNiagaraSystemComp
 		bool bUseRapidIterationParams = false;
 		bool bDisableDebugSwitches = false;
 
+		FSoftObjectPath GetSystemPath() 
+		{ 
+			return FSoftObjectPath(FTopLevelAssetPath(FName(SystemPackageName), FName(SystemName)), FString());
+		}
+
 		const FEmitterInfo* EmitterInfoBySourceEmitter(int32 InSourceEmitterIndex) const;
 		FEmitterInfo* EmitterInfoBySourceEmitter(int32 InSourceEmitterIndex);
 	};
