- Fix ResetRequired not doing a comprehensive test resulting in occasionally sampling from the wrong SKM

#Synced-CL 48152930

diff --git a/Engine/Plugins/FX/Niagara/Source/Niagara/Classes/NiagaraDataInterfaceSkeletalMesh.h b/Engine/Plugins/FX/Niagara/Source/Niagara/Classes/NiagaraDataInterfaceSkeletalMesh.h
--- a/Engine/Plugins/FX/Niagara/Source/Niagara/Classes/NiagaraDataInterfaceSkeletalMesh.h
+++ b/Engine/Plugins/FX/Niagara/Source/Niagara/Classes/NiagaraDataInterfaceSkeletalMesh.h
@@ -81,6 +81,8 @@ struct FSkeletalMeshSkinningData
 	bool IsUsed()const;
 	void ForceDataRefresh();
 
+	TWeakObjectPtr<USkeletalMeshComponent> GetWeakMeshComponent() const { return MeshComp; }
+
 	bool Tick(float InDeltaSeconds, bool bRequirePreskin = true);
 
 	inline void EnterRead()
diff --git a/Engine/Plugins/FX/Niagara/Source/Niagara/Private/NiagaraDataInterfaceSkeletalMesh.cpp b/Engine/Plugins/FX/Niagara/Source/Niagara/Private/NiagaraDataInterfaceSkeletalMesh.cpp
--- a/Engine/Plugins/FX/Niagara/Source/Niagara/Private/NiagaraDataInterfaceSkeletalMesh.cpp
+++ b/Engine/Plugins/FX/Niagara/Source/Niagara/Private/NiagaraDataInterfaceSkeletalMesh.cpp
@@ -64,6 +64,14 @@ struct FNiagaraSkelMeshDIFunctionVersion
 
 namespace NDISkelMeshLocal
 {
+	static bool GUseNewResetRequired = true;
+	static FAutoConsoleVariableRef CUseNewResetRequired(
+		TEXT("fx.Niagara.Rollback.SkeletalMeshDI.UseNewResetRequired"),
+		GUseNewResetRequired,
+		TEXT("Rollback console variable to use the old reset required detection vs the new one which is more concise."),
+		ECVF_Default
+	);
+
 	static FName NAME_GetPreSkinnedLocalBounds("GetPreSkinnedLocalBounds");
 
 	static const TCHAR* CommonShaderFile = TEXT("/Plugin/FX/Niagara/Private/NiagaraDataInterfaceSkeletalMesh.ush");
@@ -2134,83 +2142,139 @@ bool FNDISkeletalMesh_InstanceData::Init(UNiagaraDataInterfaceSkeletalMesh* Inte
 
 bool FNDISkeletalMesh_InstanceData::ResetRequired(UNiagaraDataInterfaceSkeletalMesh* Interface, FNiagaraSystemInstance* SystemInstance) const
 {
-	// Reset if the scene component we've cached has been invalidated or unregistered
-	USceneComponent* Comp = SceneComponent.Get();
-	if (bComponentValid && (!Comp || !Comp->IsRegistered()))
+	if (NDISkelMeshLocal::GUseNewResetRequired)
 	{
-		return true;
-	}
+		// Reset if any parameters changed on the data interface
+		if (Interface->ChangeId != ChangeId)
+		{
+			return true;
+		}
 
-	// Reset if any mesh was bound on init, but is now invalidated
-	USkeletalMesh* SkelMesh = SkeletalMesh.Get();
-	if (bMeshValid && !SkelMesh)
-	{
-		return true;
-	}
+		// Reset if the scene component we've cached has been invalidated or unregistered
+		USceneComponent* ExistingSceneComponent = SceneComponent.Get();
+		if (bComponentValid && (!ExistingSceneComponent || !ExistingSceneComponent->IsRegistered()))
+		{
+			return true;
+		}
 
-	if (Interface->MeshUserParameter.Parameter.IsValid())
-	{
-		// Reset if the user object ptr has been changed to look at a new object
-		if (UserParamBinding.GetValue() != CachedUserParam)
+		// Reset if the skeletal mesh / component we are currently reading has changed
+		TWeakObjectPtr<USceneComponent> FoundSceneComponent;
+		USkeletalMeshComponent* FoundSkeletalMeshComponent = nullptr;
+		USkeletalMesh* FoundSkeletalMesh = Interface->GetSkeletalMesh(SystemInstance, SystemInstance->GetAttachComponent(), FoundSceneComponent, FoundSkeletalMeshComponent, const_cast<FNDISkeletalMesh_InstanceData*>(this));
+
+		if (SceneComponent != FoundSceneComponent)
 		{
 			return true;
 		}
-	}
-	else if (Interface->GetSourceComponent())
-	{
-		// Reset if the source component changed (or there wasn't one and now there is)
-		if (Interface->GetSourceComponent() != Comp)
+
+		if (SkeletalMesh != FoundSkeletalMesh)
 		{
+			// Skinning Data is unique per component
+			// Therefore if the component did not change but the SKM did we need to refresh the skinning data to ensure the bones are valid
+			if (SkinningData.SkinningData.IsValid())
+			{
+				SkinningData.SkinningData->ForceDataRefresh();
+			}
 			return true;
 		}
+
+		// Reset if the LOD we relied on was streamed out, or if the LOD we need could now be available.
+		if (SkeletalMesh != nullptr)
+		{
+			const FStreamableRenderResourceState& SRRState = SkeletalMesh->GetStreamableResourceState();
+			const int32 NumValidLODs = FMath::Min(SRRState.NumRequestedLODs, SRRState.NumResidentLODs);
+			if (NumValidLODs == 0)
+			{
+				return true;
+			}
+
+			const int32 CurrentFirstLOD = SRRState.LODCountToAssetFirstLODIdx(NumValidLODs);
+			if (CurrentFirstLOD > CachedLODIdx || (CurrentFirstLOD < CachedLODIdx && bResetOnLODStreamedIn))
+			{
+				return true;
+			}
+		}
 	}
-	else if (USceneComponent* AttachComponent = SystemInstance->GetAttachComponent())
+	else
 	{
-		// Reset if we detect any attachment change.
-		// TODO: This check is not really comprehensive. What we really need to know is if the mesh we cached comes from a skeletal mesh component in our
-		// attachment hierarchy, and if that hierarchy has changed in the chain between the system instance's attach component and the cached component,
-		// therefore potentially invalidating the cached component and mesh as our best choice.
-		if (CachedAttachParent != AttachComponent->GetAttachParent())
+		// Reset if the scene component we've cached has been invalidated or unregistered
+		USceneComponent* Comp = SceneComponent.Get();
+		if (bComponentValid && (!Comp || !Comp->IsRegistered()))
 		{
-			// The scene component our system instance was associated with has changed attachment, so we need to reinit
 			return true;
 		}
-	}
 
-	// Reset if the LOD we relied on was streamed out, or if the LOD we need could now be available.
-	if (SkelMesh != nullptr)
-	{
-		const FStreamableRenderResourceState& SRRState = SkelMesh->GetStreamableResourceState();
-		const int32 NumValidLODs = FMath::Min(SRRState.NumRequestedLODs, SRRState.NumResidentLODs);
-		if (NumValidLODs == 0)
+		// Reset if any mesh was bound on init, but is now invalidated
+		USkeletalMesh* SkelMesh = SkeletalMesh.Get();
+		if (bMeshValid && !SkelMesh)
 		{
 			return true;
 		}
 
-		const int32 CurrentFirstLOD = SRRState.LODCountToAssetFirstLODIdx(NumValidLODs);
-		if (CurrentFirstLOD > CachedLODIdx || (CurrentFirstLOD < CachedLODIdx && bResetOnLODStreamedIn))
+		if (Interface->MeshUserParameter.Parameter.IsValid())
 		{
-			return true;
+			// Reset if the user object ptr has been changed to look at a new object
+			if (UserParamBinding.GetValue() != CachedUserParam)
+			{
+				return true;
+			}
+		}
+		else if (Interface->GetSourceComponent())
+		{
+			// Reset if the source component changed (or there wasn't one and now there is)
+			if (Interface->GetSourceComponent() != Comp)
+			{
+				return true;
+			}
+		}
+		else if (USceneComponent* AttachComponent = SystemInstance->GetAttachComponent())
+		{
+			// Reset if we detect any attachment change.
+			// TODO: This check is not really comprehensive. What we really need to know is if the mesh we cached comes from a skeletal mesh component in our
+			// attachment hierarchy, and if that hierarchy has changed in the chain between the system instance's attach component and the cached component,
+			// therefore potentially invalidating the cached component and mesh as our best choice.
+			if (CachedAttachParent != AttachComponent->GetAttachParent())
+			{
+				// The scene component our system instance was associated with has changed attachment, so we need to reinit
+				return true;
+			}
 		}
-	}
 
-	// Reset if the skeletal mesh on the cached skeletal mesh component changed.
-	if (USkeletalMeshComponent* SkelComp = Cast<USkeletalMeshComponent>(Comp))
-	{
-		if (SkelComp->GetSkeletalMeshAsset() != SkelMesh)
+		// Reset if the LOD we relied on was streamed out, or if the LOD we need could now be available.
+		if (SkelMesh != nullptr)
 		{
-			if (SkinningData.SkinningData.IsValid())
+			const FStreamableRenderResourceState& SRRState = SkelMesh->GetStreamableResourceState();
+			const int32 NumValidLODs = FMath::Min(SRRState.NumRequestedLODs, SRRState.NumResidentLODs);
+			if (NumValidLODs == 0)
 			{
-				SkinningData.SkinningData.Get()->ForceDataRefresh();
+				return true;
+			}
+
+			const int32 CurrentFirstLOD = SRRState.LODCountToAssetFirstLODIdx(NumValidLODs);
+			if (CurrentFirstLOD > CachedLODIdx || (CurrentFirstLOD < CachedLODIdx && bResetOnLODStreamedIn))
+			{
+				return true;
 			}
-			return true;
 		}
-	}
 
-	// Reset if any parameters changed on the data interface
-	if (Interface->ChangeId != ChangeId)
-	{
-		return true;
+		// Reset if the skeletal mesh on the cached skeletal mesh component changed.
+		if (USkeletalMeshComponent* SkelComp = Cast<USkeletalMeshComponent>(Comp))
+		{
+			if (SkelComp->GetSkeletalMeshAsset() != SkelMesh)
+			{
+				if (SkinningData.SkinningData.IsValid())
+				{
+					SkinningData.SkinningData.Get()->ForceDataRefresh();
+				}
+				return true;
+			}
+		}
+
+		// Reset if any parameters changed on the data interface
+		if (Interface->ChangeId != ChangeId)
+		{
+			return true;
+		}
 	}
 
 	return false;
@@ -2602,8 +2666,9 @@ void UNiagaraDataInterfaceSkeletalMesh::DrawDebugHud(FNDIDrawDebugHudContext& De
 	}
 
 	USceneComponent* SceneComponent = InstanceData_GT->SceneComponent.Get();
+	AActor* SceneActor = SceneComponent ? SceneComponent->GetAttachmentRootActor() : nullptr;
 	USkeletalMesh* SkeletalMesh = InstanceData_GT->SkeletalMesh.Get();
-	DebugHudContext.GetOutputString().Appendf(TEXT("Skeleton(%s) SkelComp(%s)"), *GetNameSafe(SkeletalMesh), *GetNameSafe(SceneComponent));
+	DebugHudContext.GetOutputString().Appendf(TEXT("Skeleton(%s) SkelComp(%s:%s)"), *GetNameSafe(SkeletalMesh), SceneActor ? *SceneActor->GetActorNameOrLabel() : TEXT("None"), *GetNameSafe(SceneComponent));
 
 	if ( DebugHudContext.IsVerbose() && SceneComponent && SkeletalMesh)
 	{
