- Fix for source mode not copy and pasting into renderers

#tests Add a mesh renderer, set source mode to emitter and copy and paste into a standard emitter

#Synced-CL 48152930

diff --git a/Engine/Plugins/FX/Niagara/Source/Niagara/Private/NiagaraRendererProperties.cpp b/Engine/Plugins/FX/Niagara/Source/Niagara/Private/NiagaraRendererProperties.cpp
--- a/Engine/Plugins/FX/Niagara/Source/Niagara/Private/NiagaraRendererProperties.cpp
+++ b/Engine/Plugins/FX/Niagara/Source/Niagara/Private/NiagaraRendererProperties.cpp
@@ -961,8 +961,9 @@ FText UNiagaraRendererProperties::GetWidgetDisplayName() const
 
 bool UNiagaraRendererProperties::SupportsEmitterMode() const
 {
-	FVersionedNiagaraEmitter VersionedEmitter = GetOuterEmitter();
-	return VersionedEmitter.Emitter != nullptr;
+	// LWE don't support emitter mode
+	// Doing the reverse test does not work as the outer might be the clipboard
+	return GetTypedOuter<UNiagaraStatelessEmitter>() == nullptr;
 }
 
 ENiagaraRendererSourceDataMode UNiagaraRendererProperties::ValidateSourceMode(ENiagaraRendererSourceDataMode InSourceMode) const
