- Fix for MeshIndex not working with Emitter mode particles
- Fix for Mesh scale / rotation / offset not working with emitter mode and GPU scene

#tests Run the MeshRenderMeshIndex test / use Android Open GL preview (note that mesh index does not work on GPU in Android OGL due to GPU sorting being banned)

#virtualized
#Synced-CL 48152930

diff --git a/Engine/Plugins/FX/Niagara/Source/Niagara/Private/NiagaraRendererMeshes.cpp b/Engine/Plugins/FX/Niagara/Source/Niagara/Private/NiagaraRendererMeshes.cpp
--- a/Engine/Plugins/FX/Niagara/Source/Niagara/Private/NiagaraRendererMeshes.cpp
+++ b/Engine/Plugins/FX/Niagara/Source/Niagara/Private/NiagaraRendererMeshes.cpp
@@ -422,6 +422,7 @@ void FNiagaraRendererMeshes::PrepareParticleMeshRenderData(FRHICommandListBase&
 			{
 				ParticleMeshRenderData.bNeedsCull = false;
 				ParticleMeshRenderData.bAllowPerParticleMeshLODs = false;
+				ParticleMeshRenderData.bAllowFirstMeshOnly = true;
 			}
 			ParticleMeshRenderData.bNeedsSort &= ParticleMeshRenderData.bSortCullOnGpu;
 
@@ -432,6 +433,7 @@ void FNiagaraRendererMeshes::PrepareParticleMeshRenderData(FRHICommandListBase&
 				ParticleMeshRenderData.bNeedsSort = false;
 				ParticleMeshRenderData.bNeedsCull = false;
 				ParticleMeshRenderData.bAllowPerParticleMeshLODs = false;
+				ParticleMeshRenderData.bAllowFirstMeshOnly = true;
 			}
 		}
 		else
@@ -843,7 +845,7 @@ FNiagaraMeshCommonParameters FNiagaraRendererMeshes::CreateCommonShaderParams(co
 	Params.FacingMode 				= (uint32)FacingMode;	
 
 	Params.MeshScale				= MeshData.Scale;	
-	Params.MeshRotation				= FVector4f(MeshData.Rotation.X, MeshData.Rotation.Y, MeshData.Rotation.Z, MeshData.Rotation.W);
+	Params.MeshRotation				= MeshData.Rotation;
 	if (MeshData.PivotOffsetSpace == ENiagaraMeshPivotOffsetSpace::Mesh)
 	{
 		Params.MeshOffset 				= MeshData.PivotOffset;
@@ -873,7 +875,7 @@ FNiagaraMeshCommonParameters FNiagaraRendererMeshes::CreateCommonShaderParams(co
 	Params.DefaultPrevVelocity			= Params.DefaultVelocity;	
 	Params.DefaultScale					= FVector3f(1.0f, 1.0f, 1.0f);
 	Params.DefaultPrevScale 			= Params.DefaultScale;
-	Params.DefaultRotation 				= FVector4f(0.0f, 0.0f, 0.0f, 1.0f);
+	Params.DefaultRotation 				= FQuat4f(0.0f, 0.0f, 0.0f, 1.0f);
 	Params.DefaultPrevRotation 			= Params.DefaultRotation;
 	Params.DefaultCameraOffset 			= 0.0f;
 	Params.DefaultPrevCameraOffset 		= Params.DefaultCameraOffset;
@@ -938,7 +940,7 @@ FNiagaraMeshCommonParameters FNiagaraRendererMeshes::CreateCommonShaderParams(co
 					FMemory::Memcpy(&Params.DefaultScale, ParameterBoundData + VFBoundOffsetsInParamStore[i], sizeof(FVector3f));
 					break;
 				case ENiagaraMeshVFLayout::Type::Rotation:
-					FMemory::Memcpy(&Params.DefaultRotation, ParameterBoundData + VFBoundOffsetsInParamStore[i], sizeof(FVector4f));
+					FMemory::Memcpy(&Params.DefaultRotation, ParameterBoundData + VFBoundOffsetsInParamStore[i], sizeof(FQuat4f));
 					break;
 				case ENiagaraMeshVFLayout::Type::CameraOffset:
 					FMemory::Memcpy(&Params.DefaultCameraOffset, ParameterBoundData + VFBoundOffsetsInParamStore[i], sizeof(float));
@@ -950,7 +952,7 @@ FNiagaraMeshCommonParameters FNiagaraRendererMeshes::CreateCommonShaderParams(co
 					FMemory::Memcpy(&Params.DefaultPrevScale, ParameterBoundData + VFBoundOffsetsInParamStore[i], sizeof(FVector3f));
 					break;
 				case ENiagaraMeshVFLayout::Type::PrevRotation:
-					FMemory::Memcpy(&Params.DefaultPrevRotation, ParameterBoundData + VFBoundOffsetsInParamStore[i], sizeof(FVector4f));
+					FMemory::Memcpy(&Params.DefaultPrevRotation, ParameterBoundData + VFBoundOffsetsInParamStore[i], sizeof(FQuat4f));
 					break;
 				case ENiagaraMeshVFLayout::Type::PrevCameraOffset:
 					FMemory::Memcpy(&Params.DefaultPrevCameraOffset, ParameterBoundData + VFBoundOffsetsInParamStore[i], sizeof(float));
@@ -1204,8 +1206,11 @@ void FNiagaraRendererMeshes::SetupElementForGPUScene(
 		FEmitterSourceInstanceData& EmitterSourceData = GPUSceneRes.EmitterSourceData;
 		EmitterSourceData.CustomData = 0.0f;		
 
-		const FQuat Rot = FQuat(CommonParameters.DefaultRotation.X, CommonParameters.DefaultRotation.Y, CommonParameters.DefaultRotation.Z, CommonParameters.DefaultRotation.W).GetNormalized();
-		FMatrix LocalToPrimitive = FTransform(Rot, FVector(CommonParameters.DefaultPosition), FVector(CommonParameters.DefaultScale)).ToMatrixWithScale();
+		FMatrix LocalToPrimitive = FTransform(
+			FQuat(CommonParameters.DefaultRotation.GetNormalized() * CommonParameters.MeshRotation),
+			FVector(CommonParameters.DefaultPosition + CommonParameters.MeshOffset),
+			FVector(CommonParameters.DefaultScale * CommonParameters.MeshScale)
+		).ToMatrixWithScale();
 
 		if (!bLocalSpace)
 		{
@@ -1221,17 +1226,18 @@ void FNiagaraRendererMeshes::SetupElementForGPUScene(
 		
 		if (bNeedsPrevTransform)
 		{
-			FQuat PrevRot;
 			FMatrix PrevLocalToPrimitive;
 			if (CommonParameters.AccurateMotionVectors)
 			{
-				PrevRot = FQuat(CommonParameters.DefaultPrevRotation.X, CommonParameters.DefaultPrevRotation.Y, CommonParameters.DefaultPrevRotation.Z, CommonParameters.DefaultPrevRotation.W).GetNormalized();
-				PrevLocalToPrimitive = FTransform(PrevRot, FVector(CommonParameters.DefaultPrevPosition), FVector(CommonParameters.DefaultPrevScale)).ToMatrixWithScale();
+				PrevLocalToPrimitive = FTransform(
+					FQuat(CommonParameters.DefaultPrevRotation.GetNormalized() * CommonParameters.MeshRotation),
+					FVector(CommonParameters.DefaultPrevPosition + CommonParameters.MeshOffset),
+					FVector(CommonParameters.DefaultPrevScale * CommonParameters.MeshScale)
+				).ToMatrixWithScale();
 			}
 			else
 			{
-				PrevRot = FQuat(CommonParameters.DefaultRotation.X, CommonParameters.DefaultRotation.Y, CommonParameters.DefaultRotation.Z, CommonParameters.DefaultRotation.W).GetNormalized();
-				PrevLocalToPrimitive = FTransform(Rot, FVector(CommonParameters.DefaultPosition), FVector(CommonParameters.DefaultScale)).ToMatrixWithScale();
+				PrevLocalToPrimitive = LocalToPrimitive;
 			}
 
 			if (!bLocalSpace)
@@ -1458,13 +1464,11 @@ void FNiagaraRendererMeshes::GetDynamicMeshElements(const TArray<const FSceneVie
 				InitializeSortInfo(ParticleMeshRenderData, *SceneProxy, *View, ViewIndex, bIsInstancedStereo, SortInfo);
 			}
 
-			const bool bHasMeshIndexValidBinding = ParticleMeshRenderData.MeshIndexOffset != INDEX_NONE || EmitterModeMeshIndex != INDEX_NONE;
-
 			FNiagaraGpuComputeDispatchInterface* ComputeDispatchInterface = SceneProxy->GetComputeDispatchInterface();
 			for (int32 MeshIndex = 0; MeshIndex < Meshes.Num(); ++MeshIndex)
 			{
 				const FMeshData& MeshData = Meshes[MeshIndex];
-				if (MeshData.SourceMeshIndex > 0 && (!bHasMeshIndexValidBinding || (!ParticleMeshRenderData.bNeedsCull && !ParticleMeshRenderData.bUseGPUScene)))
+				if (MeshData.SourceMeshIndex > 0 && ParticleMeshRenderData.bAllowFirstMeshOnly)
 				{
 					break;
 				}
diff --git a/Engine/Plugins/FX/Niagara/Source/Niagara/Public/NiagaraRendererMeshes.h b/Engine/Plugins/FX/Niagara/Source/Niagara/Public/NiagaraRendererMeshes.h
--- a/Engine/Plugins/FX/Niagara/Source/Niagara/Public/NiagaraRendererMeshes.h
+++ b/Engine/Plugins/FX/Niagara/Source/Niagara/Public/NiagaraRendererMeshes.h
@@ -56,6 +56,7 @@ protected:
 		bool							bNeedsSort = false;
 		bool							bNeedsCull = false;
 		bool							bAllowPerParticleMeshLODs = false;
+		bool							bAllowFirstMeshOnly = false;
 		bool							bIsGpuLowLatencyTranslucency = false;
 
 		const FNiagaraRendererLayout*	RendererLayout = nullptr;
diff --git a/Engine/Plugins/FX/Niagara/Source/NiagaraVertexFactories/Public/NiagaraMeshVertexFactory.h b/Engine/Plugins/FX/Niagara/Source/NiagaraVertexFactories/Public/NiagaraMeshVertexFactory.h
--- a/Engine/Plugins/FX/Niagara/Source/NiagaraVertexFactories/Public/NiagaraMeshVertexFactory.h
+++ b/Engine/Plugins/FX/Niagara/Source/NiagaraVertexFactories/Public/NiagaraMeshVertexFactory.h
@@ -42,7 +42,7 @@ BEGIN_SHADER_PARAMETER_STRUCT(FNiagaraMeshCommonParameters, NIAGARAVERTEXFACTORI
 
 	SHADER_PARAMETER(FVector3f, MeshScale)
 	SHADER_PARAMETER(FVector3f, MeshOffset)
-	SHADER_PARAMETER(FVector4f, MeshRotation)
+	SHADER_PARAMETER(FQuat4f, MeshRotation)
 	SHADER_PARAMETER(int, bMeshOffsetIsWorldSpace)
 
 	SHADER_PARAMETER(uint32, bLockedAxisEnable)
@@ -61,12 +61,12 @@ BEGIN_SHADER_PARAMETER_STRUCT(FNiagaraMeshCommonParameters, NIAGARAVERTEXFACTORI
 	SHADER_PARAMETER(int, PrevCameraOffsetDataOffset)
 
 	SHADER_PARAMETER(FVector3f, DefaultScale)
-	SHADER_PARAMETER(FVector4f, DefaultRotation)
+	SHADER_PARAMETER(FQuat4f, DefaultRotation)
 	SHADER_PARAMETER(FVector3f, DefaultPosition)
 	SHADER_PARAMETER(FVector3f, DefaultVelocity)
 	SHADER_PARAMETER(float, DefaultCameraOffset)
 	SHADER_PARAMETER(FVector3f, DefaultPrevScale)
-	SHADER_PARAMETER(FVector4f, DefaultPrevRotation)
+	SHADER_PARAMETER(FQuat4f, DefaultPrevRotation)
 	SHADER_PARAMETER(FVector3f, DefaultPrevPosition)
 	SHADER_PARAMETER(FVector3f, DefaultPrevVelocity)
 	SHADER_PARAMETER(float, DefaultPrevCameraOffset)
