Move an ensure over to a more verbose warning in hopes of collecting some useful information to address the underlying issue with the niagara compilation and a hash collision for handling debug state

diff --git a/Engine/Plugins/FX/Niagara/Source/NiagaraEditor/Private/NiagaraTraversalStateContext.cpp b/Engine/Plugins/FX/Niagara/Source/NiagaraEditor/Private/NiagaraTraversalStateContext.cpp
--- a/Engine/Plugins/FX/Niagara/Source/NiagaraEditor/Private/NiagaraTraversalStateContext.cpp
+++ b/Engine/Plugins/FX/Niagara/Source/NiagaraEditor/Private/NiagaraTraversalStateContext.cpp
@@ -83,10 +83,38 @@ void FNiagaraTraversalStateContext::PushGraphInternal(const FNiagaraCompilationN
 				? ConstantResolver.GetDebugState()
 				: InnerFunctionNode->DebugState;
 
-			if (ensure(!FunctionDebugStateMap.Contains(InnerFunctionNodeHash)))
+			if (!FunctionDebugStateMap.Contains(InnerFunctionNodeHash))
 			{
 				FunctionDebugStateMap.Add(InnerFunctionNodeHash, CachedDebugState);
 			}
+			else
+			{
+				UE_LOG(LogNiagaraEditor, Warning, TEXT("FNiagaraTraversalStateContext::PushGraphInternal() generated a non-unique InnerFunctionNodeHash.\n" \
+					"\t[CurrentHash] %s\n" \
+					"\t[ExistingDebugState] %d\n" \
+					"\t[NewDebugState] %d\n" \
+					"\t[InnerFunctionNode] %s - %s"),
+					*InnerFunctionNodeHash.ToString(EGuidFormats::DigitsWithHyphens),
+					FunctionDebugStateMap[InnerFunctionNodeHash],
+					CachedDebugState,
+					*InnerFunctionNode->FunctionName, *InnerFunctionNode->FunctionScriptName
+				);
+
+				UE_LOG(LogNiagaraEditor, Warning, TEXT("FNiagaraTraversalStateContext - Stack"))
+				for (int32 StackIt = TraversalStack.Num() - 1; StackIt >= 0; --StackIt)
+				{
+					FString StackMessage = FString::Printf(TEXT("[%d] - %s, %s"),
+						StackIt,
+						*TraversalStack[StackIt].NodeGuid.ToString(EGuidFormats::DigitsWithHyphens),
+						*TraversalStack[StackIt].FullStackHash.ToString(EGuidFormats::DigitsWithHyphens));
+
+#if WITH_NIAGARA_TRAVERSAL_FRIENDLY_NAME
+					StackMessage.Append(TEXT(", "));
+					StackMessage.Append(TraversalStack[StackIt].FriendlyName);
+#endif
+					UE_LOG(LogNiagaraEditor, Warning, TEXT("%s"), *StackMessage);
+				}
+			}
 
 			if (CallingNode)
 			{
