- Fix crash when reassigning a dynamic input where the original not longer exists (i.e. is null)
- The checks for python available would access the nullptr

#Synced-CL 48229894

diff --git a/Engine/Plugins/FX/Niagara/Source/NiagaraEditor/Private/ViewModels/Stack/NiagaraStackFunctionInput.cpp b/Engine/Plugins/FX/Niagara/Source/NiagaraEditor/Private/ViewModels/Stack/NiagaraStackFunctionInput.cpp
--- a/Engine/Plugins/FX/Niagara/Source/NiagaraEditor/Private/ViewModels/Stack/NiagaraStackFunctionInput.cpp
+++ b/Engine/Plugins/FX/Niagara/Source/NiagaraEditor/Private/ViewModels/Stack/NiagaraStackFunctionInput.cpp
@@ -2884,7 +2884,7 @@ void UNiagaraStackFunctionInput::ReassignDynamicInputScript(UNiagaraScript* Dyna
 		UNiagaraScript* OldScript = InputValues.DynamicNode->FunctionScript;
 		FVersionedNiagaraScriptData* OldScriptData = InputValues.DynamicNode->GetScriptData();
 		FVersionedNiagaraScriptData* NewScriptData = DynamicInputScript->GetLatestScriptData();
-		if (NewScriptData->ConversionUtility || OldScriptData->bUsePythonScriptConversion)
+		if (NewScriptData->ConversionUtility || (OldScriptData && OldScriptData->bUsePythonScriptConversion))
 		{
 			OldClipboardContent = UNiagaraClipboardContent::Create();
 			Copy(OldClipboardContent);
@@ -2912,7 +2912,7 @@ void UNiagaraStackFunctionInput::ReassignDynamicInputScript(UNiagaraScript* Dyna
 		RefreshChildren();
 
 		bool bIsPythonAvailable = IPythonScriptPlugin::Get()->IsPythonAvailable();
-		if (OldScriptData->bUsePythonScriptConversion && bIsPythonAvailable && OldClipboardContent != nullptr)
+		if (OldScriptData && OldScriptData->bUsePythonScriptConversion && bIsPythonAvailable && OldClipboardContent != nullptr)
 		{
 			UNiagaraClipboardContent* NewClipboardContent = UNiagaraClipboardContent::Create();
 			UNiagaraClipboardContent* NewPythonContent = nullptr;
