Changed map key in component pool from object pointer to TObjectKey to prevent stale object pointers from invalidating the map

#tests Checked engine test for component pool still works
#Synced-CL 48244998

diff --git a/Engine/Plugins/FX/Niagara/Source/Niagara/Private/NiagaraComponentPool.cpp b/Engine/Plugins/FX/Niagara/Source/Niagara/Private/NiagaraComponentPool.cpp
--- a/Engine/Plugins/FX/Niagara/Source/Niagara/Private/NiagaraComponentPool.cpp
+++ b/Engine/Plugins/FX/Niagara/Source/Niagara/Private/NiagaraComponentPool.cpp
@@ -90,7 +90,7 @@ UNiagaraComponent* FNCPool::Acquire(UWorld* World, UNiagaraSystem* Template, ENC
 	while (FreeElements.Num() && !bForceNew)//Loop until we pop a valid free element or we're empty.
 	{
 		RetElem = FreeElements.Pop(EAllowShrinking::No);
-		if (!RetElem.Component || !IsValidChecked(RetElem.Component))
+		if (!RetElem.Component || !IsValidChecked(RetElem.Component.Get()))
 		{			
 			// Possible someone still has a reference to our NC and destroyed it while it was sat in the pool. Or possibly a teardown edgecase path that is GCing components from the pool be
 			UE_LOG(LogNiagara, Warning, TEXT("Pooled NC has been destroyed or is pending kill! Possibly via a DestroyComponent() call. You should not destroy pooled components manually. \nJust deactivate them and allow them to destroy themselves or be reclaimed by the pool. | NC: %p |\t System: %s"), RetElem.Component.Get(), *Template->GetFullName());
@@ -110,7 +110,7 @@ UNiagaraComponent* FNCPool::Acquire(UWorld* World, UNiagaraSystem* Template, ENC
 		AActor* OuterActor = World->GetWorldSettings();
 		UObject* OuterObject = OuterActor ? static_cast<UObject*>(OuterActor) : static_cast<UObject*>(World);
 
-		RetElem.Component = NewObject<UNiagaraComponent>(OuterObject);
+		RetElem.Component.Reset(NewObject<UNiagaraComponent>(OuterObject));
 		RetElem.Component->SetAutoDestroy(false);// we don't auto destroy, just periodically clear up the pool.
 		RetElem.Component->bAutoActivate = false;
 		RetElem.Component->SetVisibleInRayTracing(false);
@@ -119,7 +119,7 @@ UNiagaraComponent* FNCPool::Acquire(UWorld* World, UNiagaraSystem* Template, ENC
 
 	RetElem.Component->PoolingMethod = PoolingMethod;
 
-	return RetElem.Component;
+	return RetElem.Component.Get();
 }
 
 void FNCPool::Reclaim(UNiagaraComponent* Component, const double CurrentTimeSeconds)
@@ -174,7 +174,7 @@ bool FNCPool::RemoveComponent(UNiagaraComponent* Component)
 	int32 i = 0;
 	while (i < FreeElements.Num())
 	{
-		if (FreeElements[i].Component == Component)
+		if (FreeElements[i].Component.Get() == Component)
 		{
 			FreeElements.RemoveAtSwap(i, EAllowShrinking::No);
 			return true;
@@ -189,14 +189,12 @@ extern int32 GNigaraAllowPrimedPools;
 void FNCPool::KillUnusedComponents(double KillTime, UNiagaraSystem* Template)
 {
 	int32 i = 0;
-	check(Template);
-	int32 PrimedSize = GNigaraAllowPrimedPools != 0 ? Template->PoolPrimeSize : 0;
+	int32 PrimedSize = GNigaraAllowPrimedPools != 0 && Template ? Template->PoolPrimeSize : 0;
 	while (i < FreeElements.Num() && FreeElements.Num() > PrimedSize)//Don't free below the primed size
 	{
 		if (FreeElements[i].LastUsedTime < KillTime)
 		{
-			UNiagaraComponent* Component = FreeElements[i].Component;
-			if (Component)
+			if (UNiagaraComponent* Component = FreeElements[i].Component.Get())
 			{
 				Component->PoolingMethod = ENCPoolMethod::None; // Reset so we don't trigger warnings about destroying pooled NCs.
 				Component->DestroyComponent();
@@ -234,7 +232,7 @@ void UNiagaraComponentPool::Cleanup(UWorld* World)
 {
 	for (auto& Pool : WorldParticleSystemPools)
 	{
-		FNiagaraCrashReporterScope CRScope(Pool.Key);//In practice this may be null by now :(
+		FNiagaraCrashReporterScope CRScope(Pool.Key.ResolveObjectPtr());//In practice this may be null by now :(
 		Pool.Value.Cleanup();
 	}
 
@@ -262,8 +260,7 @@ void UNiagaraComponentPool::Cleanup(UWorld* World)
 
 void UNiagaraComponentPool::ClearPool(UNiagaraSystem* System)
 {
-	FNCPool* NCPool = WorldParticleSystemPools.Find(System);
-	if (NCPool)
+	if (FNCPool* NCPool = WorldParticleSystemPools.Find(System))
 	{
 		NCPool->Cleanup();
 	}
@@ -458,7 +455,8 @@ void UNiagaraComponentPool::ReclaimWorldParticleSystem(UNiagaraComponent* Compon
 			LastParticleSytemPoolCleanTime = CurrentTime;
 			for (auto& Pair : WorldParticleSystemPools)
 			{
-				Pair.Value.KillUnusedComponents(CurrentTime - GNiagaraSystemPoolKillUnusedTime, Pair.Key);
+				UNiagaraSystem* Template = Pair.Key.ResolveObjectPtr();
+				Pair.Value.KillUnusedComponents(CurrentTime - GNiagaraSystemPoolKillUnusedTime, Template);
 			}
 		}
 
@@ -521,7 +519,7 @@ void UNiagaraComponentPool::PooledComponentDestroyed(UNiagaraComponent* Componen
 		{
 			if (UNiagaraSystem* NiagaraSystem = Component->GetAsset())
 			{
-				if (FNCPool* NCPool = WorldParticleSystemPools.Find(Component->GetAsset()))
+				if (FNCPool* NCPool = WorldParticleSystemPools.Find(NiagaraSystem))
 				{
 					if (!NCPool->RemoveComponent(Component))
 					{
@@ -580,7 +578,7 @@ void UNiagaraComponentPool::Dump()
 	uint32 TotalMemUsage = 0;
 	for (auto& Pair : WorldParticleSystemPools)
 	{
-		UNiagaraSystem* System = Pair.Key;
+		UNiagaraSystem* System = Pair.Key.ResolveObjectPtr();
 		FNCPool& Pool = Pair.Value;		
 		uint32 FreeMemUsage = 0;
 		for (FNCPoolElement& Elem : Pool.FreeElements)
diff --git a/Engine/Plugins/FX/Niagara/Source/Niagara/Public/NiagaraComponentPool.h b/Engine/Plugins/FX/Niagara/Source/Niagara/Public/NiagaraComponentPool.h
--- a/Engine/Plugins/FX/Niagara/Source/Niagara/Public/NiagaraComponentPool.h
+++ b/Engine/Plugins/FX/Niagara/Source/Niagara/Public/NiagaraComponentPool.h
@@ -3,13 +3,12 @@
 #pragma once
 
 #include "CoreMinimal.h"
-#include "UObject/ObjectMacros.h"
+#include "UObject/ObjectKey.h"
 #include "NiagaraCommon.h"
+#include "NiagaraSystem.h"
+#include "NiagaraComponent.h"
 #include "NiagaraComponentPool.generated.h"
 
-class UNiagaraComponent;
-class UNiagaraSystem;
-
 #define ENABLE_NC_POOL_DEBUGGING (!UE_BUILD_SHIPPING)
 
 
@@ -18,9 +17,7 @@ struct FNCPoolElement
 {
 	GENERATED_BODY()
 
-	UPROPERTY(transient)
-	TObjectPtr<UNiagaraComponent> Component;
-
+	TStrongObjectPtr<UNiagaraComponent> Component;
 	double LastUsedTime;
 
 	FNCPoolElement()
@@ -71,15 +68,14 @@ class UNiagaraComponentPool : public UObject
 	GENERATED_UCLASS_BODY()
 
 private:
-	UPROPERTY()
-	TMap<TObjectPtr<UNiagaraSystem>, FNCPool> WorldParticleSystemPools;
+	TMap<TObjectKey<UNiagaraSystem>, FNCPool> WorldParticleSystemPools;
 
 	double LastParticleSytemPoolCleanTime;
 
 public:
 	static NIAGARA_API bool Enabled();
 
-	NIAGARA_API ~UNiagaraComponentPool();
+	NIAGARA_API virtual ~UNiagaraComponentPool() override;
 
 	NIAGARA_API void Cleanup(UWorld* World);
 
