- Add severity / fix for warm up time

#tests Enable warmup on a system with the rule active

#Synced-CL 48229894

diff --git a/Engine/Plugins/FX/Niagara/Source/NiagaraEditor/Private/NiagaraValidationRules.cpp b/Engine/Plugins/FX/Niagara/Source/NiagaraEditor/Private/NiagaraValidationRules.cpp
--- a/Engine/Plugins/FX/Niagara/Source/NiagaraEditor/Private/NiagaraValidationRules.cpp
+++ b/Engine/Plugins/FX/Niagara/Source/NiagaraEditor/Private/NiagaraValidationRules.cpp
@@ -385,12 +385,32 @@ bool NiagaraValidation::StructContainsUObjectProperty(UStruct* Struct)
 void UNiagaraValidationRule_NoWarmupTime::CheckValidity(const FNiagaraValidationContext& Context, TArray<FNiagaraValidationResult>& Results)  const
 {
 	UNiagaraSystem& System = Context.ViewModel->GetSystem();
-	if (System.NeedsWarmup())
+	if (System.NeedsWarmup() == false)
 	{
-		UNiagaraStackSystemPropertiesItem* SystemProperties = NiagaraValidation::GetStackEntry<UNiagaraStackSystemPropertiesItem>(Context.ViewModel->GetSystemStackViewModel());
-		FNiagaraValidationResult Result(ENiagaraValidationSeverity::Error, LOCTEXT("WarumupSummary", "Warmuptime > 0 is not allowed"), LOCTEXT("WarmupDescription", "Systems with the chosen effect type do not allow warmup time, as it costs too much performance.\nPlease set the warmup time to 0 in the system properties."), SystemProperties);
-		Results.Add(Result);
+		return;
 	}
+
+	FNiagaraValidationResult& Result = Results.Emplace_GetRef(
+		Severity,
+		LOCTEXT("WarumupSummary", "Warmup is enabled and not allowed."),
+		LOCTEXT("WarmupDescription", "Warmup is enabled on this system and it not allowed due to in game performance cost."),
+		NiagaraValidation::GetStackEntry<UNiagaraStackSystemPropertiesItem>(Context.ViewModel->GetSystemStackViewModel())
+	);
+
+	Result.Fixes.Emplace(
+		LOCTEXT("WarmupSetToZero", "Set Warmup Time to 0"),
+		FNiagaraValidationFixDelegate::CreateLambda(
+			[WeakSystem = MakeWeakObjectPtr(&System)]()
+			{
+				if (UNiagaraSystem* NiagaraSystem = WeakSystem.Get())
+				{
+					const FScopedTransaction Transaction(LOCTEXT("SetWarmupToZero", "Set Warmup Time to 0.0"));
+					NiagaraSystem->Modify();
+					NiagaraSystem->SetWarmupTime(0.0f);
+				}
+			}
+		)
+	);
 }
 
 void UNiagaraValidationRule_NoEvents::CheckValidity(const FNiagaraValidationContext& Context, TArray<FNiagaraValidationResult>& OutResults)  const
diff --git a/Engine/Plugins/FX/Niagara/Source/NiagaraEditor/Public/NiagaraValidationRules.h b/Engine/Plugins/FX/Niagara/Source/NiagaraEditor/Public/NiagaraValidationRules.h
--- a/Engine/Plugins/FX/Niagara/Source/NiagaraEditor/Public/NiagaraValidationRules.h
+++ b/Engine/Plugins/FX/Niagara/Source/NiagaraEditor/Public/NiagaraValidationRules.h
@@ -97,7 +97,11 @@ UCLASS(Category = "Validation", DisplayName = "No Warmup Time")
 class UNiagaraValidationRule_NoWarmupTime : public UNiagaraValidationRule
 {
 	GENERATED_BODY()
+
 public:
+	UPROPERTY(EditAnywhere, Category = Validation)
+	ENiagaraValidationSeverity Severity = ENiagaraValidationSeverity::Error;
+
 	virtual void CheckValidity(const FNiagaraValidationContext& Context, TArray<FNiagaraValidationResult>& OutResults) const override;
 };
 
