Adjust log warning to use LOGFMT to get additional information from the warning, like the context

#Synced-CL Unknown

diff --git a/Engine/Plugins/FX/Niagara/Source/NiagaraEditor/Private/NiagaraCompilationTasks.cpp b/Engine/Plugins/FX/Niagara/Source/NiagaraEditor/Private/NiagaraCompilationTasks.cpp
--- a/Engine/Plugins/FX/Niagara/Source/NiagaraEditor/Private/NiagaraCompilationTasks.cpp
+++ b/Engine/Plugins/FX/Niagara/Source/NiagaraEditor/Private/NiagaraCompilationTasks.cpp
@@ -2024,6 +2024,7 @@ void FNiagaraSystemCompilationTask::Precompile()
 	for (int32 EmitterIt = 0; EmitterIt < EmitterCount; ++EmitterIt)
 	{
 		const FEmitterInfo& EmitterInfo = SystemInfo.EmitterInfo[EmitterIt];
+		UE_LOG_CONTEXT("NiagaraEmitter", EmitterInfo.UniqueEmitterName);
 
 		if (EmitterInfo.Enabled) // Don't pull in the emitter if it isn't going to be used.
 		{
diff --git a/Engine/Plugins/FX/Niagara/Source/NiagaraEditor/Private/NiagaraTraversalStateContext.cpp b/Engine/Plugins/FX/Niagara/Source/NiagaraEditor/Private/NiagaraTraversalStateContext.cpp
--- a/Engine/Plugins/FX/Niagara/Source/NiagaraEditor/Private/NiagaraTraversalStateContext.cpp
+++ b/Engine/Plugins/FX/Niagara/Source/NiagaraEditor/Private/NiagaraTraversalStateContext.cpp
@@ -2,6 +2,7 @@
 
 #include "NiagaraTraversalStateContext.h"
 
+#include "Logging/StructuredLog.h"
 #include "NiagaraCompilationPrivate.h"
 #include "NiagaraConstants.h"
 #include "NiagaraEditorUtilities.h"
@@ -89,18 +90,21 @@ void FNiagaraTraversalStateContext::PushGraphInternal(const FNiagaraCompilationN
 			}
 			else
 			{
-				UE_LOG(LogNiagaraEditor, Warning, TEXT("FNiagaraTraversalStateContext::PushGraphInternal() generated a non-unique InnerFunctionNodeHash.\n" \
-					"\t[CurrentHash] %s\n" \
-					"\t[ExistingDebugState] %d\n" \
-					"\t[NewDebugState] %d\n" \
-					"\t[InnerFunctionNode] %s - %s"),
-					*InnerFunctionNodeHash.ToString(EGuidFormats::DigitsWithHyphens),
-					FunctionDebugStateMap[InnerFunctionNodeHash],
-					CachedDebugState,
-					*InnerFunctionNode->FunctionName, *InnerFunctionNode->FunctionScriptName
+				UE_LOGFMT(LogNiagaraEditor, Warning, "FNiagaraTraversalStateContext::PushGraphInternal() generated a non-unique InnerFunctionNodeHash.\n" \
+					"\t[SourceScript] {SourceScript}\n" \
+					"\t[CurrentHash] {CurrentHash}\n" \
+					"\t[ExistingDebugState] {ExistingDebugState}\n" \
+					"\t[NewDebugState] {NewDebugState}\n" \
+					"\t[InnerFunctionNode] {InnerFunctionNodeName} - {InnerFunctionNodeScriptName}",
+					("SourceScript", Graph->SourceScriptFullName),
+					("CurrentHash", InnerFunctionNodeHash.ToString(EGuidFormats::DigitsWithHyphens)),
+					("ExistingDebugState", static_cast<int32>(FunctionDebugStateMap[InnerFunctionNodeHash])),
+					("NewDebugState", static_cast<int32>(CachedDebugState)),
+					("InnerFunctionNodeName", InnerFunctionNode->FunctionName),
+					("InnerFunctionNodeScriptName", InnerFunctionNode->FunctionScriptName)
 				);
 
-				UE_LOG(LogNiagaraEditor, Warning, TEXT("FNiagaraTraversalStateContext - Stack"))
+				UE_LOG(LogNiagaraEditor, Warning, TEXT("FNiagaraTraversalStateContext - Stack"));
 				for (int32 StackIt = TraversalStack.Num() - 1; StackIt >= 0; --StackIt)
 				{
 					FString StackMessage = FString::Printf(TEXT("[%d] - %s, %s"),
@@ -144,20 +148,24 @@ void FNiagaraTraversalStateContext::PushGraphInternal(const FNiagaraCompilationN
 							else
 							{
 								// generate an error message to help track down the scenario where this has happened
-								UE_LOG(LogNiagaraEditor, Warning, TEXT("FNiagaraTraversalStateContext::PushGraphInternal() generated a non-unique function call.\n" \
-									"\t[ExistingDefaultValue] %s\n" \
-									"\t[NewDefaultValue] %s\n" \
-									"\t[ValuePin->PinName] %s\n" \
-									"\t[TaggedVariable] %s - %s\n" \
-									"\t[InnerFunctionNode] %s - %s"),
-									ExistingDefaultValue ? **ExistingDefaultValue : TEXT("<null>"),
-									*CallerInputPin->DefaultValue,
-									*ValuePin->PinName.ToString(),
-									*TaggedVariable.Key.GetName().ToString(), *TaggedVariable.Value.ToString(),
-									*InnerFunctionNode->FunctionName, *InnerFunctionNode->FunctionScriptName
+								UE_LOGFMT(LogNiagaraEditor, Warning, "FNiagaraTraversalStateContext::PushGraphInternal() generated a non-unique function call.\n" \
+									"\t[SourceScript] {SourceScript}\n" \
+									"\t[ExistingDefaultValue] {ExistingDefaultValue}\n" \
+									"\t[NewDefaultValue] {NewDefaultValue}\n" \
+									"\t[ValuePin->PinName] {PinName}\n" \
+									"\t[TaggedVariable] {TaggedVariableName} - {TaggedVariableValue}\n" \
+									"\t[InnerFunctionNode] {FunctionName} - {FunctionScriptName}",
+									("SourceScript", Graph->SourceScriptFullName),
+									("ExistingDefaultValue", ExistingDefaultValue ? *ExistingDefaultValue : TEXT("<null>")),
+									("NewDefaultValue", CallerInputPin->DefaultValue),
+									("PinName", ValuePin->PinName.ToString()),
+									("TaggedVariableName", TaggedVariable.Key.GetName().ToString()),
+									("TaggedVariableValue", TaggedVariable.Value.ToString()),
+									("FunctionName", InnerFunctionNode->FunctionName),
+									("FunctionScriptName", InnerFunctionNode->FunctionScriptName)
 									);
 								
-								UE_LOG(LogNiagaraEditor, Warning, TEXT("FNiagaraTraversalStateContext - Stack"))
+								UE_LOG(LogNiagaraEditor, Warning, TEXT("FNiagaraTraversalStateContext - Stack"));
 								for (int32 StackIt = TraversalStack.Num() - 1; StackIt >= 0; --StackIt)
 								{
 									FString StackMessage = FString::Printf(TEXT("[%d] - %s, %s"),
