- Make particle read / write fence are std::atomic to fix tsan issue

#rnx

#Synced-CL 48383985

diff --git a/Engine/Plugins/FX/Niagara/Source/Niagara/Classes/NiagaraComputeExecutionContext.h b/Engine/Plugins/FX/Niagara/Source/Niagara/Classes/NiagaraComputeExecutionContext.h
--- a/Engine/Plugins/FX/Niagara/Source/Niagara/Classes/NiagaraComputeExecutionContext.h
+++ b/Engine/Plugins/FX/Niagara/Source/Niagara/Classes/NiagaraComputeExecutionContext.h
@@ -159,8 +159,8 @@ public:
 
 private:
 	TSharedPtr<const FNiagaraScriptRuntimeData> ScriptRuntimeData;
-public:
 
+public:
 	// persistent layouts used to create the constant buffers for the compute sim shader
 	uint32 ExternalCBufferLayoutSize = 0;
 	TRefCountPtr<FNiagaraRHIUniformBufferLayout> ExternalCBufferLayout;
@@ -193,9 +193,11 @@ public:
 	// Do we have a reset pending, controlled by the game thread and passed to instance data
 	bool bResetPending_GT = true;
 
-	// Particle count read fence, used to allow us to snoop the count written by the render thread but also avoid racing on a reset value
-	uint32 ParticleCountReadFence = 1;
-	uint32 ParticleCountWriteFence = 0;
+	// Particle count read / write fence used when to avoid racing with the render thread
+	// Read count is used on the game thread / game concurrent threads and allows us to read the particle count when <= the write fence
+	// When the fence has not been passed we fallback to the potentially spawned particles
+	std::atomic<uint32> ParticleCountReadFence = 1;
+	std::atomic<uint32> ParticleCountWriteFence = 0;
 
 	// Render thread data
 	FNiagaraDataBuffer* GetPrevDataBuffer() { check(IsInRenderingThread() && (BufferSwapsThisFrame_RT > 0)); return DataBuffers_RT[(BufferSwapsThisFrame_RT & 1) ^ 1]; }
diff --git a/Engine/Plugins/FX/Niagara/Source/Niagara/Private/NiagaraEmitterInstance.cpp b/Engine/Plugins/FX/Niagara/Source/Niagara/Private/NiagaraEmitterInstance.cpp
--- a/Engine/Plugins/FX/Niagara/Source/Niagara/Private/NiagaraEmitterInstance.cpp
+++ b/Engine/Plugins/FX/Niagara/Source/Niagara/Private/NiagaraEmitterInstance.cpp
@@ -38,7 +38,7 @@ int32 FNiagaraEmitterInstance::GetNumParticles() const
 	// but hopefully adequate for a system script update.
 	if (GPUExecContext)
 	{
-		if (GPUExecContext->ParticleCountReadFence <= GPUExecContext->ParticleCountWriteFence)
+		if (GPUExecContext->ParticleCountReadFence.load(std::memory_order_relaxed) <= GPUExecContext->ParticleCountWriteFence.load(std::memory_order_relaxed))
 		{
 			// Fence has passed we read directly from the GPU Exec Context which will have the most up-to-date information
 			return GPUExecContext->CurrentNumInstances_RT;
diff --git a/Engine/Plugins/FX/Niagara/Source/Niagara/Private/NiagaraGPUSystemTick.cpp b/Engine/Plugins/FX/Niagara/Source/Niagara/Private/NiagaraGPUSystemTick.cpp
--- a/Engine/Plugins/FX/Niagara/Source/Niagara/Private/NiagaraGPUSystemTick.cpp
+++ b/Engine/Plugins/FX/Niagara/Source/Niagara/Private/NiagaraGPUSystemTick.cpp
@@ -180,9 +180,9 @@ void FNiagaraGPUSystemTick::Init(FNiagaraSystemInstance* InSystemInstance)
 				InstanceData->bResetData = GPUContext->bResetPending_GT;
 				GPUContext->bResetPending_GT = false;
 
-				++GPUContext->ParticleCountReadFence;
+				GPUContext->ParticleCountReadFence.fetch_add(1, std::memory_order_relaxed);
 			}
-			InstanceData->ParticleCountFence = GPUContext->ParticleCountReadFence;
+			InstanceData->ParticleCountFence = GPUContext->ParticleCountReadFence.load(std::memory_order_relaxed);
 
 			InstanceData->EmitterParamData = ParamDataBufferPtr;
 			ParamDataBufferPtr += InterpFactor * sizeof(FNiagaraEmitterParameters);
@@ -332,7 +332,7 @@ void FNiagaraGPUSystemTick::Destroy()
 {
 	for ( FNiagaraComputeInstanceData& Instance : GetInstances() )
 	{
-		Instance.Context->ParticleCountWriteFence = Instance.ParticleCountFence;
+		Instance.Context->ParticleCountWriteFence.store(Instance.ParticleCountFence, std::memory_order_relaxed);
 		Instance.~FNiagaraComputeInstanceData();
 	}
 
diff --git a/Engine/Plugins/FX/Niagara/Source/Niagara/Private/NiagaraSimCacheHelper.h b/Engine/Plugins/FX/Niagara/Source/Niagara/Private/NiagaraSimCacheHelper.h
--- a/Engine/Plugins/FX/Niagara/Source/Niagara/Private/NiagaraSimCacheHelper.h
+++ b/Engine/Plugins/FX/Niagara/Source/Niagara/Private/NiagaraSimCacheHelper.h
@@ -870,11 +870,11 @@ struct FNiagaraSimCacheHelper
 
 		FNiagaraGpuComputeDispatchInterface* DispathInterface = EmitterInstance.GetParentSystemInstance()->GetComputeDispatchInterface();
 		FNiagaraComputeExecutionContext* GPUExecContext = EmitterInstance.GetGPUContext();
-		++GPUExecContext->ParticleCountReadFence;
+		GPUExecContext->ParticleCountReadFence.fetch_add(1, std::memory_order_relaxed);
 		GPUExecContext->bResetPending_GT = false;
 
 		ENQUEUE_RENDER_COMMAND(NiagaraSimCacheGpuReadFrame)(
-			[DispathInterface, GPUExecContext, ParticleCountFence=GPUExecContext->ParticleCountReadFence, FrameFraction=InFrameFraction, FrameDeltaSeconds=InFrameDeltaSeconds, SimDeltaSeconds=InSimDeltaSeconds, RebaseATransform=InRebaseATransform, RebaseBTransform=InRebaseBTransform, CacheLayout=&InCacheLayout, CacheBufferA=&InCacheBufferA, CacheBufferB=&InCacheBufferB, DataSet=&InDataSet, PendingCommandsCounter=&InPendingCommandsCounter](FRHICommandListImmediate& RHICmdList)
+			[DispathInterface, GPUExecContext, ParticleCountFence=GPUExecContext->ParticleCountReadFence.load(std::memory_order_relaxed), FrameFraction = InFrameFraction, FrameDeltaSeconds = InFrameDeltaSeconds, SimDeltaSeconds = InSimDeltaSeconds, RebaseATransform = InRebaseATransform, RebaseBTransform = InRebaseBTransform, CacheLayout = &InCacheLayout, CacheBufferA = &InCacheBufferA, CacheBufferB = &InCacheBufferB, DataSet = &InDataSet, PendingCommandsCounter = &InPendingCommandsCounter](FRHICommandListImmediate& RHICmdList)
 			{
 				const uint32 NumInstances = FMath::Min(CacheBufferA->NumInstances, DataSet->GetMaxAllocationCount());
 
@@ -890,7 +890,7 @@ struct FNiagaraSimCacheHelper
 					DataBuffer.SetGPUInstanceCountBufferOffset(GPUExecContext->CountOffset_RT);
 
 					GPUExecContext->CurrentNumInstances_RT = NumInstances;
-					GPUExecContext->ParticleCountWriteFence = ParticleCountFence;
+					GPUExecContext->ParticleCountWriteFence.store(ParticleCountFence, std::memory_order_relaxed);
 
 					const FRWBuffer& CountBuffer = CountManager.GetInstanceCountBuffer();
 					const TPair<uint32, int32> DataToSet(GPUExecContext->CountOffset_RT, NumInstances);
