Moved niagara toolkit names into publically accessible locations, and renamed so that extenders can cleanly identify the niagara mode (they are registered across asset editors so "Default" doesn't cut it)
Moved include of private header within a public header to the .cpp file

#Synced-CL 48403841

diff --git a/Engine/Plugins/FX/Niagara/Source/NiagaraEditor/Private/Toolkits/NiagaraSystemToolkit.cpp b/Engine/Plugins/FX/Niagara/Source/NiagaraEditor/Private/Toolkits/NiagaraSystemToolkit.cpp
--- a/Engine/Plugins/FX/Niagara/Source/NiagaraEditor/Private/Toolkits/NiagaraSystemToolkit.cpp
+++ b/Engine/Plugins/FX/Niagara/Source/NiagaraEditor/Private/Toolkits/NiagaraSystemToolkit.cpp
@@ -71,9 +71,6 @@
 
 DECLARE_CYCLE_STAT(TEXT("Niagara - SystemToolkit - OnApply"), STAT_NiagaraEditor_SystemToolkit_OnApply, STATGROUP_NiagaraEditor);
 
-const FName FNiagaraSystemToolkit::DefaultModeName(TEXT("Default"));
-const FName FNiagaraSystemToolkit::ScalabilityModeName(TEXT("Scalability"));
-
 static int32 GbLogNiagaraSystemChanges = 0;
 static FAutoConsoleVariableRef CVarLogNiagaraSystemChanges(
 	TEXT("fx.LogNiagaraSystemChanges"),
@@ -316,8 +313,8 @@ void FNiagaraSystemToolkit::InitializeInternal(const EToolkitMode::Type Mode, co
 	
 	SystemViewModel->GetDocumentViewModel()->InitializePostTabManager(SharedThis(this));
 
-	AddApplicationMode(DefaultModeName, MakeShared<FNiagaraSystemToolkitMode_Default>(SharedThis(this)));
-	AddApplicationMode(ScalabilityModeName, MakeShared<FNiagaraSystemToolkitMode_Scalability>(SharedThis(this)));
+	AddApplicationMode(FNiagaraSystemToolkitModeBase::DefaultModeName, MakeShared<FNiagaraSystemToolkitMode_Default>(SharedThis(this)));
+	AddApplicationMode(FNiagaraSystemToolkitModeBase::ScalabilityModeName, MakeShared<FNiagaraSystemToolkitMode_Scalability>(SharedThis(this)));
 
 	// set up the versioning widget
 	VersionMetadata = NewObject<UNiagaraVersionMetaData>(ToolkitObject, "VersionMetadata", RF_Transient);
@@ -334,7 +331,7 @@ void FNiagaraSystemToolkit::InitializeInternal(const EToolkitMode::Type Mode, co
 			}
 		});
 	
-	SetCurrentMode(DefaultModeName);
+	SetCurrentMode(FNiagaraSystemToolkitModeBase::DefaultModeName);
 	
 	FNiagaraEditorModule& NiagaraEditorModule = FModuleManager::LoadModuleChecked<FNiagaraEditorModule>("NiagaraEditor");
 	AddMenuExtender(NiagaraEditorModule.GetMenuExtensibilityManager()->GetAllExtenders(GetToolkitCommands(), GetEditingObjects()));
diff --git a/Engine/Plugins/FX/Niagara/Source/NiagaraEditor/Private/Toolkits/NiagaraSystemToolkit.h b/Engine/Plugins/FX/Niagara/Source/NiagaraEditor/Private/Toolkits/NiagaraSystemToolkit.h
--- a/Engine/Plugins/FX/Niagara/Source/NiagaraEditor/Private/Toolkits/NiagaraSystemToolkit.h
+++ b/Engine/Plugins/FX/Niagara/Source/NiagaraEditor/Private/Toolkits/NiagaraSystemToolkit.h
@@ -283,8 +283,4 @@ private:
 	static IConsoleVariable* GpuStatEnabledVar;
 	
 	friend class FNiagaraSystemToolkitModeBase;
-
-public:
-	static const FName DefaultModeName;
-	static const FName ScalabilityModeName;
 };
diff --git a/Engine/Plugins/FX/Niagara/Source/NiagaraEditor/Private/Toolkits/SystemToolkitModes/NiagaraSystemToolkitModeBase.cpp b/Engine/Plugins/FX/Niagara/Source/NiagaraEditor/Private/Toolkits/SystemToolkitModes/NiagaraSystemToolkitModeBase.cpp
--- a/Engine/Plugins/FX/Niagara/Source/NiagaraEditor/Private/Toolkits/SystemToolkitModes/NiagaraSystemToolkitModeBase.cpp
+++ b/Engine/Plugins/FX/Niagara/Source/NiagaraEditor/Private/Toolkits/SystemToolkitModes/NiagaraSystemToolkitModeBase.cpp
@@ -43,6 +43,7 @@
 #include "Widgets/SNiagaraSystemUserParameters.h"
 #include "Customizations/NiagaraComponentDetails.h"
 #include "Styling/StyleColors.h"
+#include "ViewModels/NiagaraBakerViewModel.h"
 #include "ViewModels/NiagaraScratchPadViewModel.h"
 #include "ViewModels/NiagaraSystemEditorDocumentsViewModel.h"
 #include "ViewModels/HierarchyEditor/NiagaraSummaryViewViewModel.h"
@@ -51,6 +52,9 @@
 
 #define LOCTEXT_NAMESPACE "NiagaraSystemToolkitModeBase"
 
+const FName FNiagaraSystemToolkitModeBase::DefaultModeName(TEXT("NiagaraSystemToolkitMode_Default"));
+const FName FNiagaraSystemToolkitModeBase::ScalabilityModeName(TEXT("NiagaraSystemToolkitMode_Scalability"));
+
 const FName FNiagaraSystemToolkitModeBase::ViewportTabID(TEXT("NiagaraSystemEditor_Viewport"));
 const FName FNiagaraSystemToolkitModeBase::CurveEditorTabID(TEXT("NiagaraSystemEditor_CurveEditor"));
 const FName FNiagaraSystemToolkitModeBase::SequencerTabID(TEXT("NiagaraSystemEditor_Sequencer"));
diff --git a/Engine/Plugins/FX/Niagara/Source/NiagaraEditor/Private/Toolkits/SystemToolkitModes/NiagaraSystemToolkitMode_Default.cpp b/Engine/Plugins/FX/Niagara/Source/NiagaraEditor/Private/Toolkits/SystemToolkitModes/NiagaraSystemToolkitMode_Default.cpp
--- a/Engine/Plugins/FX/Niagara/Source/NiagaraEditor/Private/Toolkits/SystemToolkitModes/NiagaraSystemToolkitMode_Default.cpp
+++ b/Engine/Plugins/FX/Niagara/Source/NiagaraEditor/Private/Toolkits/SystemToolkitModes/NiagaraSystemToolkitMode_Default.cpp
@@ -15,7 +15,7 @@
 
 #define LOCTEXT_NAMESPACE "NiagaraSystemToolkitMode_Default"
 
-FNiagaraSystemToolkitMode_Default::FNiagaraSystemToolkitMode_Default(TWeakPtr<FNiagaraSystemToolkit> InSystemToolkit) : FNiagaraSystemToolkitModeBase(FNiagaraSystemToolkit::DefaultModeName, InSystemToolkit)
+FNiagaraSystemToolkitMode_Default::FNiagaraSystemToolkitMode_Default(TWeakPtr<FNiagaraSystemToolkit> InSystemToolkit) : FNiagaraSystemToolkitModeBase(FNiagaraSystemToolkitModeBase::DefaultModeName, InSystemToolkit)
 {
 	TabLayout = FTabManager::NewLayout("Standalone_Niagara_System_Layout_v29")
 		->AddArea
@@ -246,10 +246,10 @@ void FNiagaraSystemToolkitMode_Default::ExtendToolbar()
 			
 			ToolbarBuilder.BeginSection("Scalability");
 			{
-				FUIAction ScalabilityToggleAction(FExecuteAction::CreateRaw(Toolkit, &FNiagaraSystemToolkit::SetCurrentMode, FNiagaraSystemToolkit::ScalabilityModeName));
+				FUIAction ScalabilityToggleAction(FExecuteAction::CreateRaw(Toolkit, &FNiagaraSystemToolkit::SetCurrentMode, FNiagaraSystemToolkitModeBase::ScalabilityModeName));
 				ScalabilityToggleAction.GetActionCheckState = FGetActionCheckState::CreateLambda([Toolkit]()
 				{
-					return Toolkit->GetCurrentMode() == FNiagaraSystemToolkit::ScalabilityModeName ? ECheckBoxState::Checked : ECheckBoxState::Unchecked;
+					return Toolkit->GetCurrentMode() == FNiagaraSystemToolkitModeBase::ScalabilityModeName ? ECheckBoxState::Checked : ECheckBoxState::Unchecked;
 				});
 				
 				ToolbarBuilder.AddToolBarButton(
diff --git a/Engine/Plugins/FX/Niagara/Source/NiagaraEditor/Private/Toolkits/SystemToolkitModes/NiagaraSystemToolkitMode_Scalability.cpp b/Engine/Plugins/FX/Niagara/Source/NiagaraEditor/Private/Toolkits/SystemToolkitModes/NiagaraSystemToolkitMode_Scalability.cpp
--- a/Engine/Plugins/FX/Niagara/Source/NiagaraEditor/Private/Toolkits/SystemToolkitModes/NiagaraSystemToolkitMode_Scalability.cpp
+++ b/Engine/Plugins/FX/Niagara/Source/NiagaraEditor/Private/Toolkits/SystemToolkitModes/NiagaraSystemToolkitMode_Scalability.cpp
@@ -6,6 +6,7 @@
 #include "NiagaraEditorModule.h"
 #include "NiagaraSettings.h"
 #include "Toolkits/NiagaraSystemToolkit.h"
+#include "NiagaraSystem.h"
 #include "ViewModels/NiagaraSystemViewModel.h"
 #include "ViewModels/Stack/NiagaraStackRendererItem.h"
 #include "Widgets/Scalability/SNiagaraSystemResolvedScalabilitySettings.h"
@@ -21,7 +22,7 @@ const FName FNiagaraSystemToolkitMode_Scalability::ScalabilityViewModeTabID(TEXT
 const FName FNiagaraSystemToolkitMode_Scalability::ResolvedScalabilityTabID(TEXT("NiagaraSystemEditor_ScalabilityMode_ResolvedScalability"));
 const FName FNiagaraSystemToolkitMode_Scalability::EffectTypeTabID(TEXT("NiagaraSystemEditor_ScalabilityMode_EffectType"));
 
-FNiagaraSystemToolkitMode_Scalability::FNiagaraSystemToolkitMode_Scalability(TWeakPtr<FNiagaraSystemToolkit> InSystemToolkit) : FNiagaraSystemToolkitModeBase(FNiagaraSystemToolkit::ScalabilityModeName, InSystemToolkit)
+FNiagaraSystemToolkitMode_Scalability::FNiagaraSystemToolkitMode_Scalability(TWeakPtr<FNiagaraSystemToolkit> InSystemToolkit) : FNiagaraSystemToolkitModeBase(FNiagaraSystemToolkitModeBase::ScalabilityModeName, InSystemToolkit)
 {
 	TabLayout = ChooseTabLayout();
 
@@ -262,10 +263,10 @@ void FNiagaraSystemToolkitMode_Scalability::ExtendToolbar()
 		{
 			ToolbarBuilder.BeginSection("Scalability");
 			{
-				FUIAction ScalabilityToggleAction(FExecuteAction::CreateRaw(Toolkit, &FNiagaraSystemToolkit::SetCurrentMode, FNiagaraSystemToolkit::DefaultModeName));
+				FUIAction ScalabilityToggleAction(FExecuteAction::CreateRaw(Toolkit, &FNiagaraSystemToolkit::SetCurrentMode, FNiagaraSystemToolkitModeBase::DefaultModeName));
 				ScalabilityToggleAction.GetActionCheckState = FGetActionCheckState::CreateLambda([Toolkit]()
 				{
-					return Toolkit->GetCurrentMode() == FNiagaraSystemToolkit::ScalabilityModeName ? ECheckBoxState::Checked : ECheckBoxState::Unchecked;
+					return Toolkit->GetCurrentMode() == FNiagaraSystemToolkitModeBase::ScalabilityModeName ? ECheckBoxState::Checked : ECheckBoxState::Unchecked;
 				});
 						
 				ToolbarBuilder.AddToolBarButton(
diff --git a/Engine/Plugins/FX/Niagara/Source/NiagaraEditor/Private/Toolkits/SystemToolkitModes/NiagaraSystemToolkitMode_Scalability.h b/Engine/Plugins/FX/Niagara/Source/NiagaraEditor/Private/Toolkits/SystemToolkitModes/NiagaraSystemToolkitMode_Scalability.h
--- a/Engine/Plugins/FX/Niagara/Source/NiagaraEditor/Private/Toolkits/SystemToolkitModes/NiagaraSystemToolkitMode_Scalability.h
+++ b/Engine/Plugins/FX/Niagara/Source/NiagaraEditor/Private/Toolkits/SystemToolkitModes/NiagaraSystemToolkitMode_Scalability.h
@@ -5,6 +5,7 @@
 #include "Toolkits/SystemToolkitModes/NiagaraSystemToolkitModeBase.h"
 #include "Customizations/NiagaraPlatformSetCustomization.h"
 
+class UNiagaraSystem;
 class FNiagaraSystemToolkit;
 
 class FNiagaraSystemToolkitMode_Scalability : public FNiagaraSystemToolkitModeBase
diff --git a/Engine/Plugins/FX/Niagara/Source/NiagaraEditor/Private/ViewModels/NiagaraSystemScalabilityViewModel.cpp b/Engine/Plugins/FX/Niagara/Source/NiagaraEditor/Private/ViewModels/NiagaraSystemScalabilityViewModel.cpp
--- a/Engine/Plugins/FX/Niagara/Source/NiagaraEditor/Private/ViewModels/NiagaraSystemScalabilityViewModel.cpp
+++ b/Engine/Plugins/FX/Niagara/Source/NiagaraEditor/Private/ViewModels/NiagaraSystemScalabilityViewModel.cpp
@@ -9,6 +9,7 @@
 #include "DeviceProfiles/DeviceProfile.h"
 #include "DeviceProfiles/DeviceProfileManager.h"
 #include "Modules/ModuleManager.h"
+#include "Toolkits/SystemToolkitModes/NiagaraSystemToolkitModeBase.h"
 #include "ViewModels/NiagaraEmitterHandleViewModel.h"
 
 #include UE_INLINE_GENERATED_CPP_BY_NAME(NiagaraSystemScalabilityViewModel)
@@ -34,7 +35,7 @@ bool UNiagaraSystemScalabilityViewModel::IsValid() const
 
 bool UNiagaraSystemScalabilityViewModel::IsActive() const
 {
-	return GetSystemViewModel().Pin()->GetWorkflowMode() == FNiagaraSystemToolkit::ScalabilityModeName;
+	return GetSystemViewModel().Pin()->GetWorkflowMode() == FNiagaraSystemToolkitModeBase::ScalabilityModeName;
 }
 
 void UNiagaraSystemScalabilityViewModel::UpdatePreviewDeviceProfile(UDeviceProfile* DeviceProfile)
diff --git a/Engine/Plugins/FX/Niagara/Source/NiagaraEditor/Public/Toolkits/SystemToolkitModes/NiagaraSystemToolkitModeBase.h b/Engine/Plugins/FX/Niagara/Source/NiagaraEditor/Public/Toolkits/SystemToolkitModes/NiagaraSystemToolkitModeBase.h
--- a/Engine/Plugins/FX/Niagara/Source/NiagaraEditor/Public/Toolkits/SystemToolkitModes/NiagaraSystemToolkitModeBase.h
+++ b/Engine/Plugins/FX/Niagara/Source/NiagaraEditor/Public/Toolkits/SystemToolkitModes/NiagaraSystemToolkitModeBase.h
@@ -2,7 +2,6 @@
 
 #pragma once
 
-#include "ViewModels/NiagaraBakerViewModel.h"
 #include "WorkflowOrientedApp/ApplicationMode.h"
 #include "Widgets/Docking/SDockTab.h"
 #include "GraphEditAction.h"
@@ -85,6 +84,9 @@ protected:
 	TWeakPtr<class FNiagaraScratchPadScriptViewModel> LastActiveScratchPadViewModel;
 
 public:
+	NIAGARAEDITOR_API static const FName DefaultModeName;
+	NIAGARAEDITOR_API static const FName ScalabilityModeName;
+	
 	static const FName ViewportTabID;
 	static const FName CurveEditorTabID;
 	static const FName SequencerTabID;
diff --git a/Engine/Plugins/FX/Niagara/Source/NiagaraEditorWidgets/Private/SNiagaraOverviewStackNode.cpp b/Engine/Plugins/FX/Niagara/Source/NiagaraEditorWidgets/Private/SNiagaraOverviewStackNode.cpp
--- a/Engine/Plugins/FX/Niagara/Source/NiagaraEditorWidgets/Private/SNiagaraOverviewStackNode.cpp
+++ b/Engine/Plugins/FX/Niagara/Source/NiagaraEditorWidgets/Private/SNiagaraOverviewStackNode.cpp
@@ -49,6 +49,7 @@
 #include "AssetThumbnail.h"
 #include "SNiagaraEmitterRendererThumbnails.h"
 #include "Framework/Application/SlateApplication.h"
+#include "Toolkits/SystemToolkitModes/NiagaraSystemToolkitModeBase.h"
 
 #define LOCTEXT_NAMESPACE "NiagaraOverviewStackNode"
 
@@ -246,7 +247,7 @@ ECheckBoxState SNiagaraOverviewStackNode::IsScalabilityModeActive() const
 	
 	if(NiagaraSystemViewModel.IsValid())
 	{
-		return NiagaraSystemViewModel->GetWorkflowMode().IsEqual(FName("Scalability")) ? ECheckBoxState::Checked : ECheckBoxState::Unchecked;
+		return NiagaraSystemViewModel->GetWorkflowMode().IsEqual(FNiagaraSystemToolkitModeBase::ScalabilityModeName) ? ECheckBoxState::Checked : ECheckBoxState::Unchecked;
 	}
 
 	return ECheckBoxState::Unchecked;
@@ -263,13 +264,13 @@ void SNiagaraOverviewStackNode::OnScalabilityModeStateChanged(ECheckBoxState Che
 
 	if(NiagaraSystemViewModel.IsValid())
 	{
-		if(NiagaraSystemViewModel->GetWorkflowMode().IsEqual(FName("Scalability")))
+		if(NiagaraSystemViewModel->GetWorkflowMode().IsEqual(FNiagaraSystemToolkitModeBase::ScalabilityModeName))
 		{
-			NiagaraSystemViewModel->SetWorkflowMode(FName("Default"));
+			NiagaraSystemViewModel->SetWorkflowMode(FNiagaraSystemToolkitModeBase::DefaultModeName);
 		}
 		else
 		{
-			NiagaraSystemViewModel->SetWorkflowMode(FName("Scalability"));
+			NiagaraSystemViewModel->SetWorkflowMode(FNiagaraSystemToolkitModeBase::ScalabilityModeName);
 		}
 	}
 }
