Essential Automation coverage for Niagara features: Module Script , Script Versioning, Parameters Module, Override Parameters, and Node Graph

#rnx
#virtualized
#Synced-CL 47108571

diff --git a/Engine/Plugins/FX/Niagara/Source/NiagaraEditor/Public/ViewModels/Stack/NiagaraStackGraphUtilities.h b/Engine/Plugins/FX/Niagara/Source/NiagaraEditor/Public/ViewModels/Stack/NiagaraStackGraphUtilities.h
--- a/Engine/Plugins/FX/Niagara/Source/NiagaraEditor/Public/ViewModels/Stack/NiagaraStackGraphUtilities.h
+++ b/Engine/Plugins/FX/Niagara/Source/NiagaraEditor/Public/ViewModels/Stack/NiagaraStackGraphUtilities.h
@@ -357,8 +357,8 @@ namespace FNiagaraStackGraphUtilities
 		FName NewAssignmentTargetName);
 
 	/** Helper to add a new pin to a ParameterMapBaseNode, and conditionally add a new parameter to the underlying UNiagaraGraph if it does not yet exist. */
-	void AddNewVariableToParameterMapNode(UNiagaraNodeParameterMapBase* MapBaseNode, bool bCreateInputPin, const FNiagaraVariable& NewVariable);
-	void AddNewVariableToParameterMapNode(UNiagaraNodeParameterMapBase* MapBaseNode, bool bCreateInputPin, const UNiagaraScriptVariable* NewScriptVar);
+	NIAGARAEDITOR_API void AddNewVariableToParameterMapNode(UNiagaraNodeParameterMapBase* MapBaseNode, bool bCreateInputPin, const FNiagaraVariable& NewVariable);
+	NIAGARAEDITOR_API void AddNewVariableToParameterMapNode(UNiagaraNodeParameterMapBase* MapBaseNode, bool bCreateInputPin, const UNiagaraScriptVariable* NewScriptVar);
 
 	void SynchronizeVariableToLibraryAndApplyToGraph(UNiagaraScriptVariable* ScriptVarToSync);
 
