- LWE Initial Scale for Decals & Lights

#tests Scaled a placed system with a LWE light / decal renderer in it, toggling the options for owner scaling.
#Synced-CL 48229894

diff --git a/Engine/Plugins/FX/Niagara/Source/Niagara/Internal/Stateless/Modules/NiagaraStatelessModuleCommon.h b/Engine/Plugins/FX/Niagara/Source/Niagara/Internal/Stateless/Modules/NiagaraStatelessModuleCommon.h
--- a/Engine/Plugins/FX/Niagara/Source/Niagara/Internal/Stateless/Modules/NiagaraStatelessModuleCommon.h
+++ b/Engine/Plugins/FX/Niagara/Source/Niagara/Internal/Stateless/Modules/NiagaraStatelessModuleCommon.h
@@ -35,9 +35,15 @@ struct FNiagaraStatelessSystemScaleBuildData
 	UPROPERTY(EditAnywhere, Category = "Parameters|Sprite", meta = (DisplayPriority = 300))
 	uint32 bScaleInitialSpriteSize : 1 = false;
 
-	UPROPERTY(EditAnywhere, Category = "Parameters|Veclocity", meta = (DisplayPriority = 100))
+	UPROPERTY(EditAnywhere, Category = "Parameters|Velocity", meta = (DisplayPriority = 100))
 	uint32 bScaleInitialVelocity : 1 = false;
 
+	UPROPERTY(EditAnywhere, Category = "Parameters|Decal", meta = (DisplayPriority = 800))
+	uint32 bScaleInitialDecalSize : 1 = false;
+
+	UPROPERTY(EditAnywhere, Category = "Parameters|Lights", meta = (DisplayPriority = 900))
+	uint32 bScaleInitialLightRadius : 1 = false;
+
 	////////////////////////////////////////////////////////////////////////////////////////////////////////
 	// Bitfields
 	UPROPERTY(EditAnywhere, Category = "Parameters|Sprite", meta = (DisplayPriority = 302, EditCondition = "bScaleInitialSpriteSize", EditConditionHides))
@@ -52,9 +58,12 @@ struct FNiagaraStatelessSystemScaleBuildData
 	//UPROPERTY(EditAnywhere, Category = "Parameters|Drag", meta = (DisplayPriority = 502))
 	//uint32 bInverseDragScaling : 1 = false;
 
+	UPROPERTY(EditAnywhere, Category = "Parameters|Lights", meta = (DisplayPriority = 902, EditCondition = "bScaleInitialLightRadius", EditConditionHides))
+	uint32 bScaleInitialLightRadiusBySingleAxis : 1 = false;
+
 	////////////////////////////////////////////////////////////////////////////////////////////////////////
 	// Velocity
-	UPROPERTY(EditAnywhere, Category = "Parameters|Veclocity", meta = (DisplayPriority = 101, EditCondition = "bScaleInitialVelocity", EditConditionHides))
+	UPROPERTY(EditAnywhere, Category = "Parameters|Velocity", meta = (DisplayPriority = 101, EditCondition = "bScaleInitialVelocity", EditConditionHides))
 	FVector3f ScaleInitialVelocityScaleAmount = FVector3f::OneVector;
 
 	////////////////////////////////////////////////////////////////////////////////////////////////////////
@@ -97,6 +106,18 @@ struct FNiagaraStatelessSystemScaleBuildData
 	float ScaleInitialRibbonWidthScaleAmount = 1.0f;
 	UPROPERTY(EditAnywhere, Category = "Parameters|Ribbon", meta = (DisplayPriority = 703, EditCondition = "bScaleInitialRibbonWidth && bScaleInitialRibbonWidthBySingleAxis", EditConditionHides))
 	ENiagaraOrientationAxis ScaleInitialRibbonWidthScalingAxis = ENiagaraOrientationAxis::XAxis;
+
+	////////////////////////////////////////////////////////////////////////////////////////////////////////
+	// Decals
+	UPROPERTY(EditAnywhere, Category = "Parameters|Decal", meta = (DisplayPriority = 801, EditCondition = "bScaleInitialDecalSize", EditConditionHides))
+	FVector3f ScaleInitialDecalSizeScaleAmount = FVector3f::OneVector;
+
+	////////////////////////////////////////////////////////////////////////////////////////////////////////
+	// Lights
+	UPROPERTY(EditAnywhere, Category = "Parameters|Lights", meta = (DisplayPriority = 901, EditCondition = "bScaleInitialLightRadius", EditConditionHides))
+	float ScaleInitialLightRadiusScaleAmount = 1.0f;
+	UPROPERTY(EditAnywhere, Category = "Parameters|Lights", meta = (DisplayPriority = 903, EditCondition = "bScaleInitialLightRadius && bScaleInitialLightRadiusBySingleAxis", EditConditionHides))
+	ENiagaraOrientationAxis ScaleInitialLightRadiusScalingAxis = ENiagaraOrientationAxis::XAxis;
 };
 
 namespace NiagaraStateless
diff --git a/Engine/Plugins/FX/Niagara/Source/Niagara/Internal/Stateless/NiagaraStatelessParticleSimContext.h b/Engine/Plugins/FX/Niagara/Source/Niagara/Internal/Stateless/NiagaraStatelessParticleSimContext.h
--- a/Engine/Plugins/FX/Niagara/Source/Niagara/Internal/Stateless/NiagaraStatelessParticleSimContext.h
+++ b/Engine/Plugins/FX/Niagara/Source/Niagara/Internal/Stateless/NiagaraStatelessParticleSimContext.h
@@ -51,6 +51,7 @@ public:
 	float GetDeltaTime() const { return DeltaTime; }
 	float GetInvDeltaTime() const { return InvDeltaTime; }
 
+	const FVector3f GetLocalToWorldScale() const;
 	const FQuat4f& GetToSimulationRotation(ENiagaraCoordinateSpace SourceSpace) const;
 
 	// When modifying particle alive you need to notify the simulation that it needs to compact killed instances or not
@@ -284,6 +285,20 @@ public:
 		return reinterpret_cast<const T*>(ShaderParametersData + Offset);
 	}
 
+	template<typename T>
+	const T* ReadSharedBuiltData() const
+	{
+		const FName DataName = T::GetName();
+		for (const TPair<FName, uint32>& Pair : SharedBuiltDataOffsets)
+		{
+			if (Pair.Key == DataName)
+			{
+				return reinterpret_cast<const T*>(SharedBuiltData.GetData() + Pair.Value);
+			}
+		}
+		return nullptr;
+	}
+
 	static FQuat4f RotatorToQuat(FVector3f Rotator)
 	{
 		Rotator.X = FMath::Fractional(Rotator.X) * UE_PI;
@@ -323,6 +338,8 @@ private:
 
 	const TConstArrayView<uint8>			BuiltData;
 	mutable int32							BuiltDataOffset = 0;
+	TConstArrayView<uint8>					SharedBuiltData;
+	TConstArrayView<TPair<FName, uint32>>	SharedBuiltDataOffsets;
 	const uint8*							ShaderParametersData = nullptr;
 	mutable int32							ShaderParameterOffset = 0;
 	const TConstArrayView<float>			StaticFloatData;
diff --git a/Engine/Plugins/FX/Niagara/Source/Niagara/Private/Stateless/Modules/NiagaraStatelessModule_ApplyOwnerScaleToAttributes.cpp b/Engine/Plugins/FX/Niagara/Source/Niagara/Private/Stateless/Modules/NiagaraStatelessModule_ApplyOwnerScaleToAttributes.cpp
--- a/Engine/Plugins/FX/Niagara/Source/Niagara/Private/Stateless/Modules/NiagaraStatelessModule_ApplyOwnerScaleToAttributes.cpp
+++ b/Engine/Plugins/FX/Niagara/Source/Niagara/Private/Stateless/Modules/NiagaraStatelessModule_ApplyOwnerScaleToAttributes.cpp
@@ -19,7 +19,9 @@ void UNiagaraStatelessModule_ApplyOwnerScaleToAttributes::BuildEmitterData(const
 		SystemScaleData.bScaleInitialVelocity ||
 		SystemScaleData.bScaleInitialRibbonWidth ||
 		//SystemScaleData.bScaleInitialMeshSize ||
-		SystemScaleData.bScaleInitialSpriteSize;
+		SystemScaleData.bScaleInitialSpriteSize ||
+		SystemScaleData.bScaleInitialDecalSize ||
+		SystemScaleData.bScaleInitialLightRadius;
 
 	if (bAnyEnabled == false)
 	{
diff --git a/Engine/Plugins/FX/Niagara/Source/Niagara/Private/Stateless/Modules/NiagaraStatelessModule_DecalAttributes.cpp b/Engine/Plugins/FX/Niagara/Source/Niagara/Private/Stateless/Modules/NiagaraStatelessModule_DecalAttributes.cpp
--- a/Engine/Plugins/FX/Niagara/Source/Niagara/Private/Stateless/Modules/NiagaraStatelessModule_DecalAttributes.cpp
+++ b/Engine/Plugins/FX/Niagara/Source/Niagara/Private/Stateless/Modules/NiagaraStatelessModule_DecalAttributes.cpp
@@ -2,6 +2,7 @@
 
 #include "Stateless/Modules/NiagaraStatelessModule_DecalAttributes.h"
 
+#include "Stateless/Modules/NiagaraStatelessModuleCommon.h"
 #include "Stateless/NiagaraStatelessCommon.h"
 #include "Stateless/NiagaraStatelessEmitterDataBuildContext.h"
 #include "Stateless/NiagaraStatelessParticleSimContext.h"
@@ -59,6 +60,18 @@ namespace NSMDecalAttributesPrivate
 		using namespace NiagaraStateless;
 
 		const FModuleBuiltData* ModuleBuiltData = ParticleSimulationContext.ReadBuiltData<FModuleBuiltData>();
+
+		FVector3f SizeScale = FVector3f::OneVector;
+		if (ModuleBuiltData->AttributeOffset[int(EModuleAttribute::Size)] != INDEX_NONE)
+		{
+			const FNiagaraStatelessSystemScaleBuildData* SystemScaleBuildData = ParticleSimulationContext.ReadSharedBuiltData<FNiagaraStatelessSystemScaleBuildData>();
+			if (SystemScaleBuildData && SystemScaleBuildData->bScaleInitialDecalSize)
+			{
+				const FVector3f LocalToWorldScale = ParticleSimulationContext.GetLocalToWorldScale();
+				SizeScale = FMath::Lerp(FVector3f::OneVector, LocalToWorldScale, SystemScaleBuildData->ScaleInitialDecalSizeScaleAmount);
+			}
+		}
+
 		const float* NormalizedAgeData = ParticleSimulationContext.GetParticleNormalizedAge();
 
 		const FQuat4f OrientationRotation = ParticleSimulationContext.GetToSimulationRotation(ModuleBuiltData->OrientationCoordinateSpace);
@@ -75,7 +88,7 @@ namespace NSMDecalAttributesPrivate
 			if (ModuleBuiltData->AttributeOffset[int(EModuleAttribute::Size)] != INDEX_NONE)
 			{
 				const FVector3f Value = ParticleSimulationContext.SampleDistributionValue<FVector3f>(ModuleBuiltData->AttributeDistributionParameters[int(EModuleAttribute::Size)], i, 0, NormalizedAgeData[i]);
-				ParticleSimulationContext.WriteParticleVariable(ModuleBuiltData->AttributeOffset[int(EModuleAttribute::Size)], i, Value);
+				ParticleSimulationContext.WriteParticleVariable(ModuleBuiltData->AttributeOffset[int(EModuleAttribute::Size)], i, Value * SizeScale);
 			}
 			if (ModuleBuiltData->AttributeOffset[int(EModuleAttribute::Fade)] != INDEX_NONE)
 			{
diff --git a/Engine/Plugins/FX/Niagara/Source/Niagara/Private/Stateless/Modules/NiagaraStatelessModule_LightAttributes.cpp b/Engine/Plugins/FX/Niagara/Source/Niagara/Private/Stateless/Modules/NiagaraStatelessModule_LightAttributes.cpp
--- a/Engine/Plugins/FX/Niagara/Source/Niagara/Private/Stateless/Modules/NiagaraStatelessModule_LightAttributes.cpp
+++ b/Engine/Plugins/FX/Niagara/Source/Niagara/Private/Stateless/Modules/NiagaraStatelessModule_LightAttributes.cpp
@@ -2,6 +2,7 @@
 
 #include "Stateless/Modules/NiagaraStatelessModule_LightAttributes.h"
 
+#include "Stateless/Modules/NiagaraStatelessModuleCommon.h"
 #include "Stateless/NiagaraStatelessCommon.h"
 #include "Stateless/NiagaraStatelessEmitterDataBuildContext.h"
 #include "Stateless/NiagaraStatelessParticleSimContext.h"
@@ -78,14 +79,26 @@ namespace NSMLightAttributesPrivate
 		const FModuleBuiltData* ModuleBuiltData = ParticleSimulationContext.ReadBuiltData<FModuleBuiltData>();
 		const float* NormalizedAgeData = ParticleSimulationContext.GetParticleNormalizedAge();
 
+		float RadiusScale = 1.0f;
+		if (const FNiagaraStatelessSystemScaleBuildData* SystemScaleBuildData = ParticleSimulationContext.ReadSharedBuiltData<FNiagaraStatelessSystemScaleBuildData>())
+		{
+			if (SystemScaleBuildData->bScaleInitialLightRadius)
+			{
+				const FVector3f LocalToWorldScale3D = ParticleSimulationContext.GetLocalToWorldScale();
+				const float LocalToWorldScale = SystemScaleBuildData->bScaleInitialLightRadiusBySingleAxis ? LocalToWorldScale3D.Component(int32(SystemScaleBuildData->ScaleInitialLightRadiusScalingAxis)) : LocalToWorldScale3D.GetMax();
+				RadiusScale = FMath::Lerp(1.0f, LocalToWorldScale, SystemScaleBuildData->ScaleInitialLightRadiusScaleAmount);
+			}
+		}
+
 		for (uint32 i=0; i < ParticleSimulationContext.GetNumInstances(); ++i)
 		{
 			for ( int a=0; a < int(EModuleAttribute::Num); ++a )
 			{
 				if (ModuleBuiltData->AttributeOffset[a] != INDEX_NONE)
 				{
+					const float Scale = a == int(EModuleAttribute::Radius) ? RadiusScale : 1.0f;
 					const float Value = ParticleSimulationContext.SampleDistributionValue<float>(ModuleBuiltData->AttributeDistributionParameters[a], i, 0, NormalizedAgeData[i]);
-					ParticleSimulationContext.WriteParticleVariable(ModuleBuiltData->AttributeOffset[a], i, Value);
+					ParticleSimulationContext.WriteParticleVariable(ModuleBuiltData->AttributeOffset[a], i, Value * Scale);
 				}
 			}
 		}
diff --git a/Engine/Plugins/FX/Niagara/Source/Niagara/Private/Stateless/NiagaraStatelessEmitterInstance.cpp b/Engine/Plugins/FX/Niagara/Source/Niagara/Private/Stateless/NiagaraStatelessEmitterInstance.cpp
--- a/Engine/Plugins/FX/Niagara/Source/Niagara/Private/Stateless/NiagaraStatelessEmitterInstance.cpp
+++ b/Engine/Plugins/FX/Niagara/Source/Niagara/Private/Stateless/NiagaraStatelessEmitterInstance.cpp
@@ -72,6 +72,7 @@ namespace NiagaraStatelessEmitterInstancePrivate
 
 	void SetShaderParameterTransforms(const FNiagaraStatelessSpaceTransforms& Transforms, NiagaraStateless::FCommonShaderParameters* ShaderParameters)
 	{
+		ShaderParameters->Common_LocalToWorldScale = Transforms.GetLocalToWorld().GetScale3D();
 		ShaderParameters->Common_ToSimulationRotations[int(ENiagaraCoordinateSpace::Simulation)] = Transforms.GetTransform(ENiagaraCoordinateSpace::Simulation, ENiagaraCoordinateSpace::Simulation).GetRotation();
 		ShaderParameters->Common_ToSimulationRotations[int(ENiagaraCoordinateSpace::World)] = Transforms.GetTransform(ENiagaraCoordinateSpace::World, ENiagaraCoordinateSpace::Simulation).GetRotation();
 		ShaderParameters->Common_ToSimulationRotations[int(ENiagaraCoordinateSpace::Local)] = Transforms.GetTransform(ENiagaraCoordinateSpace::Local, ENiagaraCoordinateSpace::Simulation).GetRotation();
diff --git a/Engine/Plugins/FX/Niagara/Source/Niagara/Private/Stateless/NiagaraStatelessParticleSimContext.cpp b/Engine/Plugins/FX/Niagara/Source/Niagara/Private/Stateless/NiagaraStatelessParticleSimContext.cpp
--- a/Engine/Plugins/FX/Niagara/Source/Niagara/Private/Stateless/NiagaraStatelessParticleSimContext.cpp
+++ b/Engine/Plugins/FX/Niagara/Source/Niagara/Private/Stateless/NiagaraStatelessParticleSimContext.cpp
@@ -40,6 +40,8 @@ namespace NiagaraStateless
 	FParticleSimulationContext::FParticleSimulationContext(const FNiagaraStatelessEmitterData* InEmitterData, const void* InShaderParametersData, TConstArrayView<uint8> InDynamicBufferData)
 		: EmitterData(InEmitterData)
 		, BuiltData(EmitterData->BuiltData)
+		, SharedBuiltData(EmitterData->SharedBuiltData)
+		, SharedBuiltDataOffsets(EmitterData->SharedBuiltDataOffsets)
 		, ShaderParametersData(static_cast<const uint8*>(InShaderParametersData))
 		, StaticFloatData(MakeConstArrayView(reinterpret_cast<const float*>(EmitterData->StaticDataBufferCpu.GetData()), EmitterData->StaticDataBufferCpu.Num() / sizeof(float)))
 		, DynamicBufferData(InDynamicBufferData)
@@ -281,6 +283,13 @@ namespace NiagaraStateless
 		}
 	}
 
+	const FVector3f FParticleSimulationContext::GetLocalToWorldScale() const
+	{
+		using namespace NiagaraStateless;
+		const FCommonShaderParameters* CommonShaderParameters = reinterpret_cast<const FCommonShaderParameters*>(ShaderParametersData);
+		return CommonShaderParameters->Common_LocalToWorldScale;
+	}
+
 	const FQuat4f& FParticleSimulationContext::GetToSimulationRotation(ENiagaraCoordinateSpace SourceSpace) const
 	{
 		using namespace NiagaraStateless;
diff --git a/Engine/Plugins/FX/Niagara/Source/NiagaraShader/Internal/Stateless/NiagaraStatelessSimulationShader.h b/Engine/Plugins/FX/Niagara/Source/NiagaraShader/Internal/Stateless/NiagaraStatelessSimulationShader.h
--- a/Engine/Plugins/FX/Niagara/Source/NiagaraShader/Internal/Stateless/NiagaraStatelessSimulationShader.h
+++ b/Engine/Plugins/FX/Niagara/Source/NiagaraShader/Internal/Stateless/NiagaraStatelessSimulationShader.h
@@ -27,6 +27,7 @@ namespace NiagaraStateless
 		SHADER_PARAMETER(float,					Common_SimulationTime)
 		SHADER_PARAMETER(float,					Common_SimulationDeltaTime)
 		SHADER_PARAMETER(float,					Common_SimulationInvDeltaTime)
+		SHADER_PARAMETER(FVector3f,				Common_LocalToWorldScale)
 		SHADER_PARAMETER_ARRAY(FQuat4f,			Common_ToSimulationRotations, [3])
 
 		SHADER_PARAMETER(uint32,				Common_OutputBufferStride)
