- Refactor Skeletal Mesh DI to allow reading from other components that inherit from USkinnedMeshComponent
- This will allow you to read from a chaos cloth component on GPU if skin cache and deformed reads are allowed
- We currently do not support reading the cloth data on the CPU

#tests Add a cloth component, setup DI to read from cloth, use triangle or vertex sample on GPU with appropriate project settings

#Synced-CL 48229894

diff --git a/Engine/Plugins/FX/Niagara/Source/Niagara/Classes/NDISkeletalMeshCommon.h b/Engine/Plugins/FX/Niagara/Source/Niagara/Classes/NDISkeletalMeshCommon.h
--- a/Engine/Plugins/FX/Niagara/Source/Niagara/Classes/NDISkeletalMeshCommon.h
+++ b/Engine/Plugins/FX/Niagara/Source/Niagara/Classes/NDISkeletalMeshCommon.h
@@ -22,9 +22,7 @@ using TNDISkelMesh_SkinningModePreSkin = TIntegralConstant<ENDISkeletalMesh_Skin
 
 struct FSkeletalMeshAccessorHelper
 {
-	FSkeletalMeshAccessorHelper()
-	{
-	}
+	FSkeletalMeshAccessorHelper() = default;
 	~FSkeletalMeshAccessorHelper()
 	{
 		if (SkinningData != nullptr)
@@ -33,16 +31,13 @@ struct FSkeletalMeshAccessorHelper
 		}
 	}
 
-	FSkeletalMeshAccessorHelper(const FSkeletalMeshAccessorHelper&) = delete;
-	FSkeletalMeshAccessorHelper(FSkeletalMeshAccessorHelper&&) = delete;
-	FSkeletalMeshAccessorHelper operator=(const FSkeletalMeshAccessorHelper&) = delete;
-	FSkeletalMeshAccessorHelper operator=(FSkeletalMeshAccessorHelper&&) = delete;
+	UE_NONCOPYABLE(FSkeletalMeshAccessorHelper);
 
 	template<typename FilterMode, typename AreaWeightingMode>
 	inline void Init(const FNDISkeletalMesh_InstanceData* InstData)
 	{
-		Comp = Cast<USkeletalMeshComponent>(InstData->SceneComponent.Get());
-		Mesh = InstData->SkeletalMesh.Get();
+		Comp = Cast<USkeletalMeshComponent>(InstData->WeakSceneComponent.Get());
+		Mesh = Cast<USkeletalMesh>(InstData->WeakSkinnedMesh.Get());
 		LODData = InstData->CachedLODData;
 		SkinWeightBuffer = InstData->GetSkinWeights();
 		IndexBuffer = LODData ? LODData->MultiSizeIndexContainer.GetIndexBuffer() : nullptr;
@@ -547,48 +542,51 @@ struct TAreaWeightingModeBinder
 		bool bAreaWeighting = false;
 		if (InstData->bAllowCPUMeshDataAccess)
 		{
-			if (InstData->SamplingRegionIndices.Num() > 1)
+			if (const USkeletalMesh* SkeletalMesh = Cast<USkeletalMesh>(InstData->WeakSkinnedMesh.Get()))
 			{
-				if (InstData->SamplingRegionAreaWeightedSampler.IsValid())
+				if (InstData->SamplingRegionIndices.Num() > 1)
 				{
-					bAreaWeighting = true;
-
-					// check that all the regions have sampling data
-					const FSkeletalMeshSamplingInfo& SamplingInfo = InstData->SkeletalMesh->GetSamplingInfo();
-
-					for (int32 RegionIndex : InstData->SamplingRegionIndices)
+					if (InstData->SamplingRegionAreaWeightedSampler.IsValid())
 					{
-						const FSkeletalMeshSamplingRegion& Region = SamplingInfo.GetRegion(RegionIndex);
-						const FSkeletalMeshSamplingRegionBuiltData& RegionBuiltData = SamplingInfo.GetRegionBuiltData(RegionIndex);
+						bAreaWeighting = true;
 
-						if (RegionBuiltData.AreaWeightedSampler.GetNumEntries() == 0)
+						// check that all the regions have sampling data
+						const FSkeletalMeshSamplingInfo& SamplingInfo = SkeletalMesh->GetSamplingInfo();
+
+						for (int32 RegionIndex : InstData->SamplingRegionIndices)
 						{
-							bAreaWeighting = false;
+							const FSkeletalMeshSamplingRegion& Region = SamplingInfo.GetRegion(RegionIndex);
+							const FSkeletalMeshSamplingRegionBuiltData& RegionBuiltData = SamplingInfo.GetRegionBuiltData(RegionIndex);
+
+							if (RegionBuiltData.AreaWeightedSampler.GetNumEntries() == 0)
+							{
+								bAreaWeighting = false;
+							}
 						}
 					}
 				}
-			}
-			else if (InstData->SamplingRegionIndices.Num() == 1)
-			{
-				const int32 RegionIndex = InstData->SamplingRegionIndices[0];
+				else if (InstData->SamplingRegionIndices.Num() == 1)
+				{
+					const int32 RegionIndex = InstData->SamplingRegionIndices[0];
 
-				const FSkeletalMeshSamplingInfo& SamplingInfo = InstData->SkeletalMesh->GetSamplingInfo();
-				const FSkeletalMeshSamplingRegion& Region = SamplingInfo.GetRegion(RegionIndex);
+					const FSkeletalMeshSamplingInfo& SamplingInfo = SkeletalMesh->GetSamplingInfo();
+					const FSkeletalMeshSamplingRegion& Region = SamplingInfo.GetRegion(RegionIndex);
 
-				bAreaWeighting = Region.bSupportUniformlyDistributedSampling
-					&& (SamplingInfo.GetRegionBuiltData(RegionIndex).AreaWeightedSampler.GetNumEntries() > 0);
-			}
-			else
-			{
-				const int32 LODIndex = InstData->GetLODIndex();
+					bAreaWeighting = Region.bSupportUniformlyDistributedSampling
+						&& (SamplingInfo.GetRegionBuiltData(RegionIndex).AreaWeightedSampler.GetNumEntries() > 0);
+				}
+				else
+				{
+					const int32 LODIndex = InstData->GetLODIndex();
 
-				const FSkeletalMeshSamplingInfo& SamplingInfo = InstData->SkeletalMesh->GetSamplingInfo();
-				const FSkeletalMeshSamplingBuiltData& SamplingBuildData = SamplingInfo.GetBuiltData();
+					const FSkeletalMeshSamplingInfo& SamplingInfo = SkeletalMesh->GetSamplingInfo();
+					const FSkeletalMeshSamplingBuiltData& SamplingBuildData = SamplingInfo.GetBuiltData();
 
-				bAreaWeighting =
-					SamplingBuildData.WholeMeshBuiltData.IsValidIndex(LODIndex) &&
-					(SamplingBuildData.WholeMeshBuiltData[LODIndex].AreaWeightedTriangleSampler.GetNumEntries() > 0) &&
-					InstData->SkeletalMesh->GetLODInfo(LODIndex)->bSupportUniformlyDistributedSampling;
+					bAreaWeighting =
+						SamplingBuildData.WholeMeshBuiltData.IsValidIndex(LODIndex) &&
+						(SamplingBuildData.WholeMeshBuiltData[LODIndex].AreaWeightedTriangleSampler.GetNumEntries() > 0) &&
+						SkeletalMesh->GetLODInfo(LODIndex)->bSupportUniformlyDistributedSampling;
+				}
 			}
 		}
 
@@ -666,7 +664,7 @@ struct TSkinningModeBinder
 	{
 		FNDISkeletalMesh_InstanceData* InstData = (FNDISkeletalMesh_InstanceData*)InstanceData;
 		UNiagaraDataInterfaceSkeletalMesh* MeshInterface = CastChecked<UNiagaraDataInterfaceSkeletalMesh>(Interface);
-		USkeletalMeshComponent* Component = Cast<USkeletalMeshComponent>(InstData->SceneComponent.Get());
+		USkeletalMeshComponent* Component = Cast<USkeletalMeshComponent>(InstData->WeakSceneComponent.Get());
 		if (MeshInterface->SkinningMode == ENDISkeletalMesh_SkinningMode::None || !Component) // Can't skin if we have no component.
 		{
 			NextBinder::template Bind<ParamTypes..., FSkinnedPositionAccessorHelper<TNDISkelMesh_SkinningModeNone>>(Interface, BindingInfo, InstanceData, OutFunc);
diff --git a/Engine/Plugins/FX/Niagara/Source/Niagara/Classes/NiagaraDataInterfaceSkeletalMesh.h b/Engine/Plugins/FX/Niagara/Source/Niagara/Classes/NiagaraDataInterfaceSkeletalMesh.h
--- a/Engine/Plugins/FX/Niagara/Source/Niagara/Classes/NiagaraDataInterfaceSkeletalMesh.h
+++ b/Engine/Plugins/FX/Niagara/Source/Niagara/Classes/NiagaraDataInterfaceSkeletalMesh.h
@@ -13,6 +13,7 @@
 #include "NiagaraDataInterfaceSkeletalMesh.generated.h"
 
 class UNiagaraDataInterfaceSkeletalMesh;
+class USkinnedAsset;
 class USkeletalMesh;
 struct FSkeletalMeshSkinningData;
 struct FSkeletalMeshConnectivity;
@@ -67,8 +68,8 @@ struct FSkeletalMeshSkinningDataHandle
 
 struct FSkeletalMeshSkinningData
 {
-	FSkeletalMeshSkinningData(TWeakObjectPtr<USkeletalMeshComponent> InMeshComp)
-		: MeshComp(InMeshComp)
+	explicit FSkeletalMeshSkinningData(TWeakObjectPtr<USkinnedMeshComponent> InWeakMeshComponent)
+		: WeakMeshComponent(InWeakMeshComponent)
 		, DeltaSeconds(.0333f)
 		, CurrIndex(0)
 		, BoneMatrixUsers(0)
@@ -81,7 +82,7 @@ struct FSkeletalMeshSkinningData
 	bool IsUsed()const;
 	void ForceDataRefresh();
 
-	TWeakObjectPtr<USkeletalMeshComponent> GetWeakMeshComponent() const { return MeshComp; }
+	TWeakObjectPtr<USkinnedMeshComponent> GetWeakMeshComponent() const { return WeakMeshComponent; }
 
 	bool Tick(float InDeltaSeconds, bool bRequirePreskin = true);
 
@@ -208,7 +209,7 @@ private:
 
 	FRWLock RWGuard;
 
-	TWeakObjectPtr<USkeletalMeshComponent> MeshComp;
+	TWeakObjectPtr<USkinnedMeshComponent> WeakMeshComponent;
 
 	/** Delta seconds between calculations of the previous and current skinned positions. */
 	float DeltaSeconds;
@@ -285,7 +286,7 @@ private:
 class FNDI_SkeletalMesh_GeneratedData : public FNDI_GeneratedData
 {
 	FRWLock CachedSkinningDataGuard;
-	TMap<TWeakObjectPtr<USkeletalMeshComponent>, TSharedPtr<FSkeletalMeshSkinningData> > CachedSkinningData;
+	TMap<TWeakObjectPtr<USkinnedMeshComponent>, TSharedPtr<FSkeletalMeshSkinningData> > CachedSkinningData;
 
 	FRWLock CachedUvMappingGuard;
 	TArray<TSharedPtr<FSkeletalMeshUvMapping>> CachedUvMapping;
@@ -294,9 +295,9 @@ class FNDI_SkeletalMesh_GeneratedData : public FNDI_GeneratedData
 	TArray<TSharedPtr<FSkeletalMeshConnectivity>> CachedConnectivity;
 
 public:
-	NIAGARA_API FSkeletalMeshSkinningDataHandle GetCachedSkinningData(TWeakObjectPtr<USkeletalMeshComponent>& InComponent, FSkeletalMeshSkinningDataUsage Usage, bool bNeedsDataImmediately);
-	NIAGARA_API FSkeletalMeshUvMappingHandle GetCachedUvMapping(TWeakObjectPtr<USkeletalMesh>& InMeshObject, int32 InLodIndex, int32 InUvSetIndex, FMeshUvMappingUsage Usage, bool bNeedsDataImmediately);
-	NIAGARA_API FSkeletalMeshConnectivityHandle GetCachedConnectivity(TWeakObjectPtr<USkeletalMesh>& InMeshObject, int32 InLodIndex, FSkeletalMeshConnectivityUsage Usage, bool bNeedsDataImmediately);
+	NIAGARA_API FSkeletalMeshSkinningDataHandle GetCachedSkinningData(TWeakObjectPtr<USkinnedMeshComponent> InComponent, FSkeletalMeshSkinningDataUsage Usage, bool bNeedsDataImmediately);
+	NIAGARA_API FSkeletalMeshUvMappingHandle GetCachedUvMapping(TWeakObjectPtr<USkinnedAsset> InMeshObject, int32 InLodIndex, int32 InUvSetIndex, FMeshUvMappingUsage Usage, bool bNeedsDataImmediately);
+	NIAGARA_API FSkeletalMeshConnectivityHandle GetCachedConnectivity(TWeakObjectPtr<USkinnedAsset> InMeshObject, int32 InLodIndex, FSkeletalMeshConnectivityUsage Usage, bool bNeedsDataImmediately);
 
 	NIAGARA_API virtual void Tick(ETickingGroup TickGroup, float DeltaSeconds) override;
 
@@ -403,7 +404,7 @@ public:
 
 	virtual FString GetFriendlyName() const override { return TEXT("FSkeletalMeshGpuSpawnStaticBuffers"); }
 
-	USceneComponent* GetSceneComponent() const { return SceneComponent.Get(); }
+	USceneComponent* GetSceneComponent() const { return WeakSceneComponent.Get(); }
 
 	int32 GetLODIndex() const { return LODIndex; };
 
@@ -441,7 +442,7 @@ public:
 	int32 GetFilteredSocketBoneOffset() const { return FilteredSocketBoneOffset; }
 
 protected:
-	TWeakObjectPtr<USceneComponent> SceneComponent = nullptr;
+	TWeakObjectPtr<USceneComponent> WeakSceneComponent = nullptr;
 
 	FBufferRHIRef BufferTriangleUniformSamplerProbAliasRHI = nullptr;
 	FShaderResourceViewRHIRef BufferTriangleUniformSamplerProbAliasSRV = nullptr;
@@ -553,18 +554,18 @@ private:
 
 struct FNDISkeletalMesh_InstanceData
 {
-	/** Cached ptr to SkeletalMeshComponent we sample from, when found. Otherwise, the scene component to use to transform the PreviewMesh */
-	TWeakObjectPtr<USceneComponent> SceneComponent;
+	/** Cached ptr to Scene Component we sample from, when found. Otherwise, the scene component to use to transform the PreviewMesh */
+	TWeakObjectPtr<USceneComponent> WeakSceneComponent;
 
 	/** A binding to the user ptr we're reading the mesh from (if we are). */
 	FNiagaraParameterDirectBinding<UObject*> UserParamBinding;
 
 	/** Always reset the DI when the attach parent changes. */
-	TWeakObjectPtr<USceneComponent> CachedAttachParent;
+	TWeakObjectPtr<USceneComponent> WeakCachedAttachParent;
 
 	UObject* CachedUserParam;
 
-	TWeakObjectPtr<USkeletalMesh> SkeletalMesh;
+	TWeakObjectPtr<USkinnedAsset> WeakSkinnedMesh;
 
 	/** Handle to our skinning data. */
 	FSkeletalMeshSkinningDataHandle SkinningData;
@@ -668,10 +669,10 @@ struct FNDISkeletalMesh_InstanceData
 
 	inline const FSkinWeightVertexBuffer* GetSkinWeights() const
 	{
-		USkeletalMeshComponent* SkelComp = Cast<USkeletalMeshComponent>(SceneComponent.Get());
-		if (SkelComp != nullptr && SkelComp->GetSkeletalMeshAsset() != nullptr)
+		USkinnedMeshComponent* SkinnedComponent = Cast<USkinnedMeshComponent>(WeakSceneComponent.Get());
+		if (SkinnedComponent != nullptr && SkinnedComponent->GetSkinnedAsset() != nullptr)
 		{
-			return SkelComp->GetSkinWeightBuffer(CachedLODIdx);
+			return SkinnedComponent->GetSkinWeightBuffer(CachedLODIdx);
 		}
 		return CachedLODData ? &CachedLODData->SkinWeightVertexBuffer : nullptr;
 	}
@@ -715,6 +716,10 @@ protected:
 	UPROPERTY(EditAnywhere, Category = "Mesh")
 	TArray<FName> ComponentTags;
 
+	/** Class type we should look for when attaching defaults to USkeletalMeshComponent if unset */
+	UPROPERTY(EditAnywhere, Category = "Mesh")
+	TSubclassOf<USkinnedMeshComponent> OptionalComponentClass;
+
 #if WITH_EDITORONLY_DATA
 	UPROPERTY()
 	TObjectPtr<AActor> Source_DEPRECATED;
@@ -722,7 +727,7 @@ protected:
 
 	/** The source component from which to sample. Takes precedence over the direct mesh. Not exposed to the user, only indirectly accessible from blueprints. */
 	UPROPERTY(Transient)
-	TObjectPtr<USkeletalMeshComponent> SourceComponent;
+	TObjectPtr<USkinnedMeshComponent> SourceComponent;
 
 public:
 	/** Reference to a user parameter if we're reading one. */
@@ -827,14 +832,14 @@ public:
 	//~ UNiagaraDataInterface interface END
 
 	/** This overload is for use when initializing per-instance data. It possibly uses the SystemInstance and instance data to initialize a user binding */
-	NIAGARA_API USkeletalMesh* GetSkeletalMesh(FNiagaraSystemInstance* SystemInstance, USceneComponent* AttachComponent, TWeakObjectPtr<USceneComponent>& SceneComponent, USkeletalMeshComponent*& FoundSkelComp, FNDISkeletalMesh_InstanceData* InstData);
+	NIAGARA_API USkinnedAsset* GetSkinnedMesh(FNiagaraSystemInstance* SystemInstance, USceneComponent* AttachComponent, TWeakObjectPtr<USceneComponent>& SceneComponent, USkinnedMeshComponent*& FoundSkinnedComponent, FNDISkeletalMesh_InstanceData* InstData);
 	/** Finds the skeletal mesh based on settings of the DI and the hierarchy of the object provided */
-	NIAGARA_API USkeletalMesh* GetSkeletalMesh(UNiagaraComponent* Component);
+	NIAGARA_API USkinnedAsset* GetSkinnedMesh(UNiagaraComponent* Component);
 
-	USkeletalMeshComponent* GetSourceComponent() const { return SourceComponent; }
+	USkinnedMeshComponent* GetSourceComponent() const { return SourceComponent; }
 	AActor* GetSourceActor() const { return SoftSourceActor.Get(); }
 
-	NIAGARA_API int32 CalculateLODIndexAndSamplingRegions(USkeletalMesh* InMesh, TArray<int32>& OutSamplingRegionIndices, bool& OutAllRegionsAreAreaWeighting) const;
+	NIAGARA_API int32 CalculateLODIndexAndSamplingRegions(USkinnedAsset* SkinnedMesh, TArray<int32>& OutSamplingRegionIndices, bool& OutAllRegionsAreAreaWeighting) const;
 
 	NIAGARA_API virtual void ProvidePerInstanceDataForRenderThread(void* DataForRenderThread, void* PerInstanceData, const FNiagaraSystemInstanceID& SystemInstance) override;
 
@@ -1002,7 +1007,7 @@ public:
 	// End of Direct Bone + Socket Sampling
 	//////////////////////////////////////////////////////////////////////////
 
-	NIAGARA_API void SetSourceComponentFromBlueprints(USkeletalMeshComponent* ComponentToUse);
+	NIAGARA_API void SetSourceComponentFromBlueprints(USkinnedMeshComponent* ComponentToUse);
 	NIAGARA_API void SetSamplingRegionsFromBlueprints(const TArray<FName>& InSamplingRegions);
 	NIAGARA_API void SetFilteredBonesFromBlueprints(const TArray<FName>& InFilteredBones);
 	NIAGARA_API void SetFilteredSocketsFromBlueprints(const TArray<FName>& InFilteredSockets);
diff --git a/Engine/Plugins/FX/Niagara/Source/Niagara/Private/Experimental/NiagaraDataInterfaceSkeletalMeshUvMapping.cpp b/Engine/Plugins/FX/Niagara/Source/Niagara/Private/Experimental/NiagaraDataInterfaceSkeletalMeshUvMapping.cpp
--- a/Engine/Plugins/FX/Niagara/Source/Niagara/Private/Experimental/NiagaraDataInterfaceSkeletalMeshUvMapping.cpp
+++ b/Engine/Plugins/FX/Niagara/Source/Niagara/Private/Experimental/NiagaraDataInterfaceSkeletalMeshUvMapping.cpp
@@ -31,15 +31,15 @@ private:
 	const FSkelMeshVertexAccessor<UseFullPrecisionUv> MeshVertexAccessor;
 };
 
-FSkeletalMeshUvMapping::FSkeletalMeshUvMapping(TWeakObjectPtr<USkeletalMesh> InMeshObject, int32 InLodIndex, int32 InUvSetIndex)
+FSkeletalMeshUvMapping::FSkeletalMeshUvMapping(TWeakObjectPtr<USkinnedAsset> InMeshObject, int32 InLodIndex, int32 InUvSetIndex)
 	: FMeshUvMapping(InLodIndex, InUvSetIndex)
 	, MeshObject(InMeshObject)
 {
 }
 
-bool FSkeletalMeshUvMapping::IsValidMeshObject(const TWeakObjectPtr<USkeletalMesh>& MeshObject, int32 InLodIndex, int32 InUvSetIndex)
+bool FSkeletalMeshUvMapping::IsValidMeshObject(const TWeakObjectPtr<USkinnedAsset>& MeshObject, int32 InLodIndex, int32 InUvSetIndex)
 {
-	USkeletalMesh* Mesh = MeshObject.Get();
+	USkinnedAsset* Mesh = MeshObject.Get();
 
 	if (!Mesh)
 	{
@@ -68,14 +68,14 @@ bool FSkeletalMeshUvMapping::IsValidMeshObject(const TWeakObjectPtr<USkeletalMes
 	return true;
 }
 
-bool FSkeletalMeshUvMapping::Matches(const TWeakObjectPtr<USkeletalMesh>& InMeshObject, int32 InLodIndex, int32 InUvSetIndex) const
+bool FSkeletalMeshUvMapping::Matches(const TWeakObjectPtr<USkinnedAsset>& InMeshObject, int32 InLodIndex, int32 InUvSetIndex) const
 {
 	return LodIndex == InLodIndex && MeshObject == InMeshObject && UvSetIndex == InUvSetIndex;
 }
 
 const FSkeletalMeshLODRenderData* FSkeletalMeshUvMapping::GetLodRenderData() const
 {
-	if (const USkeletalMesh* Mesh = MeshObject.Get())
+	if (const USkinnedAsset* Mesh = MeshObject.Get())
 	{
 		return GetLodRenderData(*Mesh, LodIndex, UvSetIndex);
 	}
@@ -145,7 +145,7 @@ int32 FSkeletalMeshUvMapping::FindFirstTriangle(const FBox2D& InUvBox, FVector&
 	return INDEX_NONE;
 }
 
-const FSkeletalMeshLODRenderData* FSkeletalMeshUvMapping::GetLodRenderData(const USkeletalMesh& Mesh, int32 LodIndex, int32 UvSetIndex)
+const FSkeletalMeshLODRenderData* FSkeletalMeshUvMapping::GetLodRenderData(const USkinnedAsset& Mesh, int32 LodIndex, int32 UvSetIndex)
 {
 	if (const FSkeletalMeshRenderData* RenderData = Mesh.GetResourceForRendering())
 	{
diff --git a/Engine/Plugins/FX/Niagara/Source/Niagara/Private/Experimental/NiagaraDataInterfaceSkeletalMeshUvMapping.h b/Engine/Plugins/FX/Niagara/Source/Niagara/Private/Experimental/NiagaraDataInterfaceSkeletalMeshUvMapping.h
--- a/Engine/Plugins/FX/Niagara/Source/Niagara/Private/Experimental/NiagaraDataInterfaceSkeletalMeshUvMapping.h
+++ b/Engine/Plugins/FX/Niagara/Source/Niagara/Private/Experimental/NiagaraDataInterfaceSkeletalMeshUvMapping.h
@@ -8,16 +8,16 @@
 #include "UObject/WeakObjectPtrTemplates.h"
 
 class FSkeletalMeshLODRenderData;
-class USkeletalMesh;
+class USkinnedAsset;
 
 struct FSkeletalMeshUvMapping : public FMeshUvMapping
 {
-	FSkeletalMeshUvMapping(TWeakObjectPtr<USkeletalMesh> InMeshObject, int32 InLodIndex, int32 InUvSetIndex);
+	FSkeletalMeshUvMapping(TWeakObjectPtr<USkinnedAsset> InMeshObject, int32 InLodIndex, int32 InUvSetIndex);
 	virtual ~FSkeletalMeshUvMapping() = default;
 
-	static bool IsValidMeshObject(const TWeakObjectPtr<USkeletalMesh>& MeshObject, int32 InLodIndex, int32 InUvSetIndex);
+	static bool IsValidMeshObject(const TWeakObjectPtr<USkinnedAsset>& MeshObject, int32 InLodIndex, int32 InUvSetIndex);
 
-	bool Matches(const TWeakObjectPtr<USkeletalMesh>& InMeshObject, int32 InLodIndex, int32 InUvSetIndex) const;
+	bool Matches(const TWeakObjectPtr<USkinnedAsset>& InMeshObject, int32 InLodIndex, int32 InUvSetIndex) const;
 
 	const FSkeletalMeshLODRenderData* GetLodRenderData() const;
 
@@ -26,8 +26,8 @@ struct FSkeletalMeshUvMapping : public FMeshUvMapping
 	int32 FindFirstTriangle(const FBox2D& InUvBox, FVector& BarycentricCoord) const;
 
 private:
-	static const FSkeletalMeshLODRenderData* GetLodRenderData(const USkeletalMesh& Mesh, int32 LodIndex, int32 UvSetIndex);
+	static const FSkeletalMeshLODRenderData* GetLodRenderData(const USkinnedAsset& Mesh, int32 LodIndex, int32 UvSetIndex);
 	virtual void BuildQuadTree() override;
 
-	TWeakObjectPtr<USkeletalMesh> MeshObject;
+	TWeakObjectPtr<USkinnedAsset> MeshObject;
 };
diff --git a/Engine/Plugins/FX/Niagara/Source/Niagara/Private/NDISkeletalMesh_TriangleSampling.cpp b/Engine/Plugins/FX/Niagara/Source/Niagara/Private/NDISkeletalMesh_TriangleSampling.cpp
--- a/Engine/Plugins/FX/Niagara/Source/Niagara/Private/NDISkeletalMesh_TriangleSampling.cpp
+++ b/Engine/Plugins/FX/Niagara/Source/Niagara/Private/NDISkeletalMesh_TriangleSampling.cpp
@@ -695,28 +695,29 @@ void UNiagaraDataInterfaceSkeletalMesh::RandomTriangle(FVectorVMExternalFunction
 		return;
 	}
 
-	USkeletalMesh* SkelMesh = MeshAccessor.Mesh;
-	check(SkelMesh);
-	const int32 LODIndex = InstData->GetLODIndex();
-	const bool bAreaWeighted = SkelMesh->GetLODInfo(LODIndex)->bSupportUniformlyDistributedSampling;
-
-	if (bAreaWeighted)
+	if (USkeletalMesh* SkelMesh = MeshAccessor.Mesh)
 	{
-		const FSkeletalMeshSamplingInfo& SamplingInfo = SkelMesh->GetSamplingInfo();
-		const FSkeletalMeshSamplingBuiltData& SamplingInfoBuiltData = SamplingInfo.GetBuiltData();
-		if (SamplingInfoBuiltData.WholeMeshBuiltData.IsValidIndex(LODIndex))
+		const int32 LODIndex = InstData->GetLODIndex();
+		const bool bAreaWeighted = SkelMesh->GetLODInfo(LODIndex)->bSupportUniformlyDistributedSampling;
+
+		if (bAreaWeighted)
 		{
-			//-TODO: Consider adding a way to warn the user we failed to get the area weighted data
-			const FSkeletalMeshSamplingLODBuiltData& WholeMeshBuiltData = SamplingInfoBuiltData.WholeMeshBuiltData[LODIndex];
-			if (WholeMeshBuiltData.AreaWeightedTriangleSampler.GetNumEntries() > 0)
+			const FSkeletalMeshSamplingInfo& SamplingInfo = SkelMesh->GetSamplingInfo();
+			const FSkeletalMeshSamplingBuiltData& SamplingInfoBuiltData = SamplingInfo.GetBuiltData();
+			if (SamplingInfoBuiltData.WholeMeshBuiltData.IsValidIndex(LODIndex))
 			{
-				for (int32 i = 0; i < Context.GetNumInstances(); ++i)
+				//-TODO: Consider adding a way to warn the user we failed to get the area weighted data
+				const FSkeletalMeshSamplingLODBuiltData& WholeMeshBuiltData = SamplingInfoBuiltData.WholeMeshBuiltData[LODIndex];
+				if (WholeMeshBuiltData.AreaWeightedTriangleSampler.GetNumEntries() > 0)
 				{
-					RandHelper.GetAndAdvance();
-					OutTri.SetAndAdvance(WholeMeshBuiltData.AreaWeightedTriangleSampler.GetEntryIndex(RandHelper.Rand(i), RandHelper.Rand(i)));
-					OutBary.SetAndAdvance(RandHelper.RandomBarycentricCoord(i));
+					for (int32 i = 0; i < Context.GetNumInstances(); ++i)
+					{
+						RandHelper.GetAndAdvance();
+						OutTri.SetAndAdvance(WholeMeshBuiltData.AreaWeightedTriangleSampler.GetEntryIndex(RandHelper.Rand(i), RandHelper.Rand(i)));
+						OutBary.SetAndAdvance(RandHelper.RandomBarycentricCoord(i));
+					}
+					return;
 				}
-				return;
 			}
 		}
 	}
@@ -1002,7 +1003,6 @@ void UNiagaraDataInterfaceSkeletalMesh::GetTriCoordColor(FVectorVMExternalFuncti
 
 	FNDIOutputParam<FLinearColor> OutColor(Context);
 
-	USkeletalMeshComponent* Comp = Cast<USkeletalMeshComponent>(InstData->SceneComponent.Get());
 	if ( const FSkeletalMeshLODRenderData* LODData = InstData->CachedLODData )
 	{
 		const FColorVertexBuffer& Colors = LODData->StaticVertexBuffers.ColorVertexBuffer;
@@ -1064,7 +1064,6 @@ void UNiagaraDataInterfaceSkeletalMesh::GetTriCoordUV(FVectorVMExternalFunctionC
 
 	FNDIOutputParam<FVector2f> OutUV(Context);
 
-	USkeletalMeshComponent* Comp = Cast<USkeletalMeshComponent>(InstData->SceneComponent.Get());
 	if ( const FSkeletalMeshLODRenderData* LODData = InstData->CachedLODData )
 	{
 		const FMultiSizeIndexContainer& Indices = LODData->MultiSizeIndexContainer;
diff --git a/Engine/Plugins/FX/Niagara/Source/Niagara/Private/NDISkeletalMesh_VertexSampling.cpp b/Engine/Plugins/FX/Niagara/Source/Niagara/Private/NDISkeletalMesh_VertexSampling.cpp
--- a/Engine/Plugins/FX/Niagara/Source/Niagara/Private/NDISkeletalMesh_VertexSampling.cpp
+++ b/Engine/Plugins/FX/Niagara/Source/Niagara/Private/NDISkeletalMesh_VertexSampling.cpp
@@ -615,7 +615,6 @@ void UNiagaraDataInterfaceSkeletalMesh::GetVertexColor(FVectorVMExternalFunction
 
 	FNDIOutputParam<FLinearColor> OutColor(Context);
 
-	USkeletalMeshComponent* Comp = Cast<USkeletalMeshComponent>(InstData->SceneComponent.Get());
 	if ( const FSkeletalMeshLODRenderData* LODData = InstData->CachedLODData )
 	{
 		const FColorVertexBuffer& Colors = LODData->StaticVertexBuffers.ColorVertexBuffer;
@@ -672,7 +671,6 @@ void UNiagaraDataInterfaceSkeletalMesh::GetVertexUV(FVectorVMExternalFunctionCon
 
 	FNDIOutputParam<FVector2f> OutUV(Context);
 
-	USkeletalMeshComponent* Comp = Cast<USkeletalMeshComponent>(InstData->SceneComponent.Get());
 	if ( const FSkeletalMeshLODRenderData* LODData = InstData->CachedLODData )
 	{
 		const FMultiSizeIndexContainer& Indices = LODData->MultiSizeIndexContainer;
diff --git a/Engine/Plugins/FX/Niagara/Source/Niagara/Private/NiagaraDataInterfaceSkeletalMesh.cpp b/Engine/Plugins/FX/Niagara/Source/Niagara/Private/NiagaraDataInterfaceSkeletalMesh.cpp
--- a/Engine/Plugins/FX/Niagara/Source/Niagara/Private/NiagaraDataInterfaceSkeletalMesh.cpp
+++ b/Engine/Plugins/FX/Niagara/Source/Niagara/Private/NiagaraDataInterfaceSkeletalMesh.cpp
@@ -208,13 +208,30 @@ namespace NDISkelMeshLocal
 	}
 
 	// Calculate which tick group the skeletal mesh component will be ready by
-	ETickingGroup GetComponentTickGroup(USkeletalMeshComponent* Component)
+	ETickingGroup GetComponentTickGroup(USkinnedMeshComponent* SkinnedComponent, USkeletalMeshComponent* SkeletalComponent)
 	{
-		const ETickingGroup ComponentTickGroup = FMath::Max(Component->PrimaryComponentTick.TickGroup, Component->PrimaryComponentTick.EndTickGroup);
-		const ETickingGroup PhysicsTickGroup = Component->bBlendPhysics ? FMath::Max(ComponentTickGroup, TG_EndPhysics) : ComponentTickGroup;
+		const ETickingGroup ComponentTickGroup = FMath::Max(SkinnedComponent->PrimaryComponentTick.TickGroup, SkinnedComponent->PrimaryComponentTick.EndTickGroup);
+		const ETickingGroup PhysicsTickGroup = SkeletalComponent && SkeletalComponent->bBlendPhysics ? FMath::Max(ComponentTickGroup, TG_EndPhysics) : ComponentTickGroup;
 		const ETickingGroup ClampedTickGroup = FMath::Clamp(ETickingGroup(PhysicsTickGroup + 1), NiagaraFirstTickGroup, NiagaraLastTickGroup);
 		return ClampedTickGroup;
 	}
+
+	// Returns a pair where first is the number of lods and second is the first available LOD
+	static const TTuple<int32, int32> GetStreamingLODData(const USkinnedAsset* SkinnedMesh)
+	{
+		const FStreamableRenderResourceState& SRRState = SkinnedMesh->GetStreamableResourceState();
+		if (SRRState.IsValid())
+		{
+			const int32 NumValidLODs = FMath::Min(SRRState.NumRequestedLODs, SRRState.NumResidentLODs);
+			const int32 CurrentFirstLOD = SRRState.LODCountToAssetFirstLODIdx(NumValidLODs);
+			return MakeTuple(NumValidLODs, CurrentFirstLOD);
+		}
+		else if (FSkeletalMeshRenderData* RenderData = SkinnedMesh->GetResourceForRendering())
+		{
+			return MakeTuple(RenderData->LODRenderData.Num(), 0);
+		}
+		return MakeTuple(0, 0);
+	}
 }
 
 //////////////////////////////////////////////////////////////////////////
@@ -234,14 +251,14 @@ float FSkeletalMeshSamplingRegionAreaWeightedSampler::GetWeights(TArray<float>&
 {
 	check(Owner);
 
-	USkeletalMesh* SkelMesh = Owner->SkeletalMesh.Get();
-	if (SkelMesh == nullptr)
+	USkeletalMesh* SkeletalMesh = Cast<USkeletalMesh>(Owner->WeakSkinnedMesh.Get());
+	if (SkeletalMesh == nullptr)
 	{
 		OutWeights.Empty();
 		return 0.0f;
 	}
 
-	check(SkelMesh->IsValidLODIndex(Owner->GetLODIndex()));
+	check(SkeletalMesh->IsValidLODIndex(Owner->GetLODIndex()));
 
 	float Total = 0.0f;
 	int32 NumUsedRegions = Owner->SamplingRegionIndices.Num();
@@ -252,7 +269,7 @@ float FSkeletalMeshSamplingRegionAreaWeightedSampler::GetWeights(TArray<float>&
 		return 0.0f;
 	}
 
-	const FSkeletalMeshSamplingInfo& SamplingInfo = SkelMesh->GetSamplingInfo();
+	const FSkeletalMeshSamplingInfo& SamplingInfo = SkeletalMesh->GetSamplingInfo();
 	OutWeights.Empty(NumUsedRegions);
 	for (int32 i = 0; i < NumUsedRegions; ++i)
 	{
@@ -317,16 +334,16 @@ void FSkeletalMeshSkinningData::RegisterUser(FSkeletalMeshSkinningDataUsage Usag
 {
 	FRWScopeLock Lock(RWGuard, SLT_Write);
 
-	USkeletalMeshComponent* SkelComp = MeshComp.Get();
-	check(SkelComp);
+	USkinnedMeshComponent* SkinnedComponent = WeakMeshComponent.Get();
+	check(SkinnedComponent);
 
-	USkeletalMesh* SkelMesh = SkelComp->GetSkeletalMeshAsset();
+	USkinnedAsset* SkinnedMesh = SkinnedComponent->GetSkinnedAsset();
 	int32 LODIndex = 0;
 	int32 NumLODInfo = 1;
 
-	if (SkelMesh != nullptr)
+	if (SkinnedMesh != nullptr)
 	{
-		NumLODInfo = SkelMesh->GetLODNum();
+		NumLODInfo = SkinnedMesh->GetLODNum();
 		LODIndex = Usage.GetLODIndex();
 		check(LODIndex != INDEX_NONE);
 		check(LODIndex < NumLODInfo);
@@ -349,7 +366,7 @@ void FSkeletalMeshSkinningData::RegisterUser(FSkeletalMeshSkinningDataUsage Usag
 	if (bNeedsDataImmediately)
 	{
 		check(IsInGameThread());
-		const int32 NumExpectedTransforms = SkelMesh ? SkelMesh->GetRefBasesInvMatrix().Num() : 0;
+		const int32 NumExpectedTransforms = SkinnedMesh ? SkinnedMesh->GetRefBasesInvMatrix().Num() : 0;
 		if (CurrBoneRefToLocals().Num() != NumExpectedTransforms || CurrComponentTransforms().Num() != NumExpectedTransforms)
 		{
 			UpdateBoneTransforms();
@@ -366,12 +383,35 @@ void FSkeletalMeshSkinningData::RegisterUser(FSkeletalMeshSkinningDataUsage Usag
 			PrevComponentTransforms() = CurrComponentTransforms();
 		}
 
-		if (Usage.NeedPreSkinnedVerts() && CurrSkinnedPositions(LODIndex).Num() == 0 && SkelMesh && SkelMesh->GetLODInfo(LODIndex)->bAllowCPUAccess)
+		if (Usage.NeedPreSkinnedVerts() && CurrSkinnedPositions(LODIndex).Num() == 0 && SkinnedMesh && SkinnedMesh->GetLODInfo(LODIndex)->bAllowCPUAccess)
 		{
-			FSkeletalMeshLODRenderData& SkelMeshLODData = SkelMesh->GetResourceForRendering()->LODRenderData[LODIndex];
-			FSkinWeightVertexBuffer* SkinWeightBuffer = SkelComp->GetSkinWeightBuffer(LODIndex);
-			USkeletalMeshComponent::ComputeSkinnedPositions(SkelComp, CurrSkinnedPositions(LODIndex), CurrBoneRefToLocals(), SkelMeshLODData, *SkinWeightBuffer);
-			USkeletalMeshComponent::ComputeSkinnedTangentBasis(SkelComp, CurrSkinnedTangentBasis(LODIndex), CurrBoneRefToLocals(), SkelMeshLODData, *SkinWeightBuffer);
+			TRefCountPtr<const FSkeletalMeshLODRenderData> SkelMeshLODData = &SkinnedMesh->GetResourceForRendering()->LODRenderData[LODIndex];
+			FSkinWeightVertexBuffer* SkinWeightBuffer = SkinnedComponent->GetSkinWeightBuffer(LODIndex);
+			if (USkeletalMeshComponent* SkeletalComponent = Cast<USkeletalMeshComponent>(SkinnedComponent))
+			{
+				USkeletalMeshComponent::ComputeSkinnedPositions(SkeletalComponent, CurrSkinnedPositions(LODIndex), CurrBoneRefToLocals(), *SkelMeshLODData, *SkinWeightBuffer);
+				USkeletalMeshComponent::ComputeSkinnedTangentBasis(SkeletalComponent, CurrSkinnedTangentBasis(LODIndex), CurrBoneRefToLocals(), *SkelMeshLODData, *SkinWeightBuffer);
+			}
+			else
+			{
+				//-TODO: Not supported currently
+				//USkinnedMeshComponent::ComputeSkinnedPositions(SkinnedComponent, CurrSkinnedPositions(LODIndex), CurrBoneRefToLocals(), SkelMeshLODData, *SkinWeightBuffer);
+				//USkinnedMeshComponent::ComputeSkinnedTangentBasis(SkinnedComponent, CurrSkinnedTangentBasis(LODIndex), CurrBoneRefToLocals(), SkelMeshLODData, *SkinWeightBuffer);
+				const int32 NumVertices = SkelMeshLODData->GetNumVertices();
+
+				TArray<FVector3f>& OutPositions = CurrSkinnedPositions(LODIndex);
+				TArray<FVector3f>& OutTangents = CurrSkinnedTangentBasis(LODIndex);
+
+				OutPositions.SetNumUninitialized(NumVertices);
+				OutTangents.SetNumUninitialized(NumVertices * 3);
+				for (int32 i = 0; i < NumVertices; ++i)
+				{
+					OutPositions[i] = FVector3f::ZeroVector;
+					OutTangents[i * 3 + 0] = FVector3f::XAxisVector;
+					OutTangents[i * 3 + 1] = FVector3f::YAxisVector;
+					OutTangents[i * 3 + 2] = FVector3f::ZAxisVector;
+				}
+			}
 
 			//Prime the previous positions if they're missing
 			if (PrevSkinnedPositions(LODIndex).Num() != CurrSkinnedPositions(LODIndex).Num())
@@ -397,8 +437,8 @@ void FSkeletalMeshSkinningData::UnregisterUser(FSkeletalMeshSkinningDataUsage Us
 
 	int32 LODIndex = 0;
 
-	USkeletalMeshComponent* SkelComp = MeshComp.Get();
-	if (SkelComp && SkelComp->GetSkeletalMeshAsset())
+	USkinnedMeshComponent* SkinnedComp = WeakMeshComponent.Get();
+	if (SkinnedComp && SkinnedComp->GetSkinnedAsset())
 	{
 		LODIndex = Usage.GetLODIndex();
 	}
@@ -437,11 +477,11 @@ bool FSkeletalMeshSkinningData::IsUsed() const
 
 void FSkeletalMeshSkinningData::UpdateBoneTransforms()
 {
-	USkeletalMeshComponent* SkelComp = MeshComp.Get();
-	check(SkelComp);
+	USkinnedMeshComponent* SkinnedComponent = WeakMeshComponent.Get();
+	check(SkinnedComponent);
 
-	const USkeletalMesh* SkelMesh = SkelComp->GetSkeletalMeshAsset();
-	if (SkelMesh == nullptr)
+	const USkinnedAsset* SkinnedMesh = SkinnedComponent->GetSkinnedAsset();
+	if (SkinnedMesh == nullptr)
 	{
 		return;
 	}
@@ -449,18 +489,19 @@ void FSkeletalMeshSkinningData::UpdateBoneTransforms()
 	TArray<FMatrix44f>& CurrBones = CurrBoneRefToLocals();
 	TArray<FTransform3f>& CurrTransforms = CurrComponentTransforms();
 
-	if (USkinnedMeshComponent* LeaderComponent = SkelComp->LeaderPoseComponent.Get())
+	if (USkinnedMeshComponent* LeaderComponent = SkinnedComponent->LeaderPoseComponent.Get())
 	{
-		const TArray<int32>& LeaderBoneMap = SkelComp->GetLeaderBoneMap();
+		const TArray<int32>& LeaderBoneMap = SkinnedComponent->GetLeaderBoneMap();
 		const int32 NumBones = LeaderBoneMap.Num();
 
 		if (NumBones == 0)
 		{
-			// This case indicates an invalid leader pose component (e.g. no skeletal mesh)
-			CurrBones.Empty(SkelMesh->GetRefSkeleton().GetNum());
-			CurrBones.AddDefaulted(SkelMesh->GetRefSkeleton().GetNum());
-			CurrTransforms.Empty(SkelMesh->GetRefSkeleton().GetNum());
-			CurrTransforms.AddDefaulted(SkelMesh->GetRefSkeleton().GetNum());
+			// This case indicates an invalid leader pose component (e.g. no skinned mesh)
+			const int32 RefNumBones = SkinnedMesh->GetRefSkeleton().GetNum();
+			CurrBones.Empty(RefNumBones);
+			CurrBones.AddDefaulted(RefNumBones);
+			CurrTransforms.Empty(RefNumBones);
+			CurrTransforms.AddDefaulted(RefNumBones);
 		}
 		else
 		{
@@ -483,8 +524,8 @@ void FSkeletalMeshSkinningData::UpdateBoneTransforms()
 
 				if ( !bFoundLeader)
 				{
-					const int32 ParentIndex = SkelMesh->GetRefSkeleton().GetParentIndex(BoneIndex);
-					FTransform3f BoneTransform = (FTransform3f)SkelMesh->GetRefSkeleton().GetRefBonePose()[BoneIndex];
+					const int32 ParentIndex = SkinnedMesh->GetRefSkeleton().GetParentIndex(BoneIndex);
+					FTransform3f BoneTransform = FTransform3f(SkinnedMesh->GetRefSkeleton().GetRefBonePose()[BoneIndex]);
 					if ((ParentIndex >= 0) && (ParentIndex < BoneIndex))
 					{
 						BoneTransform = BoneTransform * CurrTransforms[ParentIndex];
@@ -492,9 +533,9 @@ void FSkeletalMeshSkinningData::UpdateBoneTransforms()
 					CurrTransforms[BoneIndex] = BoneTransform;
 				}
 
-				if (SkelMesh->GetRefBasesInvMatrix().IsValidIndex(BoneIndex))
+				if (SkinnedMesh->GetRefBasesInvMatrix().IsValidIndex(BoneIndex))
 				{
-					CurrBones[BoneIndex] = SkelMesh->GetRefBasesInvMatrix()[BoneIndex] * CurrTransforms[BoneIndex].ToMatrixWithScale();
+					CurrBones[BoneIndex] = SkinnedMesh->GetRefBasesInvMatrix()[BoneIndex] * CurrTransforms[BoneIndex].ToMatrixWithScale();
 				}
 				else
 				{
@@ -505,8 +546,8 @@ void FSkeletalMeshSkinningData::UpdateBoneTransforms()
 	}
 	else
 	{
-		SkelComp->CacheRefToLocalMatrices(CurrBones);
-		CurrTransforms = UE::LWC::ConvertArrayType<FTransform3f>(SkelComp->GetComponentSpaceTransforms());	// LWC_TODO: Perf pessimization
+		SkinnedComponent->CacheRefToLocalMatrices(CurrBones);
+		CurrTransforms = UE::LWC::ConvertArrayType<FTransform3f>(SkinnedComponent->GetComponentSpaceTransforms());	// LWC_TODO: Perf pessimization
 	}
 }
 
@@ -514,8 +555,9 @@ bool FSkeletalMeshSkinningData::Tick(float InDeltaSeconds, bool bRequirePreskin)
 {
 	FRWScopeLock Lock(RWGuard, SLT_Write);
 
-	USkeletalMeshComponent* SkelComp = MeshComp.Get();
-	check(SkelComp);
+	USkinnedMeshComponent* SkinnedComponent = WeakMeshComponent.Get();
+	check(SkinnedComponent);
+
 	DeltaSeconds = InDeltaSeconds;
 	CurrIndex ^= 1;
 
@@ -535,24 +577,50 @@ bool FSkeletalMeshSkinningData::Tick(float InDeltaSeconds, bool bRequirePreskin)
 		PrevComponentTransforms() = CurrComponentTransforms();
 	}
 
-	if (bRequirePreskin && SkelComp->GetSkeletalMeshAsset() != nullptr)
+	if (bRequirePreskin && SkinnedComponent->GetSkinnedAsset() != nullptr)
 	{
-		const USkeletalMesh* SkeletalMesh = SkelComp->GetSkeletalMeshAsset();
-		const FSkeletalMeshRenderData* RenderData = SkeletalMesh->GetResourceForRendering();
+		USkeletalMeshComponent* SkeletalComponent = Cast<USkeletalMeshComponent>(SkinnedComponent);
+		const USkinnedAsset* SkinnedMesh = SkinnedComponent->GetSkinnedAsset();
+		const FSkeletalMeshRenderData* RenderData = SkinnedMesh->GetResourceForRendering();
 		check(RenderData);
 
 		for (int32 LODIndex = RenderData->PendingFirstLODIdx; LODIndex < LODData.Num(); ++LODIndex)
 		{
 			FLODData& LOD = LODData[LODIndex];
-			if (LOD.PreSkinnedVertsUsers > 0 && SkeletalMesh->GetLODInfo(LODIndex)->bAllowCPUAccess)
+			if (LOD.PreSkinnedVertsUsers > 0 && SkinnedMesh->GetLODInfo(LODIndex)->bAllowCPUAccess)
 			{
 				// Increment ref count to prevent stream out for happening while we are processing the CPU data.
 				TRefCountPtr<const FSkeletalMeshLODRenderData> SkelMeshLODData = &RenderData->LODRenderData[LODIndex];
-				//TODO: If we pass the sections in the usage too, we can probably skin a minimal set of verts just for the used regions.
-				FSkinWeightVertexBuffer* SkinWeightBuffer = SkelComp->GetSkinWeightBuffer(LODIndex);
-				USkeletalMeshComponent::ComputeSkinnedPositions(SkelComp, CurrSkinnedPositions(LODIndex), CurrBoneRefToLocals(), *SkelMeshLODData, *SkinWeightBuffer);
-				USkeletalMeshComponent::ComputeSkinnedTangentBasis(SkelComp, CurrSkinnedTangentBasis(LODIndex), CurrBoneRefToLocals(), *SkelMeshLODData, *SkinWeightBuffer);
-				//check(CurrSkinnedPositions(LODIndex).Num() == SkelMeshLODData.NumVertices);
+
+				if (SkeletalComponent)
+				{
+					//TODO: If we pass the sections in the usage too, we can probably skin a minimal set of verts just for the used regions.
+					FSkinWeightVertexBuffer* SkinWeightBuffer = SkeletalComponent->GetSkinWeightBuffer(LODIndex);
+					USkeletalMeshComponent::ComputeSkinnedPositions(SkeletalComponent, CurrSkinnedPositions(LODIndex), CurrBoneRefToLocals(), *SkelMeshLODData, *SkinWeightBuffer);
+					USkeletalMeshComponent::ComputeSkinnedTangentBasis(SkeletalComponent, CurrSkinnedTangentBasis(LODIndex), CurrBoneRefToLocals(), *SkelMeshLODData, *SkinWeightBuffer);
+					//check(CurrSkinnedPositions(LODIndex).Num() == SkelMeshLODData.NumVertices);
+				}
+				else
+				{
+					//-TODO: Not supported currently
+					//USkinnedMeshComponent::ComputeSkinnedPositions(SkinnedComponent, CurrSkinnedPositions(LODIndex), CurrBoneRefToLocals(), SkelMeshLODData, *SkinWeightBuffer);
+					//USkinnedMeshComponent::ComputeSkinnedTangentBasis(SkinnedComponent, CurrSkinnedTangentBasis(LODIndex), CurrBoneRefToLocals(), SkelMeshLODData, *SkinWeightBuffer);
+					const int32 NumVertices = SkelMeshLODData->GetNumVertices();
+
+					TArray<FVector3f>& OutPositions = CurrSkinnedPositions(LODIndex);
+					TArray<FVector3f>& OutTangents = CurrSkinnedTangentBasis(LODIndex);
+
+					OutPositions.SetNumUninitialized(NumVertices);
+					OutTangents.SetNumUninitialized(NumVertices * 3);
+					for (int32 i = 0; i < NumVertices; ++i)
+					{
+						OutPositions[i] = FVector3f::ZeroVector;
+						OutTangents[i * 3 + 0] = FVector3f::XAxisVector;
+						OutTangents[i * 3 + 1] = FVector3f::YAxisVector;
+						OutTangents[i * 3 + 2] = FVector3f::ZAxisVector;
+					}
+				}
+
 				//Prime the previous positions if they're missing
 				if (PrevSkinnedPositions(LODIndex).Num() != CurrSkinnedPositions(LODIndex).Num())
 				{
@@ -572,7 +640,7 @@ bool FSkeletalMeshSkinningData::Tick(float InDeltaSeconds, bool bRequirePreskin)
 
 //////////////////////////////////////////////////////////////////////////
 
-FSkeletalMeshSkinningDataHandle FNDI_SkeletalMesh_GeneratedData::GetCachedSkinningData(TWeakObjectPtr<USkeletalMeshComponent>& Component, FSkeletalMeshSkinningDataUsage Usage, bool bNeedsDataImmediately)
+FSkeletalMeshSkinningDataHandle FNDI_SkeletalMesh_GeneratedData::GetCachedSkinningData(TWeakObjectPtr<USkinnedMeshComponent> Component, FSkeletalMeshSkinningDataUsage Usage, bool bNeedsDataImmediately)
 {
 	check(Component.Get() != nullptr);
 
@@ -602,7 +670,7 @@ void FNDI_SkeletalMesh_GeneratedData::Tick(ETickingGroup TickGroup, float DeltaS
 
 	// We may want to look at separating out how we manage the ticks here
 	//-OPT: Move into different arrays per tick group, manage promotions, demotions, etc, or add ourselves as a subsequent of the component's tick
-	TArray<TWeakObjectPtr<USkeletalMeshComponent>, TInlineAllocator<32>> ToRemove;
+	TArray<TWeakObjectPtr<USkinnedMeshComponent>, TInlineAllocator<32>> ToRemove;
 	TArray<FSkeletalMeshSkinningData*, TInlineAllocator<32>> ToTickBonesOnly;
 	TArray<FSkeletalMeshSkinningData*, TInlineAllocator<32>> ToTickPreskin;
 	const bool bForceTick = TickGroup == NiagaraLastTickGroup;
@@ -610,9 +678,9 @@ void FNDI_SkeletalMesh_GeneratedData::Tick(ETickingGroup TickGroup, float DeltaS
 	ToTickBonesOnly.Reserve(CachedSkinningData.Num());
 	ToTickPreskin.Reserve(CachedSkinningData.Num());
 
-	for (TPair<TWeakObjectPtr<USkeletalMeshComponent>, TSharedPtr<FSkeletalMeshSkinningData>>& Pair : CachedSkinningData)
+	for (TPair<TWeakObjectPtr<USkinnedMeshComponent>, TSharedPtr<FSkeletalMeshSkinningData>>& Pair : CachedSkinningData)
 	{
-		USkeletalMeshComponent* Component = Pair.Key.Get();
+		USkinnedMeshComponent* SkinnedComponent = Pair.Key.Get();
 		TSharedPtr<FSkeletalMeshSkinningData>& SkinningData = Pair.Value;
 
 		if ( TickGroup == NiagaraFirstTickGroup )
@@ -621,7 +689,7 @@ void FNDI_SkeletalMesh_GeneratedData::Tick(ETickingGroup TickGroup, float DeltaS
 		}
 
 		// Should remove?
-		if ( (Component == nullptr) || SkinningData.IsUnique() || !SkinningData->IsUsed() )
+		if ( (SkinnedComponent == nullptr) || SkinningData.IsUnique() || !SkinningData->IsUsed() )
 		{
 			ToRemove.Add(Pair.Key);
 			continue;
@@ -635,7 +703,8 @@ void FNDI_SkeletalMesh_GeneratedData::Tick(ETickingGroup TickGroup, float DeltaS
 		// Has ticked or can be ticked
 		if (bForceTick == false)
 		{
-			const ETickingGroup PrereqTickGroup = NDISkelMeshLocal::GetComponentTickGroup(Component);
+			USkeletalMeshComponent* SkeletalComponent = Cast<USkeletalMeshComponent>(SkinnedComponent);
+			const ETickingGroup PrereqTickGroup = NDISkelMeshLocal::GetComponentTickGroup(SkinnedComponent, SkeletalComponent);
 			if ( PrereqTickGroup > TickGroup )
 			{
 				continue;
@@ -658,7 +727,7 @@ void FNDI_SkeletalMesh_GeneratedData::Tick(ETickingGroup TickGroup, float DeltaS
 		}
 	}
 
-	for (TWeakObjectPtr<USkeletalMeshComponent> Key : ToRemove)
+	for (TWeakObjectPtr<USkinnedMeshComponent> Key : ToRemove)
 	{
 		CachedSkinningData.Remove(Key);
 	}
@@ -726,7 +795,7 @@ void FNDI_SkeletalMesh_GeneratedData::Tick(ETickingGroup TickGroup, float DeltaS
 	}
 }
 
-FSkeletalMeshUvMappingHandle FNDI_SkeletalMesh_GeneratedData::GetCachedUvMapping(TWeakObjectPtr<USkeletalMesh>& MeshObject, int32 InLodIndex, int32 InUvSetIndex, FMeshUvMappingUsage Usage, bool bNeedsDataImmediately)
+FSkeletalMeshUvMappingHandle FNDI_SkeletalMesh_GeneratedData::GetCachedUvMapping(TWeakObjectPtr<USkinnedAsset> MeshObject, int32 InLodIndex, int32 InUvSetIndex, FMeshUvMappingUsage Usage, bool bNeedsDataImmediately)
 {
 	check(MeshObject.Get() != nullptr);
 
@@ -763,7 +832,7 @@ FSkeletalMeshUvMappingHandle FNDI_SkeletalMesh_GeneratedData::GetCachedUvMapping
 }
 
 
-FSkeletalMeshConnectivityHandle FNDI_SkeletalMesh_GeneratedData::GetCachedConnectivity(TWeakObjectPtr<USkeletalMesh>& MeshObject, int32 InLodIndex, FSkeletalMeshConnectivityUsage Usage, bool bNeedsDataImmediately)
+FSkeletalMeshConnectivityHandle FNDI_SkeletalMesh_GeneratedData::GetCachedConnectivity(TWeakObjectPtr<USkinnedAsset> MeshObject, int32 InLodIndex, FSkeletalMeshConnectivityUsage Usage, bool bNeedsDataImmediately)
 {
 	check(MeshObject.Get() != nullptr);
 
@@ -827,7 +896,7 @@ void FSkeletalMeshGpuSpawnStaticBuffers::Initialise(FNDISkeletalMesh_InstanceDat
 
 	if (InstData)
 	{
-		SceneComponent = InstData->SceneComponent;
+		WeakSceneComponent = InstData->WeakSceneComponent;
 
 		SkeletalMeshSamplingLODBuiltData = MeshSamplingLODBuiltData;
 		bUseGpuUniformlyDistributedSampling = InstData->bIsGpuUniformlyDistributedSampling;
@@ -852,7 +921,7 @@ void FSkeletalMeshGpuSpawnStaticBuffers::Initialise(FNDISkeletalMesh_InstanceDat
 			const int32 NumAreaSamples = SkeletalMeshSamplingLODBuiltData->AreaWeightedTriangleSampler.GetNumEntries();
 			if (NumAreaSamples != TriangleCount)
 			{
-				UE_LOG(LogNiagara, Warning, TEXT("FSkeletalMeshGpuSpawnStaticBuffers> AreaWeighted Triangle Sampling Count (%d) does not match triangle count (%d), disabling uniform sampling for SkelMesh(%s) System(%s)"), NumAreaSamples, TriangleCount, *GetFullNameSafe(InstData->SkeletalMesh.Get()), *GetFullNameSafe(SystemInstance->GetSystem()));
+				UE_LOG(LogNiagara, Warning, TEXT("FSkeletalMeshGpuSpawnStaticBuffers> AreaWeighted Triangle Sampling Count (%d) does not match triangle count (%d), disabling uniform sampling for SkelMesh(%s) System(%s)"), NumAreaSamples, TriangleCount, *GetFullNameSafe(InstData->WeakSkinnedMesh.Get()), *GetFullNameSafe(SystemInstance->GetSystem()));
 				bUseGpuUniformlyDistributedSampling = false;
 			}
 		}
@@ -875,45 +944,48 @@ void FSkeletalMeshGpuSpawnStaticBuffers::Initialise(FNDISkeletalMesh_InstanceDat
 		// Create triangle / vertex region sampling data
 		if (InstData->SamplingRegionIndices.Num() > 0)
 		{
-			const FSkeletalMeshSamplingInfo& SamplingInfo = InstData->SkeletalMesh->GetSamplingInfo();
-
-			// Count required regions
-			bSamplingRegionsAllAreaWeighted = true;
-			NumSamplingRegionTriangles = 0;
-			NumSamplingRegionVertices = 0;
-
-			for (const int32 RegionIndex : InstData->SamplingRegionIndices)
+			if (USkeletalMesh* SkeletalMesh = Cast<USkeletalMesh>(InstData->WeakSkinnedMesh.Get()))
 			{
-				const FSkeletalMeshSamplingRegionBuiltData& SamplingRegionBuildData = SamplingInfo.GetRegionBuiltData(RegionIndex);
-				NumSamplingRegionTriangles += SamplingRegionBuildData.TriangleIndices.Num();
-				NumSamplingRegionVertices += SamplingRegionBuildData.Vertices.Num();
-				bSamplingRegionsAllAreaWeighted &= SamplingRegionBuildData.AreaWeightedSampler.GetNumEntries() == SamplingRegionBuildData.TriangleIndices.Num();
-			}
+				const FSkeletalMeshSamplingInfo& SamplingInfo = SkeletalMesh->GetSamplingInfo();
 
-			// Build buffers
-			SampleRegionsProbAlias.AddUninitialized(NDISkelMeshLocal::GetProbAliasDWORDSize(NumSamplingRegionTriangles));
-			SampleRegionsTriangleIndicies.Reserve(NumSamplingRegionTriangles);
-			SampleRegionsVertices.Reserve(NumSamplingRegionVertices);
+				// Count required regions
+				bSamplingRegionsAllAreaWeighted = true;
+				NumSamplingRegionTriangles = 0;
+				NumSamplingRegionVertices = 0;
 
-			int32 RegionOffset = 0;
-			int32 PABufferOffset = 0;
-			for (const int32 RegionIndex : InstData->SamplingRegionIndices)
-			{
-				const FSkeletalMeshSamplingRegionBuiltData& SamplingRegionBuildData = SamplingInfo.GetRegionBuiltData(RegionIndex);
-				if (bSamplingRegionsAllAreaWeighted)
+				for (const int32 RegionIndex : InstData->SamplingRegionIndices)
 				{
-					NDISkelMeshLocal::PackProbAlias(SampleRegionsProbAlias.GetData() + PABufferOffset, SamplingRegionBuildData.AreaWeightedSampler, RegionOffset);
-					PABufferOffset += NDISkelMeshLocal::GetProbAliasDWORDSize(SamplingRegionBuildData.AreaWeightedSampler.GetNumEntries());
+					const FSkeletalMeshSamplingRegionBuiltData& SamplingRegionBuildData = SamplingInfo.GetRegionBuiltData(RegionIndex);
+					NumSamplingRegionTriangles += SamplingRegionBuildData.TriangleIndices.Num();
+					NumSamplingRegionVertices += SamplingRegionBuildData.Vertices.Num();
+					bSamplingRegionsAllAreaWeighted &= SamplingRegionBuildData.AreaWeightedSampler.GetNumEntries() == SamplingRegionBuildData.TriangleIndices.Num();
 				}
-				for (int v : SamplingRegionBuildData.TriangleIndices)
-				{
-					SampleRegionsTriangleIndicies.Add(v / 3);
-				}
-				for (int v : SamplingRegionBuildData.Vertices)
+
+				// Build buffers
+				SampleRegionsProbAlias.AddUninitialized(NDISkelMeshLocal::GetProbAliasDWORDSize(NumSamplingRegionTriangles));
+				SampleRegionsTriangleIndicies.Reserve(NumSamplingRegionTriangles);
+				SampleRegionsVertices.Reserve(NumSamplingRegionVertices);
+
+				int32 RegionOffset = 0;
+				int32 PABufferOffset = 0;
+				for (const int32 RegionIndex : InstData->SamplingRegionIndices)
 				{
-					SampleRegionsVertices.Add(v);
+					const FSkeletalMeshSamplingRegionBuiltData& SamplingRegionBuildData = SamplingInfo.GetRegionBuiltData(RegionIndex);
+					if (bSamplingRegionsAllAreaWeighted)
+					{
+						NDISkelMeshLocal::PackProbAlias(SampleRegionsProbAlias.GetData() + PABufferOffset, SamplingRegionBuildData.AreaWeightedSampler, RegionOffset);
+						PABufferOffset += NDISkelMeshLocal::GetProbAliasDWORDSize(SamplingRegionBuildData.AreaWeightedSampler.GetNumEntries());
+					}
+					for (int v : SamplingRegionBuildData.TriangleIndices)
+					{
+						SampleRegionsTriangleIndicies.Add(v / 3);
+					}
+					for (int v : SamplingRegionBuildData.Vertices)
+					{
+						SampleRegionsVertices.Add(v);
+					}
+					RegionOffset += SamplingRegionBuildData.TriangleIndices.Num();
 				}
-				RegionOffset += SamplingRegionBuildData.TriangleIndices.Num();
 			}
 		}
 	}
@@ -1268,22 +1340,22 @@ void FSkeletalMeshGpuDynamicBufferProxy::ReleaseRHI()
 void FSkeletalMeshGpuDynamicBufferProxy::NewFrame(const FNDISkeletalMesh_InstanceData* InstanceData, int32 LODIndex)
 {
 	// Grab Skeletal Component / Mesh, we must have a mesh at minimum to set the data
-	USkeletalMeshComponent* SkelComp = nullptr;
-	USkeletalMesh* SkelMesh = nullptr;
+	USkinnedMeshComponent* SkinnedComponent = nullptr;
+	USkinnedAsset* SkinnedMesh = nullptr;
 	if (InstanceData != nullptr)
 	{
-		SkelComp = Cast<USkeletalMeshComponent>(InstanceData->SceneComponent.Get());
-		if ( SkelComp != nullptr )
+		SkinnedComponent = Cast<USkinnedMeshComponent>(InstanceData->WeakSceneComponent.Get());
+		if (SkinnedComponent != nullptr )
 		{
-			SkelMesh = SkelComp->GetSkeletalMeshAsset();
+			SkinnedMesh = SkinnedComponent->GetSkinnedAsset();
 		}
-		if (SkelMesh == nullptr)
+		if (SkinnedMesh == nullptr)
 		{
-			SkelMesh = InstanceData->SkeletalMesh.Get();
+			SkinnedMesh = InstanceData->WeakSkinnedMesh.Get();
 		}
 	}
 
-	if ( SkelMesh == nullptr )
+	if (SkinnedMesh == nullptr )
 	{
 		return;
 	}
@@ -1299,7 +1371,7 @@ void FSkeletalMeshGpuDynamicBufferProxy::NewFrame(const FNDISkeletalMesh_Instanc
 			check(BoneTransforms.Num() == SamplingBoneCount);
 
 			// Fill AllSectionsRefToLocalMatrices
-			TIndirectArray<FSkeletalMeshLODRenderData>& LODRenderDataArray = SkelMesh->GetResourceForRendering()->LODRenderData;
+			TIndirectArray<FSkeletalMeshLODRenderData>& LODRenderDataArray = SkinnedMesh->GetResourceForRendering()->LODRenderData;
 			check(0 <= LODIndex && LODIndex < LODRenderDataArray.Num());
 			FSkeletalMeshLODRenderData& LODRenderData = LODRenderDataArray[LODIndex];
 			TArray<FSkelMeshRenderSection>& Sections = LODRenderData.RenderSections;
@@ -1322,7 +1394,7 @@ void FSkeletalMeshGpuDynamicBufferProxy::NewFrame(const FNDISkeletalMesh_Instanc
 					const int32 BoneIndex = Section.BoneMap[m];
 					const FTransform& BoneTransform = BoneTransforms[BoneIndex];
 					// LWC_TODO: precision loss
-					const FMatrix44f BoneMatrix = SkelMesh->GetRefBasesInvMatrix().IsValidIndex(BoneIndex) ? SkelMesh->GetRefBasesInvMatrix()[BoneIndex] * (FMatrix44f)BoneTransform.ToMatrixWithScale() : (FMatrix44f)BoneTransform.ToMatrixWithScale();
+					const FMatrix44f BoneMatrix = SkinnedMesh->GetRefBasesInvMatrix().IsValidIndex(BoneIndex) ? SkinnedMesh->GetRefBasesInvMatrix()[BoneIndex] * (FMatrix44f)BoneTransform.ToMatrixWithScale() : (FMatrix44f)BoneTransform.ToMatrixWithScale();
 					BoneMatrix.To3x4MatrixTranspose(&AllSectionsRefToLocalMatrices[Float4Count].X);
 					Float4Count += 3;
 				}
@@ -1355,12 +1427,12 @@ void FSkeletalMeshGpuDynamicBufferProxy::NewFrame(const FNDISkeletalMesh_Instanc
 		};
 
 	// If we have a component pull transforms from component otherwise grab from skel mesh
-	if (SkelComp)
+	if (SkinnedComponent)
 	{
-		if (USkinnedMeshComponent* LeaderComponent = SkelComp->LeaderPoseComponent.Get())
+		if (USkinnedMeshComponent* LeaderComponent = SkinnedComponent->LeaderPoseComponent.Get())
 		{
 			const FReferenceSkeleton* ReferenceSkeleton = nullptr;
-			const TArray<int32>& LeaderBoneMap = SkelComp->GetLeaderBoneMap();
+			const TArray<int32>& LeaderBoneMap = SkinnedComponent->GetLeaderBoneMap();
 			const int32 NumBones = LeaderBoneMap.Num();
 
 			TArray<FTransform> TempBoneTransforms;
@@ -1373,7 +1445,7 @@ void FSkeletalMeshGpuDynamicBufferProxy::NewFrame(const FNDISkeletalMesh_Instanc
 			}
 			else
 			{
-				ReferenceSkeleton = &SkelMesh->GetRefSkeleton();
+				ReferenceSkeleton = &SkinnedMesh->GetRefSkeleton();
 				const TArray<FTransform>& LeaderTransforms = LeaderComponent->GetComponentSpaceTransforms();
 				for (int32 BoneIndex=0; BoneIndex < NumBones; ++BoneIndex)
 				{
@@ -1400,16 +1472,23 @@ void FSkeletalMeshGpuDynamicBufferProxy::NewFrame(const FNDISkeletalMesh_Instanc
 		}
 		else
 		{
-			const TArray<FTransform>& ComponentTransforms = SkelComp->GetComponentSpaceTransforms();
+			const TArray<FTransform>& ComponentTransforms = SkinnedComponent->GetComponentSpaceTransforms();
 			if (ComponentTransforms.Num() > 0)
 			{
-				FillBuffers(ComponentTransforms, &SkelMesh->GetRefSkeleton());
+				FillBuffers(ComponentTransforms, &SkinnedMesh->GetRefSkeleton());
 			}
 			else
 			{
 				// Trying to catch cause of this case in the wild. Not supposed to be possible with a valid skeletal mesh
-				ensureMsgf(false, TEXT("NiagaraSkelMeshDI: Mesh has no ComponentSpaceTransforms. Component - %s (Registered: %s, Flags: %d), Mesh - %s (Flags: %d)"),
-					*GetFullNameSafe(SkelComp), SkelComp->IsRegistered() ? TEXT("Yes") : TEXT("No"), SkelComp->GetFlags(), *GetFullNameSafe(SkelMesh), SkelMesh->GetFlags());
+				ensureMsgf(
+					false,
+					TEXT("NiagaraSkelMeshDI: Mesh has no ComponentSpaceTransforms. Component - %s (Registered: %s, Flags: %d), Mesh - %s (Flags: %d)"),
+					*GetFullNameSafe(SkinnedComponent),
+					SkinnedComponent && SkinnedComponent->IsRegistered() ? TEXT("Yes") : TEXT("No"),
+					SkinnedComponent ? SkinnedComponent->GetFlags() : 0,
+					*GetFullNameSafe(SkinnedMesh),
+					SkinnedMesh ? SkinnedMesh->GetFlags() : 0
+				);
 
 				TArray<FTransform> TempBoneTransforms;
 				TempBoneTransforms.AddDefaulted(SamplingBoneCount);
@@ -1420,11 +1499,11 @@ void FSkeletalMeshGpuDynamicBufferProxy::NewFrame(const FNDISkeletalMesh_Instanc
 	else
 	{
 		//-TODO: Opt and combine with MaterPoseComponent
-		const FReferenceSkeleton* ReferenceSkeleton = &SkelMesh->GetRefSkeleton();
+		const FReferenceSkeleton* ReferenceSkeleton = &SkinnedMesh->GetRefSkeleton();
 		TArray<FTransform> TempBoneTransforms;
 		TempBoneTransforms.Reserve(SamplingBoneCount);
 
-		const TArray<FTransform>& RefTransforms = SkelMesh->GetRefSkeleton().GetRefBonePose();
+		const TArray<FTransform>& RefTransforms = SkinnedMesh->GetRefSkeleton().GetRefBonePose();
 		for (int32 i=0; i < RefTransforms.Num(); ++i)
 		{
 			FTransform BoneTransform = RefTransforms[i];
@@ -1530,60 +1609,68 @@ void UNiagaraDataInterfaceSkeletalMesh::ProvidePerInstanceDataForRenderThread(vo
 	Data->ConnectivityBuffer = SourceData->Connectivity.GetProxy();
 }
 
-USkeletalMesh* UNiagaraDataInterfaceSkeletalMesh::GetSkeletalMesh(FNiagaraSystemInstance* SystemInstance, USceneComponent* AttachComponent, TWeakObjectPtr<USceneComponent>& SceneComponent, USkeletalMeshComponent*& FoundSkelComp, FNDISkeletalMesh_InstanceData* InstData)
+USkinnedAsset* UNiagaraDataInterfaceSkeletalMesh::GetSkinnedMesh(FNiagaraSystemInstance* SystemInstance, USceneComponent* AttachComponent, TWeakObjectPtr<USceneComponent>& SceneComponent, USkinnedMeshComponent*& FoundSkinnedComponent, FNDISkeletalMesh_InstanceData* InstData)
 {
-	auto IsValidComponent = [&](const USkeletalMeshComponent* Component, bool bAllowNullMesh=false)
-	{
-		return IsValid(Component) && (bAllowNullMesh || IsValid(Component->GetSkeletalMeshAsset())) &&
-			(ComponentTags.IsEmpty()
-			|| ComponentTags.ContainsByPredicate([&](const FName& Tag) { return Tag == NAME_None || Component->ComponentHasTag(Tag); }));
-	};
+	UClass* TargetComponentClass = OptionalComponentClass.Get() ? OptionalComponentClass.Get() : USkeletalMeshComponent::StaticClass();
 
-	// Helper to scour an actor (or its parents) for a valid skeletal mesh component
-	auto FindActorSkelMeshComponent = [&](AActor* Actor, bool bRecurseParents = false) -> USkeletalMeshComponent*
-	{
-		if (ASkeletalMeshActor* SkelMeshActor = Cast<ASkeletalMeshActor>(Actor))
+	auto ConvertToValidComponent =
+		[&](USceneComponent* Component, bool bAllowNullMesh=false) -> USkinnedMeshComponent*
 		{
-			USkeletalMeshComponent* Comp = SkelMeshActor->GetSkeletalMeshComponent();
-			if (IsValidComponent(Comp, true))
+			if (IsValid(Component) && Component->IsA(TargetComponentClass) )
 			{
-				return Comp;
+				USkinnedMeshComponent* SkinnedComponent = CastChecked<USkinnedMeshComponent>(Component);
+				if ((bAllowNullMesh || IsValid(SkinnedComponent->GetSkinnedAsset())) &&
+					(ComponentTags.IsEmpty() || ComponentTags.ContainsByPredicate([&](const FName& Tag) { return Tag == NAME_None || Component->ComponentHasTag(Tag); })))
+				{
+					return SkinnedComponent;
+				}
 			}
-		}
+			return nullptr;
+		};
 
-		// Fall back on any valid component on the actor
-		while (Actor)
+	// Helper to scour an actor (or its parents) for a valid skeletal mesh component
+	auto FindActorSkelMeshComponent =
+		[&](AActor* Actor, bool bRecurseParents = false) -> USkinnedMeshComponent*
 		{
-			if (USceneComponent* ActorRootComponent = Actor->GetRootComponent())
+			if (ASkeletalMeshActor* SkelMeshActor = Cast<ASkeletalMeshActor>(Actor))
 			{
-				TArray<USceneComponent*> ActorComponents;
-				ActorRootComponent->GetChildrenComponents(true /*bIncludeAllDescendants*/, ActorComponents);
-				ActorComponents.Insert(ActorRootComponent, 0);
+				USkinnedMeshComponent* Comp = SkelMeshActor->GetSkeletalMeshComponent();
+				if (ConvertToValidComponent(Comp, true))
+				{
+					return Comp;
+				}
+			}
 
-				for (USceneComponent* SceneComponent : ActorComponents)
+			// Fall back on any valid component on the actor
+			while (Actor)
+			{
+				if (USceneComponent* ActorRootComponent = Actor->GetRootComponent())
 				{
-					if (USkeletalMeshComponent* SkeletalComp = Cast<USkeletalMeshComponent>(SceneComponent))
+					TArray<USceneComponent*> ActorComponents;
+					ActorRootComponent->GetChildrenComponents(true /*bIncludeAllDescendants*/, ActorComponents);
+					ActorComponents.Insert(ActorRootComponent, 0);
+
+					for (USceneComponent* SceneComponent : ActorComponents)
 					{
-						if (IsValidComponent(SkeletalComp))
+						if (USkinnedMeshComponent* SkinnedComponent = ConvertToValidComponent(SceneComponent))
 						{
-							return SkeletalComp;
+							return SkinnedComponent;
 						}
 					}
 				}
-			}
 
-			if (bRecurseParents)
-			{
-				Actor = Actor->GetParentActor();
-			}
-			else
-			{
-				break;
+				if (bRecurseParents)
+				{
+					Actor = Actor->GetParentActor();
+				}
+				else
+				{
+					break;
+				}
 			}
-		}
 
-		return nullptr;
-	};
+			return nullptr;
+		};
 
 	const bool bTrySource = SourceMode == ENDISkeletalMesh_SourceMode::Default || SourceMode == ENDISkeletalMesh_SourceMode::Source;
 	const bool bTryAttachParent = SourceMode == ENDISkeletalMesh_SourceMode::Default || SourceMode == ENDISkeletalMesh_SourceMode::AttachParent;
@@ -1601,12 +1688,12 @@ USkeletalMesh* UNiagaraDataInterfaceSkeletalMesh::GetSkeletalMesh(FNiagaraSystem
 			{
 				if (IsValid(UserSkelMeshComp))
 				{
-					FoundSkelComp = UserSkelMeshComp;
+					FoundSkinnedComponent = UserSkelMeshComp;
 				}
 			}
 			else if (AActor* Actor = Cast<AActor>(UserParamObject))
 			{
-				FoundSkelComp = FindActorSkelMeshComponent(Actor);
+				FoundSkinnedComponent = FindActorSkelMeshComponent(Actor);
 			}
 			else
 			{
@@ -1624,47 +1711,46 @@ USkeletalMesh* UNiagaraDataInterfaceSkeletalMesh::GetSkeletalMesh(FNiagaraSystem
 	}
 	else if (bTrySource && IsValid(SourceComponent))
 	{
-		FoundSkelComp = SourceComponent;
+		FoundSkinnedComponent = SourceComponent;
 	}
 	else if (bTrySource && SoftSourceActor.Get())
 	{
-		FoundSkelComp = FindActorSkelMeshComponent(SoftSourceActor.Get());
+		FoundSkinnedComponent = FindActorSkelMeshComponent(SoftSourceActor.Get());
 	}
 	else if (bTryAttachParent && AttachComponent)
 	{
 		// First, try to find the mesh component up the attachment hierarchy
 		for (USceneComponent* Curr = AttachComponent; Curr; Curr = Curr->GetAttachParent())
 		{
-			USkeletalMeshComponent* ParentComp = Cast<USkeletalMeshComponent>(Curr);
-			if (IsValidComponent(ParentComp, true))
+			if (USkinnedMeshComponent* SkinnedComponent = ConvertToValidComponent(Curr, true))
 			{
-				FoundSkelComp = ParentComp;
+				FoundSkinnedComponent = SkinnedComponent;
 				break;
 			}
 		}
 
-		if (!FoundSkelComp)
+		if (!FoundSkinnedComponent)
 		{
 			// Next, try to find one in our outer chain
-			USkeletalMeshComponent* OuterComp = AttachComponent->GetTypedOuter<USkeletalMeshComponent>();
-			if (IsValidComponent(OuterComp, true))
+			USkinnedMeshComponent* OuterComp = Cast<USkinnedMeshComponent>(AttachComponent->GetTypedOuter(TargetComponentClass));
+			if (ConvertToValidComponent(OuterComp, true))
 			{
-				FoundSkelComp = OuterComp;
+				FoundSkinnedComponent = OuterComp;
 			}
 			else if (AActor* Actor = AttachComponent->GetAttachmentRootActor())
 			{
 				// Final fall-back, look for any mesh component on our root actor or any of its parents
-				FoundSkelComp = FindActorSkelMeshComponent(Actor, true);
+				FoundSkinnedComponent = FindActorSkelMeshComponent(Actor, true);
 			}
 		}
 	}
 
-	USkeletalMesh* Mesh = nullptr;
+	USkinnedAsset* Mesh = nullptr;
 	SceneComponent = nullptr;
-	if (FoundSkelComp)
+	if (FoundSkinnedComponent)
 	{
-		Mesh = FoundSkelComp->GetSkeletalMeshAsset();
-		SceneComponent = FoundSkelComp;
+		Mesh = FoundSkinnedComponent->GetSkinnedAsset();
+		SceneComponent = FoundSkinnedComponent;
 	}
 	else if (bTryDefaultMesh && DefaultMesh)
 	{
@@ -1681,12 +1767,12 @@ USkeletalMesh* UNiagaraDataInterfaceSkeletalMesh::GetSkeletalMesh(FNiagaraSystem
 	return Mesh;
 }
 
-USkeletalMesh* UNiagaraDataInterfaceSkeletalMesh::GetSkeletalMesh(UNiagaraComponent* Component)
+USkinnedAsset* UNiagaraDataInterfaceSkeletalMesh::GetSkinnedMesh(UNiagaraComponent* Component)
 {
 	// NOTE: We don't need the system instance when not initializing instance data, and when using a UNiagaraComponent, it is always the attach component
 	TWeakObjectPtr<USceneComponent> SceneComponent;
-	USkeletalMeshComponent* FoundSkelComp = nullptr;
-	return GetSkeletalMesh(nullptr, Component, SceneComponent, FoundSkelComp, nullptr);
+	USkinnedMeshComponent* FoundSkinnedComponent = nullptr;
+	return GetSkinnedMesh(nullptr, Component, SceneComponent, FoundSkinnedComponent, nullptr);
 }
 
 bool FNDISkeletalMesh_InstanceData::Init(UNiagaraDataInterfaceSkeletalMesh* Interface, FNiagaraSystemInstance* SystemInstance)
@@ -1695,10 +1781,10 @@ bool FNDISkeletalMesh_InstanceData::Init(UNiagaraDataInterfaceSkeletalMesh* Inte
 	check(SystemInstance);
 
 	// Initialize members
-	SceneComponent = nullptr;
-	CachedAttachParent = nullptr;
+	WeakSceneComponent = nullptr;
+	WeakCachedAttachParent = nullptr;
 	CachedUserParam = nullptr;
-	SkeletalMesh = nullptr;
+	WeakSkinnedMesh = nullptr;
 	Transform = FMatrix::Identity;
 	TransformInverseTransposed = FMatrix::Identity;
 	PrevTransform = FMatrix::Identity;
@@ -1717,13 +1803,13 @@ bool FNDISkeletalMesh_InstanceData::Init(UNiagaraDataInterfaceSkeletalMesh* Inte
 	PreSkinnedLocalBoundsExtents = FVector3f::ZeroVector;
 
 	// Get skel mesh and confirm have valid data
-	USkeletalMeshComponent* NewSkelComp = nullptr;
+	USkinnedMeshComponent* NewSkinnedComponent = nullptr;
 	USceneComponent* AttachComponent = SystemInstance->GetAttachComponent();
-	USkeletalMesh* Mesh = Interface->GetSkeletalMesh(SystemInstance, AttachComponent, SceneComponent, NewSkelComp, this);
+	USkinnedAsset* SkinnedMesh = Interface->GetSkinnedMesh(SystemInstance, AttachComponent, WeakSceneComponent, NewSkinnedComponent, this);
 
-	SkeletalMesh = Mesh;
-	bMeshValid = Mesh != nullptr;
-	bComponentValid = SceneComponent.IsValid();
+	WeakSkinnedMesh = SkinnedMesh;
+	bMeshValid = SkinnedMesh != nullptr;
+	bComponentValid = WeakSceneComponent.IsValid();
 
 	const FTransform ComponentTransform = CalculateComponentTransform(SystemInstance);
 	Transform = ComponentTransform.ToMatrixWithScale();
@@ -1732,7 +1818,7 @@ bool FNDISkeletalMesh_InstanceData::Init(UNiagaraDataInterfaceSkeletalMesh* Inte
 
 	if (AttachComponent)
 	{
-		CachedAttachParent = AttachComponent->GetAttachParent();
+		WeakCachedAttachParent = AttachComponent->GetAttachParent();
 	}
 
 	bResetOnLODStreamedIn = false;
@@ -1743,15 +1829,15 @@ bool FNDISkeletalMesh_InstanceData::Init(UNiagaraDataInterfaceSkeletalMesh* Inte
 	SamplingRegionIndices.Empty();
 	bool bAllRegionsAreAreaWeighting = true;
 
-	if (Mesh)
+	if (SkinnedMesh)
 	{
 		// Determine the LOD index and sampling region indices
-		const FStreamableRenderResourceState& SRRState = Mesh->GetStreamableResourceState();		
-		const int32 NumValidLODs = FMath::Min(SRRState.NumRequestedLODs, SRRState.NumResidentLODs);
+		const TTuple<int32, int32> StreamingLODData = NDISkelMeshLocal::GetStreamingLODData(SkinnedMesh);
+		const int32 NumValidLODs = StreamingLODData.Get<0>();
 		if (NumValidLODs > 0)
 		{
-			const int32 CurrentFirstLOD = SRRState.LODCountToAssetFirstLODIdx(NumValidLODs);
-			const int32 DesiredLODIndex = Interface->CalculateLODIndexAndSamplingRegions(Mesh, SamplingRegionIndices, bAllRegionsAreAreaWeighting);
+			const int32 CurrentFirstLOD = StreamingLODData.Get<1>();
+			const int32 DesiredLODIndex = Interface->CalculateLODIndexAndSamplingRegions(SkinnedMesh, SamplingRegionIndices, bAllRegionsAreAreaWeighting);
 			if (DesiredLODIndex != INDEX_NONE)
 			{
 				if (DesiredLODIndex >= CurrentFirstLOD)
@@ -1765,7 +1851,7 @@ bool FNDISkeletalMesh_InstanceData::Init(UNiagaraDataInterfaceSkeletalMesh* Inte
 				}
 
 				// Attempt to cache the LOD
-				if (FSkeletalMeshRenderData* RenderData = Mesh->GetResourceForRendering())
+				if (FSkeletalMeshRenderData* RenderData = SkinnedMesh->GetResourceForRendering())
 				{
 					if (RenderData->LODRenderData.IsValidIndex(CachedLODIdx))
 					{
@@ -1777,14 +1863,14 @@ bool FNDISkeletalMesh_InstanceData::Init(UNiagaraDataInterfaceSkeletalMesh* Inte
 						// NOTE: Assumption here is that the LOD render data is cacheable from GameThread as long as it's considered resident by
 						// the StreamableRenderResourceState on GameThread. If this warning gets hit, that assumption has become incorrect.
 						UE_LOG(LogNiagara, Log, TEXT("SkeletalMesh data interface failed to cache LOD %d. Sampling will fail. %s"), CachedLODIdx, *Interface->GetFullName());
-						Mesh = nullptr;
+						SkinnedMesh = nullptr;
 					}
 				}
 				else
 				{
 					// Warn and continue as if the component has no mesh
 					UE_LOG(LogNiagara, Log, TEXT("SkeletalMesh data interface with no Render data. Sampling will fail. %s"), *Interface->GetFullName());
-					Mesh = nullptr;
+					SkinnedMesh = nullptr;
 				}
 			}
 			else
@@ -1796,10 +1882,10 @@ bool FNDISkeletalMesh_InstanceData::Init(UNiagaraDataInterfaceSkeletalMesh* Inte
 		{
 			// Warn and continue as if the component has no mesh
 			UE_LOG(LogNiagara, Log, TEXT("SkeletalMesh data interface with no resident LODs. Sampling will fail. %s"), *Interface->GetFullName());
-			Mesh = nullptr;
+			SkinnedMesh = nullptr;
 		}
 
-		if (Mesh == nullptr)
+		if (SkinnedMesh == nullptr)
 		{
 			CachedLODIdx = 0;
 			bResetOnLODStreamedIn = false;
@@ -1810,8 +1896,11 @@ bool FNDISkeletalMesh_InstanceData::Init(UNiagaraDataInterfaceSkeletalMesh* Inte
 			// HACK! This only works on systems created by a Niagara component...should maybe move somewhere else to cover non-component systems
 			if (UNiagaraComponent* NiagaraComponent = Cast<UNiagaraComponent>(AttachComponent))
 			{
-				Mesh->GetOnMeshChanged().AddUObject(NiagaraComponent, &UNiagaraComponent::ReinitializeSystem);
-				if (USkeleton* Skeleton = Mesh->GetSkeleton())
+				if (USkeletalMesh* SkeletalMesh = Cast<USkeletalMesh>(SkinnedMesh))
+				{
+					SkeletalMesh->GetOnMeshChanged().AddUObject(NiagaraComponent, &UNiagaraComponent::ReinitializeSystem);
+				}
+				if (USkeleton* Skeleton = SkinnedMesh->GetSkeleton())
 				{
 					Skeleton->RegisterOnSkeletonHierarchyChanged(USkeleton::FOnSkeletonHierarchyChanged::CreateUObject(NiagaraComponent, &UNiagaraComponent::ReinitializeSystem));
 				}
@@ -1819,9 +1908,9 @@ bool FNDISkeletalMesh_InstanceData::Init(UNiagaraDataInterfaceSkeletalMesh* Inte
 		}
 #endif
 
-		if ( Mesh != nullptr )
+		if (SkinnedMesh != nullptr )
 		{
-			const FBoxSphereBounds LocalBounds = Mesh->GetBounds();
+			const FBoxSphereBounds LocalBounds = SkinnedMesh->GetBounds();
 			PreSkinnedLocalBoundsCenter = FVector3f(LocalBounds.Origin);
 			PreSkinnedLocalBoundsExtents = FVector3f(LocalBounds.BoxExtent) * 2.0f;
 		}
@@ -1838,14 +1927,14 @@ bool FNDISkeletalMesh_InstanceData::Init(UNiagaraDataInterfaceSkeletalMesh* Inte
 	);
 
 	// GetSkeletalMeshGeneratedData() is meant to match with the required lod, so don't access unless we are using it.
-	if (NewSkelComp)
+	if (NewSkinnedComponent)
 	{
 		// TODO: This change is temporary to work around a crash that happens when you change the
 		// source mesh on a system which is running in the level from the details panel.
 		// bool bNeedsDataImmediately = SystemInstance->IsSolo();
 		const bool bNeedsDataImmediately = true;
 
-		TWeakObjectPtr<USkeletalMeshComponent> SkelWeakCompPtr = NewSkelComp;
+		TWeakObjectPtr<USkinnedMeshComponent> SkelWeakCompPtr = NewSkinnedComponent;
 		FNDI_SkeletalMesh_GeneratedData& GeneratedData = SystemInstance->GetWorldManager()->EditGeneratedData<FNDI_SkeletalMesh_GeneratedData>();
 		SkinningData = GeneratedData.GetCachedSkinningData(SkelWeakCompPtr, Usage, bNeedsDataImmediately);
 	}
@@ -1862,7 +1951,7 @@ bool FNDISkeletalMesh_InstanceData::Init(UNiagaraDataInterfaceSkeletalMesh* Inte
 		SystemInstance->EvaluateBoundFunction(FSkeletalMeshInterfaceHelper::GetTriangleCoordAtUVName, UsedByCpuUvMapping, UsedByGpuUvMapping);
 		SystemInstance->EvaluateBoundFunction(FSkeletalMeshInterfaceHelper::GetTriangleCoordInAabbName, UsedByCpuUvMapping, UsedByGpuUvMapping);
 
-		const bool MeshValid = SkeletalMesh.IsValid();
+		const bool MeshValid = SkinnedMesh != nullptr;
 		const bool SupportUvMappingCpu = UsedByCpuUvMapping && MeshValid;
 		const bool SupportUvMappingGpu = UsedByGpuUvMapping && MeshValid && Interface->IsUsedWithGPUScript();
 
@@ -1873,7 +1962,7 @@ bool FNDISkeletalMesh_InstanceData::Init(UNiagaraDataInterfaceSkeletalMesh* Inte
 			const bool bNeedsDataImmediately = true;
 
 			FNDI_SkeletalMesh_GeneratedData& GeneratedData = SystemInstance->GetWorldManager()->EditGeneratedData<FNDI_SkeletalMesh_GeneratedData>();
-			UvMapping = GeneratedData.GetCachedUvMapping(SkeletalMesh, CachedLODIdx, Interface->UvSetIndex, UvMappingUsage, bNeedsDataImmediately);
+			UvMapping = GeneratedData.GetCachedUvMapping(WeakSkinnedMesh, CachedLODIdx, Interface->UvSetIndex, UvMappingUsage, bNeedsDataImmediately);
 		}
 		else
 		{
@@ -1889,7 +1978,7 @@ bool FNDISkeletalMesh_InstanceData::Init(UNiagaraDataInterfaceSkeletalMesh* Inte
 		SystemInstance->EvaluateBoundFunction(FSkeletalMeshInterfaceHelper::GetAdjacentTriangleIndexName, UsedByCpuConnectivity, UsedByGpuConnectivity);
 		SystemInstance->EvaluateBoundFunction(FSkeletalMeshInterfaceHelper::GetTriangleNeighborName, UsedByCpuConnectivity, UsedByGpuConnectivity);
 
-		const bool MeshValid = SkeletalMesh.IsValid();
+		const bool MeshValid = SkinnedMesh != nullptr;
 		const bool SupportConnectivityCpu = UsedByCpuConnectivity && MeshValid;
 		const bool SupportConnectivityGpu = UsedByGpuConnectivity && MeshValid && Interface->IsUsedWithGPUScript();
 
@@ -1900,7 +1989,7 @@ bool FNDISkeletalMesh_InstanceData::Init(UNiagaraDataInterfaceSkeletalMesh* Inte
 			const bool bNeedsDataImmediately = true;
 
 			FNDI_SkeletalMesh_GeneratedData& GeneratedData = SystemInstance->GetWorldManager()->EditGeneratedData<FNDI_SkeletalMesh_GeneratedData>();
-			Connectivity = GeneratedData.GetCachedConnectivity(SkeletalMesh, CachedLODIdx, ConnectivityUsage, bNeedsDataImmediately);
+			Connectivity = GeneratedData.GetCachedConnectivity(WeakSkinnedMesh, CachedLODIdx, ConnectivityUsage, bNeedsDataImmediately);
 		}
 		else
 		{
@@ -1915,7 +2004,7 @@ bool FNDISkeletalMesh_InstanceData::Init(UNiagaraDataInterfaceSkeletalMesh* Inte
 		SamplingRegionAreaWeightedSampler.Init(this);
 	}
 
-	if (Mesh)
+	if (SkinnedMesh)
 	{
 		check(CachedLODData);
 
@@ -1924,7 +2013,7 @@ bool FNDISkeletalMesh_InstanceData::Init(UNiagaraDataInterfaceSkeletalMesh* Inte
 		check(SkinWeightBuffer);
 
 		// Check for the validity of the Mesh's cpu data.
-		if (Mesh->GetLODInfo(CachedLODIdx)->bAllowCPUAccess)
+		if (SkinnedMesh->GetLODInfo(CachedLODIdx)->bAllowCPUAccess)
 		{
 			const bool LODDataNumVerticesCorrect = CachedLODData->GetNumVertices() > 0;
 			const bool LODDataPositonNumVerticesCorrect = CachedLODData->StaticVertexBuffers.PositionVertexBuffer.GetNumVertices() > 0;
@@ -1946,14 +2035,14 @@ bool FNDISkeletalMesh_InstanceData::Init(UNiagaraDataInterfaceSkeletalMesh* Inte
 		}
 
 		// Generate excluded root bone index (if any)
-		FReferenceSkeleton& RefSkel = Mesh->GetRefSkeleton();
+		FReferenceSkeleton& RefSkel = SkinnedMesh->GetRefSkeleton();
 		ExcludedBoneIndex = INDEX_NONE;
 		if (Interface->bExcludeBone && !Interface->ExcludeBoneName.IsNone())
 		{
 			ExcludedBoneIndex = RefSkel.FindBoneIndex(Interface->ExcludeBoneName);
 			if (ExcludedBoneIndex == INDEX_NONE)
 			{
-				UE_LOG(LogNiagara, Warning, TEXT("Skeletal Mesh Data Interface '%s' is missing bone '%s' this is ok but may not exclude what you want Mesh '%s' Component '%s'"), *Interface->GetFullName(), *Interface->ExcludeBoneName.ToString(), *Mesh->GetFullName(), *SceneComponent->GetFullName());
+				UE_LOG(LogNiagara, Warning, TEXT("Skeletal Mesh Data Interface '%s' is missing bone '%s' this is ok but may not exclude what you want Mesh '%s' Component '%s'"), *Interface->GetFullName(), *Interface->ExcludeBoneName.ToString(), *SkinnedMesh->GetFullName(), *GetFullNameSafe(WeakSceneComponent.Get()));
 			}
 		}
 
@@ -1962,7 +2051,7 @@ bool FNDISkeletalMesh_InstanceData::Init(UNiagaraDataInterfaceSkeletalMesh* Inte
 		{
 			if (RefSkel.GetNum() > TNumericLimits<uint16>::Max())
 			{
-				UE_LOG(LogNiagara, Warning, TEXT("Skeletal Mesh Data Interface '%s' requires more bones '%d' than we currently support '%d' Mesh '%s' Component '%s'"), *Interface->GetFullName(), RefSkel.GetNum(), TNumericLimits<uint16>::Max(), *Mesh->GetFullName(), *SceneComponent->GetFullName());
+				UE_LOG(LogNiagara, Warning, TEXT("Skeletal Mesh Data Interface '%s' requires more bones '%d' than we currently support '%d' Mesh '%s' Component '%s'"), *Interface->GetFullName(), RefSkel.GetNum(), TNumericLimits<uint16>::Max(), *SkinnedMesh->GetFullName(), *GetFullNameSafe(WeakSceneComponent.Get()));
 				return false;
 			}
 
@@ -2021,7 +2110,7 @@ bool FNDISkeletalMesh_InstanceData::Init(UNiagaraDataInterfaceSkeletalMesh* Inte
 
 			if (MissingFilteredBones.Len() > 0)
 			{
-				UE_LOG(LogNiagara, Warning, TEXT("Skeletal Mesh Data Interface is trying to sample from filtered bones that don't exist in it's skeleton. Mesh(%s) Bones(%s) System(%s)"), *GetFullNameSafe(Mesh), MissingFilteredBones.ToString(), *GetFullNameSafe(SystemInstance->GetSystem()));
+				UE_LOG(LogNiagara, Warning, TEXT("Skeletal Mesh Data Interface is trying to sample from filtered bones that don't exist in it's skeleton. Mesh(%s) Bones(%s) System(%s)"), *GetFullNameSafe(SkinnedMesh), MissingFilteredBones.ToString(), *GetFullNameSafe(SystemInstance->GetSystem()));
 			}
 		}
 		else
@@ -2036,12 +2125,12 @@ bool FNDISkeletalMesh_InstanceData::Init(UNiagaraDataInterfaceSkeletalMesh* Inte
 			FilteredSocketInfo.SetNum(FilteredSockets.Num());
 
 			//-TODO: We may need to handle skinning mode changes here
-			if (NewSkelComp != nullptr)
+			if (NewSkinnedComponent != nullptr)
 			{
 				for (int32 i = 0; i < FilteredSocketInfo.Num(); ++i)
 				{
 					FTransform SocketTransform;
-					NewSkelComp->GetSocketInfoByName(FilteredSockets[i], SocketTransform, FilteredSocketInfo[i].BoneIdx);
+					NewSkinnedComponent->GetSocketInfoByName(FilteredSockets[i], SocketTransform, FilteredSocketInfo[i].BoneIdx);
 					FilteredSocketInfo[i].Transform = (FTransform3f)SocketTransform;
 				}
 			}
@@ -2049,12 +2138,12 @@ bool FNDISkeletalMesh_InstanceData::Init(UNiagaraDataInterfaceSkeletalMesh* Inte
 			{
 				for (int32 i = 0; i < FilteredSocketInfo.Num(); ++i)
 				{
-					FilteredSocketInfo[i].Transform = FTransform3f(FMatrix44f(Mesh->GetComposedRefPoseMatrix(FilteredSockets[i])));
+					FilteredSocketInfo[i].Transform = FTransform3f(FMatrix44f(SkinnedMesh->GetComposedRefPoseMatrix(FilteredSockets[i])));
 					FilteredSocketInfo[i].BoneIdx = INDEX_NONE;
 				}
 			}
 
-			FilteredSocketBoneOffset = Mesh->GetRefSkeleton().GetNum();
+			FilteredSocketBoneOffset = SkinnedMesh->GetRefSkeleton().GetNum();
 
 			FilteredSocketTransformsIndex = 0;
 			FilteredSocketTransforms[0].Reset(FilteredSockets.Num());
@@ -2071,7 +2160,7 @@ bool FNDISkeletalMesh_InstanceData::Init(UNiagaraDataInterfaceSkeletalMesh* Inte
 				TStringBuilder<512> MissingSockets;
 				for (FName SocketName : FilteredSockets)
 				{
-					if (Mesh->FindSocket(SocketName) == nullptr)
+					if (SkinnedMesh->FindSocket(SocketName) == nullptr)
 					{
 						if (MissingSockets.Len() != 0)
 						{
@@ -2083,7 +2172,7 @@ bool FNDISkeletalMesh_InstanceData::Init(UNiagaraDataInterfaceSkeletalMesh* Inte
 
 				if (MissingSockets.Len() > 0)
 				{
-					UE_LOG(LogNiagara, Warning, TEXT("Skeletal Mesh Data Interface is trying to sample from filtered sockets that don't exist in it's skeleton. Mesh(%s) Sockets(%s) System(%s)"), *GetFullNameSafe(Mesh), MissingSockets.ToString(), *GetFullNameSafe(SystemInstance->GetSystem()));
+					UE_LOG(LogNiagara, Warning, TEXT("Skeletal Mesh Data Interface is trying to sample from filtered sockets that don't exist in it's skeleton. Mesh(%s) Sockets(%s) System(%s)"), *GetFullNameSafe(SkinnedMesh), MissingSockets.ToString(), *GetFullNameSafe(SystemInstance->GetSystem()));
 				}
 			}
 		}
@@ -2100,7 +2189,7 @@ bool FNDISkeletalMesh_InstanceData::Init(UNiagaraDataInterfaceSkeletalMesh* Inte
 			//check(MeshSkinWeightBufferSrv->IsValid()); // not available in this stream
 			MeshSkinWeightLookupBuffer = SkinWeightBuffer->GetLookupVertexBuffer();
 
-			FSkeletalMeshLODInfo* LODInfo = Mesh->GetLODInfo(CachedLODIdx);
+			FSkeletalMeshLODInfo* LODInfo = SkinnedMesh->GetLODInfo(CachedLODIdx);
 			bIsGpuUniformlyDistributedSampling = LODInfo->bSupportUniformlyDistributedSampling && bAllRegionsAreAreaWeighting;
 
 			const auto MaxInfluenceType = GetDefault<UNiagaraSettings>()->NDISkelMesh_GpuMaxInfluences;
@@ -2120,12 +2209,20 @@ bool FNDISkeletalMesh_InstanceData::Init(UNiagaraDataInterfaceSkeletalMesh* Inte
 
 			if (MaxInfluenceCount > 0 && (static_cast<uint32>(MaxInfluenceCount) < CachedLODData->GetVertexBufferMaxBoneInfluences()))
 			{
-				UE_LOG(LogNiagara, Warning, TEXT("Skeletal Mesh %s has bones extra influence: spawning from it might not work properly."), *Mesh->GetName());
+				UE_LOG(LogNiagara, Warning, TEXT("Skeletal Mesh %s has bones extra influence: spawning from it might not work properly."), *SkinnedMesh->GetName());
 			}
 
-			const FSkeletalMeshSamplingInfo& SamplingInfo = Mesh->GetSamplingInfo();
-			const FSkeletalMeshSamplingLODBuiltData* MeshSamplingBuiltData = SamplingInfo.GetBuiltData().WholeMeshBuiltData.IsValidIndex(CachedLODIdx) ? &SamplingInfo.GetBuiltData().WholeMeshBuiltData[CachedLODIdx] : nullptr;
-			bIsGpuUniformlyDistributedSampling &= MeshSamplingBuiltData != nullptr;
+			const FSkeletalMeshSamplingLODBuiltData* MeshSamplingBuiltData = nullptr;
+			if (USkeletalMesh* SkeletalMesh = Cast<USkeletalMesh>(SkinnedMesh))
+			{
+				const FSkeletalMeshSamplingInfo& SamplingInfo = SkeletalMesh->GetSamplingInfo();
+				MeshSamplingBuiltData = SamplingInfo.GetBuiltData().WholeMeshBuiltData.IsValidIndex(CachedLODIdx) ? &SamplingInfo.GetBuiltData().WholeMeshBuiltData[CachedLODIdx] : nullptr;
+				bIsGpuUniformlyDistributedSampling &= MeshSamplingBuiltData != nullptr;
+			}
+			else
+			{
+				bIsGpuUniformlyDistributedSampling = false;
+			}
 
 			MeshGpuSpawnStaticBuffers = new FSkeletalMeshGpuSpawnStaticBuffers();
 			MeshGpuSpawnStaticBuffers->Initialise(this, CachedLODIdx, *CachedLODData, MeshSamplingBuiltData, SystemInstance);
@@ -2151,23 +2248,24 @@ bool FNDISkeletalMesh_InstanceData::ResetRequired(UNiagaraDataInterfaceSkeletalM
 		}
 
 		// Reset if the scene component we've cached has been invalidated or unregistered
-		USceneComponent* ExistingSceneComponent = SceneComponent.Get();
+		USceneComponent* ExistingSceneComponent = WeakSceneComponent.Get();
 		if (bComponentValid && (!ExistingSceneComponent || !ExistingSceneComponent->IsRegistered()))
 		{
 			return true;
 		}
 
 		// Reset if the skeletal mesh / component we are currently reading has changed
-		TWeakObjectPtr<USceneComponent> FoundSceneComponent;
-		USkeletalMeshComponent* FoundSkeletalMeshComponent = nullptr;
-		USkeletalMesh* FoundSkeletalMesh = Interface->GetSkeletalMesh(SystemInstance, SystemInstance->GetAttachComponent(), FoundSceneComponent, FoundSkeletalMeshComponent, const_cast<FNDISkeletalMesh_InstanceData*>(this));
+		TWeakObjectPtr<USceneComponent> WeakFoundSceneComponent;
+		USkinnedMeshComponent* FoundSkinnedMeshComponent = nullptr;
+		USkinnedAsset* FoundSkinnedMesh = Interface->GetSkinnedMesh(SystemInstance, SystemInstance->GetAttachComponent(), WeakFoundSceneComponent, FoundSkinnedMeshComponent, const_cast<FNDISkeletalMesh_InstanceData*>(this));
 
-		if (SceneComponent != FoundSceneComponent)
+		if (WeakSceneComponent != WeakFoundSceneComponent)
 		{
 			return true;
 		}
 
-		if (SkeletalMesh != FoundSkeletalMesh)
+		USkinnedAsset* SkinnedMesh = WeakSkinnedMesh.Get();
+		if (SkinnedMesh != FoundSkinnedMesh)
 		{
 			// Skinning Data is unique per component
 			// Therefore if the component did not change but the SKM did we need to refresh the skinning data to ensure the bones are valid
@@ -2179,16 +2277,16 @@ bool FNDISkeletalMesh_InstanceData::ResetRequired(UNiagaraDataInterfaceSkeletalM
 		}
 
 		// Reset if the LOD we relied on was streamed out, or if the LOD we need could now be available.
-		if (SkeletalMesh != nullptr)
+		if (SkinnedMesh != nullptr)
 		{
-			const FStreamableRenderResourceState& SRRState = SkeletalMesh->GetStreamableResourceState();
-			const int32 NumValidLODs = FMath::Min(SRRState.NumRequestedLODs, SRRState.NumResidentLODs);
+			const TTuple<int32, int32> StreamingLODData = NDISkelMeshLocal::GetStreamingLODData(SkinnedMesh);
+			const int32 NumValidLODs = StreamingLODData.Get<0>();
 			if (NumValidLODs == 0)
 			{
 				return true;
 			}
 
-			const int32 CurrentFirstLOD = SRRState.LODCountToAssetFirstLODIdx(NumValidLODs);
+			const int32 CurrentFirstLOD = StreamingLODData.Get<1>();
 			if (CurrentFirstLOD > CachedLODIdx || (CurrentFirstLOD < CachedLODIdx && bResetOnLODStreamedIn))
 			{
 				return true;
@@ -2198,15 +2296,15 @@ bool FNDISkeletalMesh_InstanceData::ResetRequired(UNiagaraDataInterfaceSkeletalM
 	else
 	{
 		// Reset if the scene component we've cached has been invalidated or unregistered
-		USceneComponent* Comp = SceneComponent.Get();
-		if (bComponentValid && (!Comp || !Comp->IsRegistered()))
+		USceneComponent* SceneComponent = WeakSceneComponent.Get();
+		if (bComponentValid && (!SceneComponent || !SceneComponent->IsRegistered()))
 		{
 			return true;
 		}
 
 		// Reset if any mesh was bound on init, but is now invalidated
-		USkeletalMesh* SkelMesh = SkeletalMesh.Get();
-		if (bMeshValid && !SkelMesh)
+		USkinnedAsset* SkinnedMesh = WeakSkinnedMesh.Get();
+		if (bMeshValid && !SkinnedMesh)
 		{
 			return true;
 		}
@@ -2222,7 +2320,7 @@ bool FNDISkeletalMesh_InstanceData::ResetRequired(UNiagaraDataInterfaceSkeletalM
 		else if (Interface->GetSourceComponent())
 		{
 			// Reset if the source component changed (or there wasn't one and now there is)
-			if (Interface->GetSourceComponent() != Comp)
+			if (Interface->GetSourceComponent() != SceneComponent)
 			{
 				return true;
 			}
@@ -2233,7 +2331,7 @@ bool FNDISkeletalMesh_InstanceData::ResetRequired(UNiagaraDataInterfaceSkeletalM
 			// TODO: This check is not really comprehensive. What we really need to know is if the mesh we cached comes from a skeletal mesh component in our
 			// attachment hierarchy, and if that hierarchy has changed in the chain between the system instance's attach component and the cached component,
 			// therefore potentially invalidating the cached component and mesh as our best choice.
-			if (CachedAttachParent != AttachComponent->GetAttachParent())
+			if (WeakCachedAttachParent.Get() != AttachComponent->GetAttachParent())
 			{
 				// The scene component our system instance was associated with has changed attachment, so we need to reinit
 				return true;
@@ -2241,16 +2339,16 @@ bool FNDISkeletalMesh_InstanceData::ResetRequired(UNiagaraDataInterfaceSkeletalM
 		}
 
 		// Reset if the LOD we relied on was streamed out, or if the LOD we need could now be available.
-		if (SkelMesh != nullptr)
+		if (SkinnedMesh != nullptr)
 		{
-			const FStreamableRenderResourceState& SRRState = SkelMesh->GetStreamableResourceState();
-			const int32 NumValidLODs = FMath::Min(SRRState.NumRequestedLODs, SRRState.NumResidentLODs);
+			const TTuple<int32, int32> StreamingLODData = NDISkelMeshLocal::GetStreamingLODData(SkinnedMesh);
+			const int32 NumValidLODs = StreamingLODData.Get<0>();
 			if (NumValidLODs == 0)
 			{
 				return true;
 			}
 
-			const int32 CurrentFirstLOD = SRRState.LODCountToAssetFirstLODIdx(NumValidLODs);
+			const int32 CurrentFirstLOD = StreamingLODData.Get<1>();
 			if (CurrentFirstLOD > CachedLODIdx || (CurrentFirstLOD < CachedLODIdx && bResetOnLODStreamedIn))
 			{
 				return true;
@@ -2258,9 +2356,9 @@ bool FNDISkeletalMesh_InstanceData::ResetRequired(UNiagaraDataInterfaceSkeletalM
 		}
 
 		// Reset if the skeletal mesh on the cached skeletal mesh component changed.
-		if (USkeletalMeshComponent* SkelComp = Cast<USkeletalMeshComponent>(Comp))
+		if (USkinnedMeshComponent* SkinnedComponent = Cast<USkinnedMeshComponent>(SceneComponent))
 		{
-			if (SkelComp->GetSkeletalMeshAsset() != SkelMesh)
+			if (SkinnedComponent->GetSkinnedAsset() != SkinnedMesh)
 			{
 				if (SkinningData.SkinningData.IsValid())
 				{
@@ -2310,9 +2408,9 @@ bool FNDISkeletalMesh_InstanceData::Tick(UNiagaraDataInterfaceSkeletalMesh* Inte
 
 FTransform FNDISkeletalMesh_InstanceData::CalculateComponentTransform(FNiagaraSystemInstance* SystemInstance) const
 {
-	if (USceneComponent* SceneComp = SceneComponent.Get())
+	if (USceneComponent* SceneComponent = WeakSceneComponent.Get())
 	{
-		FTransform ComponentTransform = SceneComp->GetComponentToWorld();
+		FTransform ComponentTransform = SceneComponent->GetComponentToWorld();
 		ComponentTransform.AddToTranslation(FVector(SystemInstance->GetLWCTile()) * -FLargeWorldRenderScalar::GetTileSize());
 		return ComponentTransform;
 	}
@@ -2324,13 +2422,13 @@ FTransform FNDISkeletalMesh_InstanceData::CalculateComponentTransform(FNiagaraSy
 
 void FNDISkeletalMesh_InstanceData::UpdateFilteredSocketTransforms()
 {
-	USkeletalMeshComponent* SkelComp = Cast<USkeletalMeshComponent>(SceneComponent.Get());
+	USkinnedMeshComponent* SceneComponent = Cast<USkinnedMeshComponent>(WeakSceneComponent.Get());
 	TArray<FTransform3f>& WriteBuffer = GetFilteredSocketsWriteBuffer();
 
 	for (int32 i = 0; i < FilteredSocketInfo.Num(); ++i)
 	{
 		const FCachedSocketInfo& SocketInfo = FilteredSocketInfo[i];
-		const FTransform& BoneTransform = SocketInfo.BoneIdx != INDEX_NONE ? SkelComp->GetBoneTransform(SocketInfo.BoneIdx, FTransform::Identity) : FTransform::Identity;
+		const FTransform& BoneTransform = SocketInfo.BoneIdx != INDEX_NONE ? SceneComponent->GetBoneTransform(SocketInfo.BoneIdx, FTransform::Identity) : FTransform::Identity;
 		WriteBuffer[i] = SocketInfo.Transform * FTransform3f(BoneTransform);
 	}
 }
@@ -2530,7 +2628,7 @@ void UNiagaraDataInterfaceSkeletalMesh::GetVMExternalFunction(const FVMExternalF
 	if (OutFunc.IsBound())
 	{
 		// Only warn about failed access if we are bound to a skeletal mesh, there are cases where the user chooses to have no SKM so this warning is invalid
-		if (!InstData->bAllowCPUMeshDataAccess && InstData->SkeletalMesh.IsValid())
+		if (!InstData->bAllowCPUMeshDataAccess && InstData->WeakSkinnedMesh.IsValid())
 		{
 			UE_LOG(LogNiagara, Log, TEXT("Skeletal Mesh Data Interface is trying to use triangle sampling function '%s', but either no CPU access is set on the mesh or the data is invalid. Interface: %s"),
 				*BindingInfo.Name.ToString(), *GetFullName());
@@ -2543,7 +2641,7 @@ void UNiagaraDataInterfaceSkeletalMesh::GetVMExternalFunction(const FVMExternalF
 	if (OutFunc.IsBound())
 	{
 		// Only warn about failed access if we are bound to a skeletal mesh, there are cases where the user chooses to have no SKM so this warning is invalid
-		if (!InstData->bAllowCPUMeshDataAccess && InstData->SkeletalMesh.IsValid())
+		if (!InstData->bAllowCPUMeshDataAccess && InstData->WeakSkinnedMesh.IsValid())
 		{
 			UE_LOG(LogNiagara, Log, TEXT("Skeletal Mesh Data Interface is trying to use vertex sampling function '%s' but either no CPU access is set on the mesh, or the data is invalid. Interface: %s"),
 				*BindingInfo.Name.ToString(), *GetFullName());
@@ -2566,6 +2664,7 @@ bool UNiagaraDataInterfaceSkeletalMesh::CopyToInternal(UNiagaraDataInterface* De
 	OtherTyped->SoftSourceActor = SoftSourceActor;
 	OtherTyped->MeshUserParameter = MeshUserParameter;
 	OtherTyped->ComponentTags = ComponentTags;
+	OtherTyped->OptionalComponentClass = OptionalComponentClass;
 	OtherTyped->SourceComponent = SourceComponent;
 	OtherTyped->SkinningMode = SkinningMode;
 	OtherTyped->SamplingRegions = SamplingRegions;
@@ -2601,6 +2700,7 @@ bool UNiagaraDataInterfaceSkeletalMesh::Equals(const UNiagaraDataInterface* Othe
 		OtherTyped->SoftSourceActor == SoftSourceActor &&
 		OtherTyped->MeshUserParameter == MeshUserParameter &&
 		OtherTyped->ComponentTags == ComponentTags &&
+		OtherTyped->OptionalComponentClass == OptionalComponentClass &&
 		OtherTyped->SourceComponent == SourceComponent &&
 		OtherTyped->SkinningMode == SkinningMode &&
 		OtherTyped->SamplingRegions == SamplingRegions &&
@@ -2626,12 +2726,15 @@ void UNiagaraDataInterfaceSkeletalMesh::DestroyPerInstanceData(void* PerInstance
 	FNDISkeletalMesh_InstanceData* Inst = (FNDISkeletalMesh_InstanceData*)PerInstanceData;
 
 #if WITH_EDITOR
-	if (USkeletalMesh* SkeletalMesh = Inst->SkeletalMesh.Get())
+	if (USkinnedAsset* SkinnedMesh = Inst->WeakSkinnedMesh.Get())
 	{
 		if (UNiagaraComponent* NiagaraComponent = Cast<UNiagaraComponent>(SystemInstance->GetAttachComponent()))
 		{
-			SkeletalMesh->GetOnMeshChanged().RemoveAll(NiagaraComponent);
-			if (USkeleton* Skeleton = SkeletalMesh->GetSkeleton())
+			if (USkeletalMesh* SkeletalMesh = Cast<USkeletalMesh>(SkinnedMesh))
+			{
+				SkeletalMesh->GetOnMeshChanged().RemoveAll(NiagaraComponent);
+			}
+			if (USkeleton* Skeleton = SkinnedMesh->GetSkeleton())
 			{
 				Skeleton->UnregisterOnSkeletonHierarchyChanged(NiagaraComponent);
 			}
@@ -2665,12 +2768,12 @@ void UNiagaraDataInterfaceSkeletalMesh::DrawDebugHud(FNDIDrawDebugHudContext& De
 		return;
 	}
 
-	USceneComponent* SceneComponent = InstanceData_GT->SceneComponent.Get();
+	USceneComponent* SceneComponent = InstanceData_GT->WeakSceneComponent.Get();
 	AActor* SceneActor = SceneComponent ? SceneComponent->GetAttachmentRootActor() : nullptr;
-	USkeletalMesh* SkeletalMesh = InstanceData_GT->SkeletalMesh.Get();
-	DebugHudContext.GetOutputString().Appendf(TEXT("Skeleton(%s) SkelComp(%s:%s)"), *GetNameSafe(SkeletalMesh), SceneActor ? *SceneActor->GetActorNameOrLabel() : TEXT("None"), *GetNameSafe(SceneComponent));
+	USkinnedAsset* SkinnedMesh = InstanceData_GT->WeakSkinnedMesh.Get();
+	DebugHudContext.GetOutputString().Appendf(TEXT("Skeleton(%s) SkelComp(%s:%s)"), *GetNameSafe(SkinnedMesh), SceneActor ? *SceneActor->GetActorNameOrLabel() : TEXT("None"), *GetNameSafe(SceneComponent));
 
-	if ( DebugHudContext.IsVerbose() && SceneComponent && SkeletalMesh)
+	if ( DebugHudContext.IsVerbose() && SceneComponent && SkinnedMesh)
 	{
 		FSkeletalMeshAccessorHelper MeshAccessor;
 		MeshAccessor.Init<TNDISkelMesh_FilterModeNone, TNDISkelMesh_AreaWeightingOff>(InstanceData_GT);
@@ -2707,15 +2810,15 @@ void UNiagaraDataInterfaceSkeletalMesh::GetFeedback(UNiagaraSystem* Asset, UNiag
 
 	bool bHasCPUAccessWarning = false;
 	bool bHasNoMeshAssignedWarning = false;
-	USkeletalMesh* SkelMesh = GetSkeletalMesh(Component);
+	USkinnedAsset* SkinnedMesh = GetSkinnedMesh(Component);
 
 	// Collect Errors
-	if (SkelMesh != nullptr)
+	if (SkinnedMesh != nullptr)
 	{
 		bool bHasCPUAccess = true;
-		for (int32 LODIndex = 0, LODNum = SkelMesh->GetLODNum(); LODIndex < LODNum; ++LODIndex)
+		for (int32 LODIndex = 0, LODNum = SkinnedMesh->GetLODNum(); LODIndex < LODNum; ++LODIndex)
 		{
-			if (!SkelMesh->GetLODInfo(LODIndex)->bAllowCPUAccess)
+			if (!SkinnedMesh->GetLODInfo(LODIndex)->bAllowCPUAccess)
 			{
 				bHasCPUAccess = false;
 				break;
@@ -2780,18 +2883,18 @@ void UNiagaraDataInterfaceSkeletalMesh::GetFeedback(UNiagaraSystem* Asset, UNiag
 	}
 
 	// Report Errors/Warnings
-	if (SkelMesh && bHasCPUAccessWarning)
+	if (SkinnedMesh && bHasCPUAccessWarning)
 	{
 		FNiagaraDataInterfaceFeedback CPUAccessNotAllowedWarning(
-			FText::Format(LOCTEXT("CPUAccessNotAllowedError", "This mesh may need CPU access in order to be used properly (even when used by GPU emitters). ({0})"), FText::FromString(SkelMesh->GetName())),
+			FText::Format(LOCTEXT("CPUAccessNotAllowedError", "This mesh may need CPU access in order to be used properly (even when used by GPU emitters). ({0})"), FText::FromString(SkinnedMesh->GetName())),
 			LOCTEXT("CPUAccessNotAllowedErrorSummary", "CPU access error"),
 			FNiagaraDataInterfaceFix::CreateLambda([=]()
 				{
-					SkelMesh->Modify();
+					SkinnedMesh->Modify();
 				
-					for (int32 LODIndex = 0, LODNum = SkelMesh->GetLODNum(); LODIndex < LODNum; ++LODIndex)
+					for (int32 LODIndex = 0, LODNum = SkinnedMesh->GetLODNum(); LODIndex < LODNum; ++LODIndex)
 					{
-						SkelMesh->GetLODInfo(LODIndex)->bAllowCPUAccess = true;
+						SkinnedMesh->GetLODInfo(LODIndex)->bAllowCPUAccess = true;
 					}
 					return true;
 				}));
@@ -2799,54 +2902,57 @@ void UNiagaraDataInterfaceSkeletalMesh::GetFeedback(UNiagaraSystem* Asset, UNiag
 		OutWarnings.Add(CPUAccessNotAllowedWarning);
 	}
 
-	if (SkelMesh && SkelMesh->HasActiveClothingAssets())
+	if (USkeletalMesh* SkeletalMesh = Cast<USkeletalMesh>(SkinnedMesh))
 	{
-		const FText ClothSummary = LOCTEXT("UsingClothSummary", "Active clothing assets found.  Results may be incorrect.");
-		if (GetDefault<UNiagaraSettings>()->NDISkelMesh_SupportReadingDeformedGeometry == false)
-		{
-			OutInfo.Emplace(
-				FText::Format(LOCTEXT("UsingClothNeedsProjectSetting", "Sampling from clothing regions on mesh '{0} will be incorrect.  Please enable reading from deformed geometry is disabled in project settings."), FText::FromString(SkelMesh->GetName())),
-				ClothSummary,
-				FNiagaraDataInterfaceFix()
-			);
-		}
-		else
+		if (SkeletalMesh->HasActiveClothingAssets())
 		{
-			OutInfo.Emplace(
-				FText::Format(LOCTEXT("UsingClothCheckValid", "Sampling from clothing regions on mesh '{0}' may be incorrect.  Ensure you are using a GPU emitter and the skin cache or deformer graph are enabled."), FText::FromString(SkelMesh->GetName())),
-				ClothSummary,
-				FNiagaraDataInterfaceFix()
-			);
+			const FText ClothSummary = LOCTEXT("UsingClothSummary", "Active clothing assets found.  Results may be incorrect.");
+			if (GetDefault<UNiagaraSettings>()->NDISkelMesh_SupportReadingDeformedGeometry == false)
+			{
+				OutInfo.Emplace(
+					FText::Format(LOCTEXT("UsingClothNeedsProjectSetting", "Sampling from clothing regions on mesh '{0} will be incorrect.  Please enable reading from deformed geometry is disabled in project settings."), FText::FromString(SkeletalMesh->GetName())),
+					ClothSummary,
+					FNiagaraDataInterfaceFix()
+				);
+			}
+			else
+			{
+				OutInfo.Emplace(
+					FText::Format(LOCTEXT("UsingClothCheckValid", "Sampling from clothing regions on mesh '{0}' may be incorrect.  Ensure you are using a GPU emitter and the skin cache or deformer graph are enabled."), FText::FromString(SkeletalMesh->GetName())),
+					ClothSummary,
+					FNiagaraDataInterfaceFix()
+				);
+			}
 		}
-	}
-
-	// Do we have any sampling regions enabled, if so check they are valid.
-	if (SkelMesh && SamplingRegions.Num() > 0)
-	{
-		FString MissingRegions;
 
-		const FSkeletalMeshSamplingInfo& SamplingInfo = SkelMesh->GetSamplingInfo();
-		for (FName RegionName : SamplingRegions)
+		// Do we have any sampling regions enabled, if so check they are valid.
+		if (SkeletalMesh && SamplingRegions.Num() > 0)
 		{
-			const int32 RegionIdx = SamplingInfo.IndexOfRegion(RegionName);
-			if (RegionIdx == INDEX_NONE)
+			FString MissingRegions;
+
+			const FSkeletalMeshSamplingInfo& SamplingInfo = SkeletalMesh->GetSamplingInfo();
+			for (FName RegionName : SamplingRegions)
 			{
-				if (!MissingRegions.IsEmpty())
+				const int32 RegionIdx = SamplingInfo.IndexOfRegion(RegionName);
+				if (RegionIdx == INDEX_NONE)
 				{
-					MissingRegions.Append(TEXT(", "));
+					if (!MissingRegions.IsEmpty())
+					{
+						MissingRegions.Append(TEXT(", "));
+					}
+					RegionName.AppendString(MissingRegions);
 				}
-				RegionName.AppendString(MissingRegions);
 			}
-		}
 
-		if (!MissingRegions.IsEmpty())
-		{
-			OutWarnings.Emplace(
-				FText::Format(LOCTEXT("MissingSamplingRegions", "The following regions were not found '{0}'"), FText::FromString(MissingRegions)),
-				LOCTEXT("MissingSamplingRegionsSummary", "Invalid Sampling Regions"),
-				FNiagaraDataInterfaceFix(
-				)
-			);
+			if (!MissingRegions.IsEmpty())
+			{
+				OutWarnings.Emplace(
+					FText::Format(LOCTEXT("MissingSamplingRegions", "The following regions were not found '{0}'"), FText::FromString(MissingRegions)),
+					LOCTEXT("MissingSamplingRegionsSummary", "Invalid Sampling Regions"),
+					FNiagaraDataInterfaceFix(
+					)
+				);
+			}
 		}
 	}
 
@@ -2877,13 +2983,13 @@ void UNiagaraDataInterfaceSkeletalMesh::GetFeedback(UNiagaraSystem* Asset, UNiag
 	}
 
 	// Look for bones being used that are LOD'ed out
-	if ( (SkelMesh != nullptr) && (SkelMesh->GetResourceForRendering() != nullptr) )
+	if ( (SkinnedMesh != nullptr) && (SkinnedMesh->GetResourceForRendering() != nullptr) )
 	{
-		FSkeletalMeshRenderData* SkelResource = SkelMesh->GetResourceForRendering();
+		FSkeletalMeshRenderData* SkelResource = SkinnedMesh->GetResourceForRendering();
 		auto IsBoneRequiredInAllLODs =
 			[&](const FName BoneName)
 			{
-				const int32 BoneIndex = SkelMesh->GetRefSkeleton().FindBoneIndex(BoneName);
+				const int32 BoneIndex = SkinnedMesh->GetRefSkeleton().FindBoneIndex(BoneName);
 				if (BoneIndex == INDEX_NONE)
 				{
 					return false;
@@ -3498,8 +3604,8 @@ void UNiagaraDataInterfaceSkeletalMesh::SetShaderParameters(const FNiagaraDataIn
 
 		if (InstanceData->bReadDeformedGeometry)
 		{
-			USkeletalMeshComponent* SkelComp = Cast<USkeletalMeshComponent>(StaticBuffers->GetSceneComponent());
-			FSkeletalMeshObject* SkelMeshObject = SkelComp ? SkelComp->MeshObject : nullptr;
+			USkinnedMeshComponent* SkinnedComponent = Cast<USkinnedMeshComponent>(StaticBuffers->GetSceneComponent());
+			FSkeletalMeshObject* SkelMeshObject = SkinnedComponent ? SkinnedComponent->MeshObject : nullptr;
 			if (SkelMeshObject)
 			{
 				FCachedGeometry CacheGeometry;
@@ -3701,7 +3807,7 @@ void UNiagaraDataInterfaceSkeletalMesh::OnSourceEndPlay(AActor* InSource, EEndPl
 	SourceComponent = nullptr;
 }
 
-void UNiagaraDataInterfaceSkeletalMesh::SetSourceComponentFromBlueprints(USkeletalMeshComponent* ComponentToUse)
+void UNiagaraDataInterfaceSkeletalMesh::SetSourceComponentFromBlueprints(USkinnedMeshComponent* ComponentToUse)
 {
 	// NOTE: When ChangeId changes the next tick will be skipped and a reset of the per-instance data will be initiated.
 	++ChangeId;
@@ -3742,89 +3848,89 @@ void UNiagaraDataInterfaceSkeletalMesh::SetWholeMeshLODFromBlueprints(int32 InWh
 ETickingGroup UNiagaraDataInterfaceSkeletalMesh::CalculateTickGroup(const void* PerInstanceData) const
 {
 	const FNDISkeletalMesh_InstanceData* InstData = static_cast<const FNDISkeletalMesh_InstanceData*>(PerInstanceData);
-	USkeletalMeshComponent* Component = Cast<USkeletalMeshComponent>(InstData->SceneComponent.Get());
-	if (Component && bRequireCurrentFrameData)
+	USkinnedMeshComponent* SkinnedComponent = Cast<USkinnedMeshComponent>(InstData->WeakSceneComponent.Get());
+	if (SkinnedComponent && bRequireCurrentFrameData)
 	{
-		return NDISkelMeshLocal::GetComponentTickGroup(Component);
+		USkeletalMeshComponent* SkeletalComponent = Cast<USkeletalMeshComponent>(SkinnedComponent);
+		return NDISkelMeshLocal::GetComponentTickGroup(SkinnedComponent, SkeletalComponent);
 	}
 	return NiagaraFirstTickGroup;
 }
 
 
-int32 UNiagaraDataInterfaceSkeletalMesh::CalculateLODIndexAndSamplingRegions(USkeletalMesh* InMesh, TArray<int32>& OutSamplingRegionIndices, bool& OutAllRegionsAreAreaWeighting) const
+int32 UNiagaraDataInterfaceSkeletalMesh::CalculateLODIndexAndSamplingRegions(USkinnedAsset* SkinnedMesh, TArray<int32>& OutSamplingRegionIndices, bool& OutAllRegionsAreAreaWeighting) const
 {
-	check(InMesh);
+	check(SkinnedMesh);
 
-	if (!SamplingRegions.Num())
+	if (SamplingRegions.Num())
 	{
-		//If we have no regions, sample the whole mesh at the specified LOD.
-		if (WholeMeshLOD == INDEX_NONE)
+		if (USkeletalMesh* SkeletalMesh = Cast<USkeletalMesh>(SkinnedMesh))
 		{
-			return InMesh->GetLODNum() - 1;
-		}
-		else
-		{
-			return FMath::Clamp(WholeMeshLOD, 0, InMesh->GetLODNum() - 1);
-		}
-	}
-	else
-	{
-		int32 LastRegionLODIndex = INDEX_NONE;
+			int32 LastRegionLODIndex = INDEX_NONE;
 
-		//Sampling from regions. Gather the indices of the regions we'll sample from.
-		const FSkeletalMeshSamplingInfo& SamplingInfo = InMesh->GetSamplingInfo();
-		for (FName RegionName : SamplingRegions)
-		{
-			const int32 RegionIdx = SamplingInfo.IndexOfRegion(RegionName);
-			if (RegionIdx != INDEX_NONE)
+			//Sampling from regions. Gather the indices of the regions we'll sample from.
+			const FSkeletalMeshSamplingInfo& SamplingInfo = SkeletalMesh->GetSamplingInfo();
+			for (FName RegionName : SamplingRegions)
 			{
-				const FSkeletalMeshSamplingRegion& Region = SamplingInfo.GetRegion(RegionIdx);
-				const FSkeletalMeshSamplingRegionBuiltData& RegionBuiltData = SamplingInfo.GetRegionBuiltData(RegionIdx);
+				const int32 RegionIdx = SamplingInfo.IndexOfRegion(RegionName);
+				if (RegionIdx != INDEX_NONE)
+				{
+					const FSkeletalMeshSamplingRegion& Region = SamplingInfo.GetRegion(RegionIdx);
+					const FSkeletalMeshSamplingRegionBuiltData& RegionBuiltData = SamplingInfo.GetRegionBuiltData(RegionIdx);
 
-				const int32 RegionLODIndex = Region.LODIndex == INDEX_NONE ?
-					(InMesh->GetLODNum() - 1) :
-					(FMath::Clamp(Region.LODIndex, 0, InMesh->GetLODNum() - 1));
+					const int32 RegionLODIndex = Region.LODIndex == INDEX_NONE ? (SkinnedMesh->GetLODNum() - 1) : (FMath::Clamp(Region.LODIndex, 0, SkinnedMesh->GetLODNum() - 1));
 
-				if (LastRegionLODIndex == INDEX_NONE)
-				{
-					LastRegionLODIndex = RegionLODIndex;
-				}
-				else if (RegionLODIndex != LastRegionLODIndex)
-				{
-					UE_LOG(LogNiagara, Warning, TEXT("Skeletal Mesh Data Interface is trying to use regions on LODs levels that are either streamed or cooked out. This is currently unsupported.\nInterface: %s\nMesh: %s\nRegion: %s"),
-						*GetFullName(),
-						*InMesh->GetFullName(),
-						*RegionName.ToString());
-					return INDEX_NONE;
-				}
+					if (LastRegionLODIndex == INDEX_NONE)
+					{
+						LastRegionLODIndex = RegionLODIndex;
+					}
+					else if (RegionLODIndex != LastRegionLODIndex)
+					{
+						UE_LOG(LogNiagara, Warning, TEXT("Skeletal Mesh Data Interface is trying to use regions on LODs levels that are either streamed or cooked out. This is currently unsupported.\nInterface: %s\nMesh: %s\nRegion: %s"),
+							*GetFullName(),
+							*SkinnedMesh->GetFullName(),
+							*RegionName.ToString());
+						return INDEX_NONE;
+					}
 
-				if (RegionBuiltData.TriangleIndices.Num() > 0)
-				{
-					OutSamplingRegionIndices.Add(RegionIdx);
-					OutAllRegionsAreAreaWeighting &= Region.bSupportUniformlyDistributedSampling;
+					if (RegionBuiltData.TriangleIndices.Num() > 0)
+					{
+						OutSamplingRegionIndices.Add(RegionIdx);
+						OutAllRegionsAreAreaWeighting &= Region.bSupportUniformlyDistributedSampling;
+					}
+					else
+					{
+						UE_LOG(LogNiagara, Warning, TEXT("Skeletal Mesh Data Interface is trying to use a region with no associated triangles.\nLOD: %d\nInterface: %s\nMesh: %s\nRegion: %s"),
+							 RegionLODIndex,
+							*GetFullName(),
+							*SkinnedMesh->GetFullName(),
+							*RegionName.ToString());
+
+						return INDEX_NONE;
+					}
 				}
 				else
 				{
-					UE_LOG(LogNiagara, Warning, TEXT("Skeletal Mesh Data Interface is trying to use a region with no associated triangles.\nLOD: %d\nInterface: %s\nMesh: %s\nRegion: %s"),
-						 RegionLODIndex,
+					UE_LOG(LogNiagara, Warning, TEXT("Skeletal Mesh Data Interface is trying to use a region on a mesh that does not provide this region.\nInterface: %s\nMesh: %s\nRegion: %s"),
 						*GetFullName(),
-						*InMesh->GetFullName(),
+						*SkinnedMesh->GetFullName(),
 						*RegionName.ToString());
 
 					return INDEX_NONE;
 				}
 			}
-			else
-			{
-				UE_LOG(LogNiagara, Warning, TEXT("Skeletal Mesh Data Interface is trying to use a region on a mesh that does not provide this region.\nInterface: %s\nMesh: %s\nRegion: %s"),
-					*GetFullName(),
-					*InMesh->GetFullName(),
-					*RegionName.ToString());
-
-				return INDEX_NONE;
-			}
+			return LastRegionLODIndex;
 		}
-		return LastRegionLODIndex;
+	}
+
+	//If we have no regions, sample the whole mesh at the specified LOD.
+	if (WholeMeshLOD == INDEX_NONE)
+	{
+		return SkinnedMesh->GetLODNum() - 1;
+	}
+	else
+	{
+		return FMath::Clamp(WholeMeshLOD, 0, SkinnedMesh->GetLODNum() - 1);
 	}
 }
 
@@ -3859,8 +3965,8 @@ void UNiagaraDataInterfaceSkeletalMesh::VMGetPreSkinnedLocalBounds(FVectorVMExte
 template<>
 void FSkeletalMeshAccessorHelper::Init<TNDISkelMesh_FilterModeSingle, TNDISkelMesh_AreaWeightingOff>(const FNDISkeletalMesh_InstanceData* InstData)
 {
-	Comp = Cast<USkeletalMeshComponent>(InstData->SceneComponent.Get());
-	Mesh = InstData->SkeletalMesh.Get();
+	Comp = Cast<USkeletalMeshComponent>(InstData->WeakSceneComponent.Get());
+	Mesh = Cast<USkeletalMesh>(InstData->WeakSkinnedMesh.Get());
 	LODData = InstData->CachedLODData;
 	SkinWeightBuffer = InstData->GetSkinWeights();
 	IndexBuffer = LODData ? LODData->MultiSizeIndexContainer.GetIndexBuffer() : nullptr;
@@ -3888,8 +3994,8 @@ void FSkeletalMeshAccessorHelper::Init<TNDISkelMesh_FilterModeSingle, TNDISkelMe
 template<>
 void FSkeletalMeshAccessorHelper::Init<TNDISkelMesh_FilterModeSingle, TNDISkelMesh_AreaWeightingOn>(const FNDISkeletalMesh_InstanceData* InstData)
 {
-	Comp = Cast<USkeletalMeshComponent>(InstData->SceneComponent.Get());
-	Mesh = InstData->SkeletalMesh.Get();
+	Comp = Cast<USkeletalMeshComponent>(InstData->WeakSceneComponent.Get());
+	Mesh = Cast<USkeletalMesh>(InstData->WeakSkinnedMesh.Get());
 	LODData = InstData->CachedLODData;
 	SkinWeightBuffer = InstData->GetSkinWeights();
 	IndexBuffer = LODData ? LODData->MultiSizeIndexContainer.GetIndexBuffer() : nullptr;
diff --git a/Engine/Plugins/FX/Niagara/Source/Niagara/Private/NiagaraDataInterfaceSkeletalMeshConnectivity.cpp b/Engine/Plugins/FX/Niagara/Source/Niagara/Private/NiagaraDataInterfaceSkeletalMeshConnectivity.cpp
--- a/Engine/Plugins/FX/Niagara/Source/Niagara/Private/NiagaraDataInterfaceSkeletalMeshConnectivity.cpp
+++ b/Engine/Plugins/FX/Niagara/Source/Niagara/Private/NiagaraDataInterfaceSkeletalMeshConnectivity.cpp
@@ -84,7 +84,7 @@ void FSkeletalMeshConnectivityHandle::PinAndInvalidateHandle()
 	ConnectivityData.Reset();
 }
 
-FSkeletalMeshConnectivity::FSkeletalMeshConnectivity(TWeakObjectPtr<USkeletalMesh> InMeshObject, int32 InLodIndex)
+FSkeletalMeshConnectivity::FSkeletalMeshConnectivity(TWeakObjectPtr<USkinnedAsset> InMeshObject, int32 InLodIndex)
 	: LodIndex(InLodIndex)
 	, MeshObject(InMeshObject)
 	, GpuUserCount(0)
@@ -149,14 +149,14 @@ void FSkeletalMeshConnectivity::Release()
 	}
 }
 
-bool FSkeletalMeshConnectivity::CanBeUsed(const TWeakObjectPtr<USkeletalMesh>& InMeshObject, int32 InLodIndex) const
+bool FSkeletalMeshConnectivity::CanBeUsed(const TWeakObjectPtr<USkinnedAsset>& InMeshObject, int32 InLodIndex) const
 {
 	return IsUsed() && LodIndex == InLodIndex && MeshObject == InMeshObject;
 }
 
-bool FSkeletalMeshConnectivity::IsValidMeshObject(TWeakObjectPtr<USkeletalMesh>& MeshObject, int32 InLodIndex)
+bool FSkeletalMeshConnectivity::IsValidMeshObject(TWeakObjectPtr<USkinnedAsset>& MeshObject, int32 InLodIndex)
 {
-	USkeletalMesh* Mesh = MeshObject.Get();
+	USkinnedAsset* Mesh = MeshObject.Get();
 
 	if (!Mesh)
 	{
@@ -197,7 +197,7 @@ void FSkeletalMeshConnectivity::GetTriangleVertices(int32 TriangleIndex, int32&
 	OutVertex2 = INDEX_NONE;
 }
 
-const FSkeletalMeshLODRenderData* FSkeletalMeshConnectivity::GetLodRenderData(const USkeletalMesh& Mesh, int32 LodIndex)
+const FSkeletalMeshLODRenderData* FSkeletalMeshConnectivity::GetLodRenderData(const USkinnedAsset& Mesh, int32 LodIndex)
 {
 	if (FSkeletalMeshRenderData* RenderData = Mesh.GetResourceForRendering())
 	{
@@ -212,7 +212,7 @@ const FSkeletalMeshLODRenderData* FSkeletalMeshConnectivity::GetLodRenderData(co
 
 const FSkeletalMeshLODRenderData* FSkeletalMeshConnectivity::GetLodRenderData() const
 {
-	if (const USkeletalMesh* Mesh = MeshObject.Get())
+	if (const USkinnedAsset* Mesh = MeshObject.Get())
 	{
 		return GetLodRenderData(*Mesh, LodIndex);
 	}
@@ -222,7 +222,7 @@ const FSkeletalMeshLODRenderData* FSkeletalMeshConnectivity::GetLodRenderData()
 
 FString FSkeletalMeshConnectivity::GetMeshName() const
 {
-	if (const USkeletalMesh* Mesh = MeshObject.Get())
+	if (const USkinnedAsset* Mesh = MeshObject.Get())
 	{
 		return Mesh->GetPathName();
 	}
@@ -317,7 +317,7 @@ static bool BuildAdjacencyBuffer(const FSkeletalMeshLODRenderData& LodRenderData
 			{
 				const TriangleIndexType TestTriangle = AdjacentTriangles[InsertionPoint];
 
-				if (TestTriangle == INDEX_NONE)
+				if (TestTriangle == TriangleIndexType(INDEX_NONE))
 				{
 					AdjacentTriangles[InsertionPoint] = TriangleIt;
 					break;
@@ -328,7 +328,7 @@ static bool BuildAdjacencyBuffer(const FSkeletalMeshLODRenderData& LodRenderData
 				{
 					// skip empty entries
 					int32 ShiftIt = MaxAdjacencyCount - 1;
-					while (AdjacentTriangles[ShiftIt - 1] == INDEX_NONE)
+					while (AdjacentTriangles[ShiftIt - 1] == TriangleIndexType(INDEX_NONE))
 					{
 						--ShiftIt;
 					}
diff --git a/Engine/Plugins/FX/Niagara/Source/Niagara/Private/NiagaraDataInterfaceSkeletalMeshConnectivity.h b/Engine/Plugins/FX/Niagara/Source/Niagara/Private/NiagaraDataInterfaceSkeletalMeshConnectivity.h
--- a/Engine/Plugins/FX/Niagara/Source/Niagara/Private/NiagaraDataInterfaceSkeletalMeshConnectivity.h
+++ b/Engine/Plugins/FX/Niagara/Source/Niagara/Private/NiagaraDataInterfaceSkeletalMeshConnectivity.h
@@ -7,7 +7,7 @@
 #include "UObject/WeakObjectPtr.h"
 
 class FSkeletalMeshLODRenderData;
-class USkeletalMesh;
+class USkinnedAsset;
 struct FSkeletalMeshConnectivityUsage;
 
 struct FSkeletalMeshConnectivity;
@@ -40,17 +40,17 @@ struct FSkeletalMeshConnectivity
 {
 	FSkeletalMeshConnectivity() = delete;
 	FSkeletalMeshConnectivity(const FSkeletalMeshConnectivity&) = delete;
-	FSkeletalMeshConnectivity(TWeakObjectPtr<USkeletalMesh> InMeshObject, int32 InLodIndex);
+	FSkeletalMeshConnectivity(TWeakObjectPtr<USkinnedAsset> InMeshObject, int32 InLodIndex);
 	~FSkeletalMeshConnectivity();
 
 	bool IsUsed() const;
 	bool CanBeDestroyed() const;
 	void RegisterUser(FSkeletalMeshConnectivityUsage Usage, bool bNeedsDataImmediately);
 	void UnregisterUser(FSkeletalMeshConnectivityUsage Usage);
-	bool CanBeUsed(const TWeakObjectPtr<USkeletalMesh>& InMeshObject, int32 InLodIndex) const;
+	bool CanBeUsed(const TWeakObjectPtr<USkinnedAsset>& InMeshObject, int32 InLodIndex) const;
 	const FSkeletalMeshConnectivityProxy* GetProxy() const;
 
-	static bool IsValidMeshObject(TWeakObjectPtr<USkeletalMesh>& MeshObject, int32 InLodIndex);
+	static bool IsValidMeshObject(TWeakObjectPtr<USkinnedAsset>& MeshObject, int32 InLodIndex);
 
 	int32 GetAdjacentTriangleIndex(int32 VertexIndex, int32 AdjacencyIndex) const;
 	void GetTriangleVertices(int32 TriangleIndex, int32& OutVertex0, int32& OutVertex1, int32& OutVertex2) const;
@@ -60,11 +60,11 @@ struct FSkeletalMeshConnectivity
 	const int32 LodIndex;
 
 private:
-	static const FSkeletalMeshLODRenderData* GetLodRenderData(const USkeletalMesh& Mesh, int32 LodIndex);
+	static const FSkeletalMeshLODRenderData* GetLodRenderData(const USkinnedAsset& Mesh, int32 LodIndex);
 
 	void Release();
 
-	TWeakObjectPtr<USkeletalMesh> MeshObject;
+	TWeakObjectPtr<USkinnedAsset> MeshObject;
 	TUniquePtr<FSkeletalMeshConnectivityProxy> Proxy;
 
 	std::atomic<int32> GpuUserCount;
diff --git a/Engine/Plugins/FX/Niagara/Source/NiagaraEditorWidgets/Private/DetailCustomizations/NiagaraDataInterfaceSkeletalMeshDetails.cpp b/Engine/Plugins/FX/Niagara/Source/NiagaraEditorWidgets/Private/DetailCustomizations/NiagaraDataInterfaceSkeletalMeshDetails.cpp
--- a/Engine/Plugins/FX/Niagara/Source/NiagaraEditorWidgets/Private/DetailCustomizations/NiagaraDataInterfaceSkeletalMeshDetails.cpp
+++ b/Engine/Plugins/FX/Niagara/Source/NiagaraEditorWidgets/Private/DetailCustomizations/NiagaraDataInterfaceSkeletalMeshDetails.cpp
@@ -32,11 +32,12 @@ void FNiagaraDataInterfaceSkeletalMeshDetails::CustomizeDetails(IDetailLayoutBui
 	Interface->OnChanged().RemoveAll(this);
 	Interface->OnChanged().AddSP(this, &FNiagaraDataInterfaceSkeletalMeshDetails::OnInterfaceChanged);
 
-	MeshObject = Interface->GetSkeletalMesh(Cast<UNiagaraComponent>(Interface->GetOuter()));
-	if (MeshObject.IsValid())
+	SkinnedMeshObject = Interface->GetSkinnedMesh(Cast<UNiagaraComponent>(Interface->GetOuter()));
+	SkeletalMeshObject = Cast<USkeletalMesh>(SkinnedMeshObject);
+	if (SkeletalMeshObject.IsValid())
 	{
-		MeshObject->GetOnMeshChanged().RemoveAll(this);
-		MeshObject->GetOnMeshChanged().AddSP(this, &FNiagaraDataInterfaceSkeletalMeshDetails::OnDataChanged);
+		SkeletalMeshObject->GetOnMeshChanged().RemoveAll(this);
+		SkeletalMeshObject->GetOnMeshChanged().AddSP(this, &FNiagaraDataInterfaceSkeletalMeshDetails::OnDataChanged);
 	}
 
 	MeshCategory = &DetailBuilder.EditCategory(MeshCategoryName, LOCTEXT("Mesh", "Mesh"));
@@ -122,16 +123,17 @@ TSharedRef<IDetailCustomization> FNiagaraDataInterfaceSkeletalMeshDetails::MakeI
 void FNiagaraDataInterfaceSkeletalMeshDetails::OnInterfaceChanged()
 {
 	// Rebuild the data changed listener
-	if (MeshObject.IsValid())
+	if (SkeletalMeshObject.IsValid())
 	{
-		MeshObject->GetOnMeshChanged().RemoveAll(this);
+		SkeletalMeshObject->GetOnMeshChanged().RemoveAll(this);
 	}
 
 	UNiagaraDataInterfaceSkeletalMesh* Interface = MeshInterface.Get();
-	MeshObject = Interface->GetSkeletalMesh(Cast<UNiagaraComponent>(Interface->GetOuter()));
-	if (MeshObject.IsValid())
+	SkinnedMeshObject = Interface->GetSkinnedMesh(Cast<UNiagaraComponent>(Interface->GetOuter()));
+	SkeletalMeshObject = Cast<USkeletalMesh>(SkinnedMeshObject);
+	if (SkeletalMeshObject.IsValid())
 	{
-		MeshObject->GetOnMeshChanged().AddSP(this, &FNiagaraDataInterfaceSkeletalMeshDetails::OnDataChanged);
+		SkeletalMeshObject->GetOnMeshChanged().AddSP(this, &FNiagaraDataInterfaceSkeletalMeshDetails::OnDataChanged);
 	}
 
 	OnDataChanged();
@@ -173,7 +175,7 @@ void FNiagaraDataInterfaceSkeletalMeshDetails::GenerateRegionsArray(TArray<TShar
 	SourceArray.Reset();
 	if (UNiagaraDataInterfaceSkeletalMesh* Interface = MeshInterface.Get())
 	{
-		if (USkeletalMesh* Mesh = Interface->GetSkeletalMesh(Cast<UNiagaraComponent>(Interface->GetOuter())))
+		if (USkeletalMesh* Mesh = Cast<USkeletalMesh>(Interface->GetSkinnedMesh(Cast<UNiagaraComponent>(Interface->GetOuter()))))
 		{
 			for (FSkeletalMeshSamplingRegion Region : Mesh->GetSamplingInfo().Regions)
 			{
@@ -188,7 +190,7 @@ void FNiagaraDataInterfaceSkeletalMeshDetails::GenerateRegionsArray(TArray<TShar
 	SourceArray.Reset();
 	if (UNiagaraDataInterfaceSkeletalMesh* Interface = MeshInterface.Get())
 	{
-		if (USkeletalMesh* Mesh = Interface->GetSkeletalMesh(Cast<UNiagaraComponent>(Interface->GetOuter())))
+		if (USkinnedAsset* Mesh = Interface->GetSkinnedMesh(Cast<UNiagaraComponent>(Interface->GetOuter())))
 		{
 			for (const FMeshBoneInfo& Bone : Mesh->GetRefSkeleton().GetRefBoneInfo())
 			{
@@ -204,7 +206,7 @@ void FNiagaraDataInterfaceSkeletalMeshDetails::GenerateSocketsArray(TArray<TShar
 	if (MeshInterface.IsValid())
 	{
 		UNiagaraDataInterfaceSkeletalMesh* Interface = MeshInterface.Get();
-		if (USkeletalMesh* Mesh = Interface->GetSkeletalMesh(Cast<UNiagaraComponent>(Interface->GetOuter())))
+		if (USkeletalMesh* Mesh = Cast<USkeletalMesh>(Interface->GetSkinnedMesh(Cast<UNiagaraComponent>(Interface->GetOuter()))))
 		{
 			for (int32 SocketIdx = 0; SocketIdx < Mesh->NumSockets(); ++SocketIdx)
 			{
@@ -221,9 +223,9 @@ FNiagaraDataInterfaceSkeletalMeshDetails::~FNiagaraDataInterfaceSkeletalMeshDeta
 	{
 		MeshInterface->OnChanged().RemoveAll(this);
 	}
-	if (MeshObject.IsValid())
+	if (SkeletalMeshObject.IsValid())
 	{
-		MeshObject->GetOnMeshChanged().RemoveAll(this);
+		SkeletalMeshObject->GetOnMeshChanged().RemoveAll(this);
 	}
 }
 
diff --git a/Engine/Plugins/FX/Niagara/Source/NiagaraEditorWidgets/Private/DetailCustomizations/NiagaraDataInterfaceSkeletalMeshDetails.h b/Engine/Plugins/FX/Niagara/Source/NiagaraEditorWidgets/Private/DetailCustomizations/NiagaraDataInterfaceSkeletalMeshDetails.h
--- a/Engine/Plugins/FX/Niagara/Source/NiagaraEditorWidgets/Private/DetailCustomizations/NiagaraDataInterfaceSkeletalMeshDetails.h
+++ b/Engine/Plugins/FX/Niagara/Source/NiagaraEditorWidgets/Private/DetailCustomizations/NiagaraDataInterfaceSkeletalMeshDetails.h
@@ -14,6 +14,7 @@ class FName;
 class FNiagaraDetailSourcedArrayBuilder;
 class UNiagaraDataInterfaceSkeletalMesh;
 class USkeletalMesh;
+class USkinnedAsset;
 class SNiagaraNamePropertySelector;
 
 /** Details customization for Niagara skeletal mesh data interface. */
@@ -38,7 +39,8 @@ private:
 	TSharedPtr<SNiagaraNamePropertySelector> ExcludeBoneWidget;
 	IDetailLayoutBuilder* LayoutBuilder;
 	TWeakObjectPtr<UNiagaraDataInterfaceSkeletalMesh>  MeshInterface;
-	TWeakObjectPtr<USkeletalMesh> MeshObject;
+	TWeakObjectPtr<USkeletalMesh> SkeletalMeshObject;
+	TWeakObjectPtr<USkinnedAsset> SkinnedMeshObject;
 
 	IDetailCategoryBuilder* MeshCategory;
 	IDetailCategoryBuilder* SkelCategory;
