Fix handling of mismatched system scripts if it occurs during a cook.  This allows the cook to perform the forced recompile rather than stalling the cook process.

#tests Cooked asset that was hanging a full cook
#Synced-CL 12905823

diff --git a/Engine/Plugins/FX/Niagara/Source/Niagara/Private/NiagaraSystem.cpp b/Engine/Plugins/FX/Niagara/Source/Niagara/Private/NiagaraSystem.cpp
--- a/Engine/Plugins/FX/Niagara/Source/Niagara/Private/NiagaraSystem.cpp
+++ b/Engine/Plugins/FX/Niagara/Source/Niagara/Private/NiagaraSystem.cpp
@@ -426,7 +426,23 @@ bool UNiagaraSystem::IsCachedCookedPlatformDataLoaded(const ITargetPlatform* Tar
 
 	if (!ActiveCompilations.IsEmpty())
 	{
-		return PollForCompilationComplete(false);
+		const bool bWasCompilationComplete = PollForCompilationComplete(false);
+
+		if (bWasCompilationComplete)
+		{
+			// in order to handle the possibility that the requested compile failed in such a way that a new compile is being requested
+			// then issue that compile now and continue to wait.  This can happen in the case of system scripts being mismatched (for
+			// a currently unknown reason) and a force-compile is required to address the issue.
+			if (NeedsRequestCompile())
+			{
+				RequestCompile(false, nullptr, TargetPlatform);
+				return false;
+			}
+
+			return true;
+		}
+
+		return false;
 	}
 
 	return true;
