- Fix crash when converting LWC <-> SWC with aligned types

#tests Add the MergeHelper_Apply module to a particle update script
#rnx

#Synced-CL 48563820

diff --git a/Engine/Plugins/FX/Niagara/Source/Niagara/Public/NiagaraTypes.h b/Engine/Plugins/FX/Niagara/Source/Niagara/Public/NiagaraTypes.h
--- a/Engine/Plugins/FX/Niagara/Source/Niagara/Public/NiagaraTypes.h
+++ b/Engine/Plugins/FX/Niagara/Source/Niagara/Public/NiagaraTypes.h
@@ -293,14 +293,16 @@ struct FNiagaraStructConversionStep
 	void CopyToSim(const uint8* LWCData, uint8* SimulationData, int32 Count, int32 LWCStride, int32 SimulationStride) const;
 	void CopyFromSim(uint8* LWCData, const uint8* SimulationData, int32 Count, int32 LWCStride, int32 SimulationStride) const;
 
-	template<typename DstType, typename SrcType>
+	template<typename TDstType, typename TSrcType>
 	void CopyInternal(uint8* DestinationData, int32 DestStride, int32 DestOffset, const uint8* SourceData, int32 SourceStride, int32 SourceOffset, int32 Count)const
 	{
 		for (int32 i = 0; i < Count; ++i)
 		{
-			DstType* Dst = reinterpret_cast<DstType*>((DestinationData + i * DestStride) + DestOffset);
-			const SrcType* Src = reinterpret_cast<const SrcType*>((SourceData + i * SourceStride) + SourceOffset);			
-			*Dst = DstType(*Src);
+			TSrcType SrcValue;
+			FMemory::Memcpy(&SrcValue, (SourceData + i * SourceStride) + SourceOffset, sizeof(TSrcType));
+
+			TDstType DstValue = TDstType(SrcValue);
+			FMemory::Memcpy((DestinationData + i * DestStride) + DestOffset, &DstValue, sizeof(TDstType));
 		}
 	}
 };
