- Query material validation for used renderer materials
- Allows Niagara to display warnings from the material validation systems, such as creating permutations from static switches

#tests Create a NS which uses a material that contains a MIC, enable a permutation that is not already available

#Synced-CL 48563820

diff --git a/Engine/Plugins/FX/Niagara/Source/NiagaraEditor/NiagaraEditor.Build.cs b/Engine/Plugins/FX/Niagara/Source/NiagaraEditor/NiagaraEditor.Build.cs
--- a/Engine/Plugins/FX/Niagara/Source/NiagaraEditor/NiagaraEditor.Build.cs
+++ b/Engine/Plugins/FX/Niagara/Source/NiagaraEditor/NiagaraEditor.Build.cs
@@ -54,7 +54,8 @@ public class NiagaraEditor : ModuleRules
 				"SparseVolumeTexture",
 				"PlacementMode",
 				"EngineSettings",
-				"EditorConfig"
+				"EditorConfig",
+				"MaterialEditor",
 			}
 		);
 
diff --git a/Engine/Plugins/FX/Niagara/Source/NiagaraEditor/Private/ViewModels/Stack/NiagaraStackRendererItem.cpp b/Engine/Plugins/FX/Niagara/Source/NiagaraEditor/Private/ViewModels/Stack/NiagaraStackRendererItem.cpp
--- a/Engine/Plugins/FX/Niagara/Source/NiagaraEditor/Private/ViewModels/Stack/NiagaraStackRendererItem.cpp
+++ b/Engine/Plugins/FX/Niagara/Source/NiagaraEditor/Private/ViewModels/Stack/NiagaraStackRendererItem.cpp
@@ -29,6 +29,9 @@
 #include "ViewModels/Stack/NiagaraStackViewModel.h"
 #include "NiagaraEditorStyle.h"
 
+#include "Materials/MaterialInterface.h"
+#include "MaterialEditorValidation.h"
+
 #include UE_INLINE_GENERATED_CPP_BY_NAME(NiagaraStackRendererItem)
 
 #define LOCTEXT_NAMESPACE "UNiagaraStackRendererItem"
@@ -652,8 +655,35 @@ void UNiagaraStackRendererItem::RefreshIssues(TArray<FStackIssue>& NewIssues)
 		TArray<FNiagaraRendererFeedback> Warnings;
 		TArray<FNiagaraRendererFeedback> Infos;
 
+		// Query feedback from the renderer
 		RendererProperties->GetRendererFeedback(GetEmitterViewModel()->GetEmitter(), Errors, Warnings, Infos);
 
+		// Query feedback from the material validation system
+		if (UMaterialEditorValidationSubsystem* MaterialValidationSubsystem = GEditor->GetEditorSubsystem<UMaterialEditorValidationSubsystem>())
+		{
+			TArray<UMaterialInterface*> RendererMaterials;
+			RendererProperties->GetUsedMaterials(nullptr, RendererMaterials);
+
+			for (UMaterialInterface* Material : RendererMaterials)
+			{
+				if (Material == nullptr)
+				{
+					continue;
+				}
+
+				FMaterialEditorValidatorContext MaterialValidatorContext;
+				MaterialValidationSubsystem->Validate(Material, MaterialValidatorContext);
+				if (MaterialValidatorContext.Message.IsEmpty() == false)
+				{
+					Warnings.Emplace(
+						MaterialValidatorContext.Message,
+						FText::Format(LOCTEXT("MaterialValidationMessage", "Material '{0}' has validation issues."), FText::FromName(Material->GetFName()))
+					);
+				}
+			}
+		}
+
+		// Convert feedback into stack issues
 		ProcessRendererIssues(Errors, EStackIssueSeverity::Error, NewIssues);		
 		ProcessRendererIssues(Warnings, EStackIssueSeverity::Warning, NewIssues);		
 		ProcessRendererIssues(Infos, EStackIssueSeverity::Info, NewIssues);	
diff --git a/Engine/Source/Editor/MaterialEditor/Public/MaterialEditorValidation.h b/Engine/Source/Editor/MaterialEditor/Public/MaterialEditorValidation.h
--- a/Engine/Source/Editor/MaterialEditor/Public/MaterialEditorValidation.h
+++ b/Engine/Source/Editor/MaterialEditor/Public/MaterialEditorValidation.h
@@ -39,7 +39,7 @@ public:
 	UE_API void RemoveValidator(UMaterialEditorValidatorBase* InValidator);
 	
 	/** Call Validate() on all registered validators.*/
-	bool Validate(UMaterialInterface* InMaterial, FMaterialEditorValidatorContext& InContext) const;
+	UE_API bool Validate(UMaterialInterface* InMaterial, FMaterialEditorValidatorContext& InContext) const;
 
 private:
 	/** The array of registered validators. */
