- Add support to convert an ID to index in the sim cache reader DI
- This is not optimal as it is a linear search currently

#tests Make a Sim Cache with persistant IDs (ideally GPU with > 64 particles), make a emitter that reads from the cache using persistant IDs

#Synced-CL 48563820

diff --git a/Engine/Plugins/FX/Niagara/Shaders/Private/NiagaraDataInterfaceSimCacheReaderTemplate.ush b/Engine/Plugins/FX/Niagara/Shaders/Private/NiagaraDataInterfaceSimCacheReaderTemplate.ush
--- a/Engine/Plugins/FX/Niagara/Shaders/Private/NiagaraDataInterfaceSimCacheReaderTemplate.ush
+++ b/Engine/Plugins/FX/Niagara/Shaders/Private/NiagaraDataInterfaceSimCacheReaderTemplate.ush
@@ -138,3 +138,26 @@ void GetInstanceID_{ParameterName}(int FrameIndex, int EmitterIndex, int Attribu
 		Value.AcquireTag = asint({ParameterName}_CacheData[FrameData.Int32DataOffset + ((AttributeIndex + 1) * FrameData.NumInstances) + InstanceIndex]);
 	}
 }
+
+void GetInstanceIndexFromID_{ParameterName}(int FrameIndex, int EmitterIndex, int AttributeIndex, NiagaraID ID, out bool bValid, out int Value)
+{
+	bValid = false;
+	Value = -1;
+
+	if ( AttributeIndex != -1 )
+	{
+		{ParameterName}_EmitterFrameData FrameData = GetEmitterFrameData_{ParameterName}(FrameIndex, EmitterIndex);
+		for ( int InstanceIndex=0; InstanceIndex < FrameData.NumInstances; ++InstanceIndex )
+		{
+			NiagaraID InstanceID;
+			InstanceID.Index = asint({ParameterName}_CacheData[FrameData.Int32DataOffset + ((AttributeIndex + 0) * FrameData.NumInstances) + InstanceIndex]);
+			InstanceID.AcquireTag = asint({ParameterName}_CacheData[FrameData.Int32DataOffset + ((AttributeIndex + 1) * FrameData.NumInstances) + InstanceIndex]);
+			if (InstanceID.Index == ID.Index && InstanceID.AcquireTag == ID.AcquireTag)
+			{
+				bValid = true;
+				Value = InstanceIndex;
+				return;
+			}
+		}
+	}
+}
diff --git a/Engine/Plugins/FX/Niagara/Source/Niagara/Private/DataInterface/NiagaraDataInterfaceSimCacheReader.cpp b/Engine/Plugins/FX/Niagara/Source/Niagara/Private/DataInterface/NiagaraDataInterfaceSimCacheReader.cpp
--- a/Engine/Plugins/FX/Niagara/Source/Niagara/Private/DataInterface/NiagaraDataInterfaceSimCacheReader.cpp
+++ b/Engine/Plugins/FX/Niagara/Source/Niagara/Private/DataInterface/NiagaraDataInterfaceSimCacheReader.cpp
@@ -34,6 +34,7 @@ namespace NDISimCacheReaderLocal
 	static const FName GetInstanceQuat("GetInstanceQuat");
 	static const FName GetInstancePosition("GetInstancePosition");
 	static const FName GetInstanceID("GetInstanceID");
+	static const FName GetInstanceIndexFromID("GetInstanceIndexFromID");
 
 	static const FName NAME_Attribute("Attribute");
 
@@ -110,6 +111,11 @@ namespace NDISimCacheReaderLocal
 			}
 		}
 
+		if (FunctionName == GetInstanceIndexFromID)
+		{
+			return FNiagaraTypeDefinition::GetIDDef();
+		}
+
 		check(false);
 		return FNiagaraTypeDefinition::GetFloatDef();
 	}
@@ -375,6 +381,22 @@ void UNiagaraDataInterfaceSimCacheReader::GetFunctionsInternal(TArray<FNiagaraFu
 		Sig.Name = GenFuncInfo.FunctionName;
 		Sig.Outputs.Emplace(GenFuncInfo.TypeDef, TEXT("Value"));
 	}
+
+	// Index from ID
+	{
+		FNiagaraFunctionSignature& Sig = OutFunctions.Add_GetRef(DefaultSignature);
+		Sig.Name = GetInstanceIndexFromID;
+		Sig.Inputs.Emplace(FNiagaraTypeDefinition::GetIntDef(), TEXT("FrameIndex"));
+		Sig.Inputs.Emplace(FNiagaraTypeDefinition::GetIDDef(), TEXT("ID"));
+		Sig.Outputs.Emplace(FNiagaraTypeDefinition::GetBoolDef(), TEXT("Valid"));
+		Sig.Outputs.Emplace(FNiagaraTypeDefinition::GetIntDef(), TEXT("InstanceIndex"));
+		Sig.FunctionSpecifiers.Emplace(NAME_Attribute);
+		Sig.bExperimental = true;
+	#if WITH_EDITORONLY_DATA
+		Sig.ExperimentalMessage = LOCTEXT("GetInstanceIndexFromIDExperimental", "This functions is experimental and not optimal as it may perform a linear search.");
+	#endif
+		Sig.SetDescription(LOCTEXT("GetInstanceIndexFromIDDesc", "Returns the Instance Index if the ID was found otherwise -1."));
+	}
 }
 #endif
 
@@ -443,8 +465,14 @@ void UNiagaraDataInterfaceSimCacheReader::GetVMExternalFunction(const FVMExterna
 	else if (BindingInfo.Name == GetInstanceID)
 	{
 		const int32 AttributeIndex = AttributeIndexFromVariable(InstanceSimCache, BindingInfo.Name, BindingInfo.FunctionSpecifiers[0].Value, InstanceData->EmitterIndex);
-		OutFunc = FVMExternalFunction::CreateUObject(this, &UNiagaraDataInterfaceSimCacheReader::VMGetInstanceID, AttributeIndex, false);
+		OutFunc = FVMExternalFunction::CreateUObject(this, &UNiagaraDataInterfaceSimCacheReader::VMGetInstanceID, AttributeIndex);
+	}
+	else if (BindingInfo.Name == GetInstanceIndexFromID)
+	{
+		const int32 AttributeIndex = AttributeIndexFromVariable(InstanceSimCache, BindingInfo.Name, BindingInfo.FunctionSpecifiers[0].Value, InstanceData->EmitterIndex);
+		OutFunc = FVMExternalFunction::CreateUObject(this, &UNiagaraDataInterfaceSimCacheReader::VMGetInstanceIndexFromID, AttributeIndex);
 	}
+	
 }
 
 #if WITH_EDITORONLY_DATA
@@ -491,6 +519,16 @@ bool UNiagaraDataInterfaceSimCacheReader::GetFunctionHLSL(const FNiagaraDataInte
 		}
 	}
 
+	if (FunctionInfo.DefinitionName == GetInstanceIndexFromID)
+	{
+		OutHLSL.Appendf(TEXT("void %s(int FrameIndex, NiagaraID ID, out bool bValid, out int Value)\n"), *FunctionInfo.InstanceName);
+		OutHLSL.Append(TEXT("{\n"));
+		OutHLSL.Appendf(TEXT("	int AttributeIndex = %s_ComponentOffsets[%d][%d];\n"), *ParamInfo.DataInterfaceHLSLSymbol, Float4Index, ElementIndex);
+		OutHLSL.Appendf(TEXT("	GetInstanceIndexFromID_%s(FrameIndex, %s_EmitterIndex, AttributeIndex, ID, bValid, Value);\n"), *ParamInfo.DataInterfaceHLSLSymbol, *ParamInfo.DataInterfaceHLSLSymbol);
+		OutHLSL.Append(TEXT("}\n"));
+		return true;
+	}
+
 	return true;
 }
 #endif
@@ -705,10 +743,62 @@ void UNiagaraDataInterfaceSimCacheReader::VMGetInstancePosition(FVectorVMExterna
 	VMInstanceGetFunction<FVector3f>(Context, FVector3f::ZeroVector, AttributeIndex, [](FNiagaraSimCacheAttributeReaderHelper& AttributeReader, int32 InstanceIndex) { return AttributeReader.ReadFloat3f(InstanceIndex); });
 }
 
-void UNiagaraDataInterfaceSimCacheReader::VMGetInstanceID(FVectorVMExternalFunctionContext& Context, int AttributeIndex, bool bRebase)
+void UNiagaraDataInterfaceSimCacheReader::VMGetInstanceID(FVectorVMExternalFunctionContext& Context, int AttributeIndex)
 {
 	using namespace NDISimCacheReaderLocal;
 	VMInstanceGetFunction<FNiagaraID>(Context, FNiagaraID(), AttributeIndex, [](FNiagaraSimCacheAttributeReaderHelper& AttributeReader, int32 InstanceIndex) { return AttributeReader.ReadID(InstanceIndex); });
 }
 
+void UNiagaraDataInterfaceSimCacheReader::VMGetInstanceIndexFromID(FVectorVMExternalFunctionContext& Context, int AttributeIndex)
+{
+	using namespace NDISimCacheReaderLocal;
+
+	VectorVM::FUserPtrHandler<FInstanceData_GameThread> InstanceData_GT(Context);
+	FNDIInputParam<int32>		InFrameIndex(Context);
+	FNDIInputParam<FNiagaraID>	InID(Context);
+	FNDIOutputParam<bool>		OutValid(Context);
+	FNDIOutputParam<int32>		OutValue(Context);
+
+	if (UNiagaraSimCache* InstanceSimCache = InstanceData_GT->WeakSimCache.Get())
+	{
+		FNiagaraSimCacheAttributeReaderHelper AttributeReader;
+		int CachedFrameIndex = INDEX_NONE;
+
+		for (int32 i = 0; i < Context.GetNumInstances(); ++i)
+		{
+			const int32 FrameIndex = InFrameIndex.GetAndAdvance();
+			const FNiagaraID ID = InID.GetAndAdvance();
+
+			if (CachedFrameIndex != FrameIndex)
+			{
+				CachedFrameIndex = FrameIndex;
+				AttributeReader.Initialize(InstanceSimCache, InstanceData_GT->EmitterIndex, AttributeIndex, FrameIndex);
+			}
+
+			int32 InstanceIndex = INDEX_NONE;
+			if (AttributeReader.IsValid())
+			{
+				for (int32 j = 0; j < AttributeReader.GetNumInstances(); ++j)
+				{
+					if (AttributeReader.ReadID(j) == ID)
+					{
+						InstanceIndex = j;
+						break;
+					}
+				}
+			}
+			OutValid.SetAndAdvance(InstanceIndex != INDEX_NONE);
+			OutValue.SetAndAdvance(InstanceIndex);
+		}
+	}
+	else
+	{
+		for (int32 i = 0; i < Context.GetNumInstances(); ++i)
+		{
+			OutValid.SetAndAdvance(false);
+			OutValue.SetAndAdvance(INDEX_NONE);
+		}
+	}
+}
+
 #undef LOCTEXT_NAMESPACE
diff --git a/Engine/Plugins/FX/Niagara/Source/Niagara/Private/DataInterface/NiagaraDataInterfaceSimCacheReader.h b/Engine/Plugins/FX/Niagara/Source/Niagara/Private/DataInterface/NiagaraDataInterfaceSimCacheReader.h
--- a/Engine/Plugins/FX/Niagara/Source/Niagara/Private/DataInterface/NiagaraDataInterfaceSimCacheReader.h
+++ b/Engine/Plugins/FX/Niagara/Source/Niagara/Private/DataInterface/NiagaraDataInterfaceSimCacheReader.h
@@ -85,5 +85,6 @@ private:
 	void VMGetInstanceColor(FVectorVMExternalFunctionContext& Context, int AttributeIndex);
 	void VMGetInstanceQuat(FVectorVMExternalFunctionContext& Context, int AttributeIndex, bool bRebase);
 	void VMGetInstancePosition(FVectorVMExternalFunctionContext& Context, int AttributeIndex, bool bRebase);
-	void VMGetInstanceID(FVectorVMExternalFunctionContext& Context, int AttributeIndex, bool bRebase);
+	void VMGetInstanceID(FVectorVMExternalFunctionContext& Context, int AttributeIndex);
+	void VMGetInstanceIndexFromID(FVectorVMExternalFunctionContext& Context, int AttributeIndex);
 };
