Remove deprecated UAnimNotify & UAnimNotifyState Notify* members

#Synced-CL 48563746

diff --git a/Engine/Plugins/FX/Niagara/Source/NiagaraAnimNotifies/Private/AnimNotifyState_TimedNiagaraEffect.cpp b/Engine/Plugins/FX/Niagara/Source/NiagaraAnimNotifies/Private/AnimNotifyState_TimedNiagaraEffect.cpp
--- a/Engine/Plugins/FX/Niagara/Source/NiagaraAnimNotifies/Private/AnimNotifyState_TimedNiagaraEffect.cpp
+++ b/Engine/Plugins/FX/Niagara/Source/NiagaraAnimNotifies/Private/AnimNotifyState_TimedNiagaraEffect.cpp
@@ -72,10 +72,6 @@ UFXSystemComponent* UAnimNotifyState_TimedNiagaraEffect::GetSpawnedEffect(UMeshC
 	return nullptr;
 }
 
-void UAnimNotifyState_TimedNiagaraEffect::NotifyBegin(USkeletalMeshComponent * MeshComp, class UAnimSequenceBase * Animation, float TotalDuration)
-{
-}
-
 void UAnimNotifyState_TimedNiagaraEffect::NotifyBegin(USkeletalMeshComponent * MeshComp, class UAnimSequenceBase * Animation, float TotalDuration, const FAnimNotifyEventReference& EventReference)
 {
 	if (UFXSystemComponent* Component = SpawnEffect(MeshComp, Animation))
@@ -87,10 +83,6 @@ void UAnimNotifyState_TimedNiagaraEffect::NotifyBegin(USkeletalMeshComponent * M
 	Super::NotifyBegin(MeshComp, Animation, TotalDuration, EventReference);
 }
 
-void UAnimNotifyState_TimedNiagaraEffect::NotifyEnd(USkeletalMeshComponent * MeshComp, class UAnimSequenceBase * Animation)
-{
-}
-
 void UAnimNotifyState_TimedNiagaraEffect::NotifyEnd(USkeletalMeshComponent * MeshComp, class UAnimSequenceBase * Animation, const FAnimNotifyEventReference& EventReference)
 {
 	if (UFXSystemComponent* FXComponent = GetSpawnedEffect(MeshComp))
@@ -182,10 +174,6 @@ void UAnimNotifyState_TimedNiagaraEffectAdvanced::PostLoad()
 }
 #endif
 
-void UAnimNotifyState_TimedNiagaraEffectAdvanced::NotifyBegin(USkeletalMeshComponent* MeshComp, class UAnimSequenceBase* Animation, float TotalDuration)
-{
-}
-
 void UAnimNotifyState_TimedNiagaraEffectAdvanced::NotifyBegin(USkeletalMeshComponent* MeshComp, class UAnimSequenceBase* Animation, float TotalDuration, const FAnimNotifyEventReference& EventReference)
 {
 	Super::NotifyBegin(MeshComp, Animation, TotalDuration, EventReference);
@@ -195,20 +183,12 @@ void UAnimNotifyState_TimedNiagaraEffectAdvanced::NotifyBegin(USkeletalMeshCompo
 	NewInfo.Elapsed = 0.0f;
 }
 
-void UAnimNotifyState_TimedNiagaraEffectAdvanced::NotifyEnd(USkeletalMeshComponent* MeshComp, class UAnimSequenceBase* Animation)
-{
-}
-
 void UAnimNotifyState_TimedNiagaraEffectAdvanced::NotifyEnd(USkeletalMeshComponent* MeshComp, class UAnimSequenceBase* Animation, const FAnimNotifyEventReference& EventReference)
 {
 	Super::NotifyEnd(MeshComp, Animation, EventReference);
 	ProgressInfoMap.Remove(MeshComp);
 }
 
-void UAnimNotifyState_TimedNiagaraEffectAdvanced::NotifyTick(USkeletalMeshComponent* MeshComp, UAnimSequenceBase* Animation, float FrameDeltaTime)
-{
-}
-
 void UAnimNotifyState_TimedNiagaraEffectAdvanced::NotifyTick(USkeletalMeshComponent* MeshComp, UAnimSequenceBase* Animation, float FrameDeltaTime, const FAnimNotifyEventReference& EventReference)
 {
 	Super::NotifyTick(MeshComp, Animation, FrameDeltaTime, EventReference);
diff --git a/Engine/Plugins/FX/Niagara/Source/NiagaraAnimNotifies/Private/AnimNotify_PlayNiagaraEffect.cpp b/Engine/Plugins/FX/Niagara/Source/NiagaraAnimNotifies/Private/AnimNotify_PlayNiagaraEffect.cpp
--- a/Engine/Plugins/FX/Niagara/Source/NiagaraAnimNotifies/Private/AnimNotify_PlayNiagaraEffect.cpp
+++ b/Engine/Plugins/FX/Niagara/Source/NiagaraAnimNotifies/Private/AnimNotify_PlayNiagaraEffect.cpp
@@ -75,10 +75,6 @@ void UAnimNotify_PlayNiagaraEffect::ValidateAssociatedAssets()
 }
 #endif
 
-void UAnimNotify_PlayNiagaraEffect::Notify(class USkeletalMeshComponent* MeshComp, class UAnimSequenceBase* Animation)
-{
-}
-
 void UAnimNotify_PlayNiagaraEffect::Notify(class USkeletalMeshComponent* MeshComp, class UAnimSequenceBase* Animation, const FAnimNotifyEventReference& EventReference)
 {
 	//Store the spawned effect in a protected variable
@@ -135,4 +131,3 @@ UFXSystemComponent* UAnimNotify_PlayNiagaraEffect::GetSpawnedEffect()
 {
 	return SpawnedEffect;
 }
-
diff --git a/Engine/Plugins/FX/Niagara/Source/NiagaraAnimNotifies/Public/AnimNotifyState_TimedNiagaraEffect.h b/Engine/Plugins/FX/Niagara/Source/NiagaraAnimNotifies/Public/AnimNotifyState_TimedNiagaraEffect.h
--- a/Engine/Plugins/FX/Niagara/Source/NiagaraAnimNotifies/Public/AnimNotifyState_TimedNiagaraEffect.h
+++ b/Engine/Plugins/FX/Niagara/Source/NiagaraAnimNotifies/Public/AnimNotifyState_TimedNiagaraEffect.h
@@ -60,12 +60,6 @@ class UAnimNotifyState_TimedNiagaraEffect : public UAnimNotifyState
 	// the emitters.
 	UPROPERTY(EditAnywhere, BlueprintReadWrite, Category = NiagaraSystem, meta = (DisplayName = "Destroy Immediately", ToolTip = "Whether the Niagara system should be immediately destroyed at the end of the notify state or be allowed to finish"))
 	bool bDestroyAtEnd;
-
-	UE_DEPRECATED(5.0, "Please use the other NotifyBegin function instead")
-	NIAGARAANIMNOTIFIES_API virtual void NotifyBegin(class USkeletalMeshComponent * MeshComp, class UAnimSequenceBase * Animation, float TotalDuration) override;
-
-	UE_DEPRECATED(5.0, "Please use the other NotifyEnd function instead")
-	NIAGARAANIMNOTIFIES_API virtual void NotifyEnd(class USkeletalMeshComponent * MeshComp, class UAnimSequenceBase * Animation) override;
 	
 	NIAGARAANIMNOTIFIES_API virtual void NotifyBegin(class USkeletalMeshComponent * MeshComp, class UAnimSequenceBase * Animation, float TotalDuration, const FAnimNotifyEventReference& EventReference) override;
 	NIAGARAANIMNOTIFIES_API virtual void NotifyEnd(class USkeletalMeshComponent * MeshComp, class UAnimSequenceBase * Animation, const FAnimNotifyEventReference& EventReference) override;
@@ -110,15 +104,6 @@ public:
 #if WITH_EDITORONLY_DATA
 	NIAGARAANIMNOTIFIES_API virtual void PostLoad() override;
 #endif
-
-	UE_DEPRECATED(5.0, "Please use the other NotifyBegin function instead")
-	NIAGARAANIMNOTIFIES_API virtual void NotifyBegin(class USkeletalMeshComponent* MeshComp, class UAnimSequenceBase* Animation, float TotalDuration) override;
-
-	UE_DEPRECATED(5.0, "Please use the other NotifyEnd function instead")
-	NIAGARAANIMNOTIFIES_API virtual void NotifyEnd(class USkeletalMeshComponent* MeshComp, class UAnimSequenceBase* Animation) override;
-
-	UE_DEPRECATED(5.0, "Please use the other NotifyTick function instead")
-	NIAGARAANIMNOTIFIES_API virtual void NotifyTick(USkeletalMeshComponent * MeshComp, UAnimSequenceBase * Animation, float FrameDeltaTime) override;
 	
 	NIAGARAANIMNOTIFIES_API virtual void NotifyBegin(class USkeletalMeshComponent* MeshComp, class UAnimSequenceBase* Animation, float TotalDuration, const FAnimNotifyEventReference& EventReference) override;
 	NIAGARAANIMNOTIFIES_API virtual void NotifyEnd(class USkeletalMeshComponent* MeshComp, class UAnimSequenceBase* Animation, const FAnimNotifyEventReference& EventReference) override;
diff --git a/Engine/Plugins/FX/Niagara/Source/NiagaraAnimNotifies/Public/AnimNotify_PlayNiagaraEffect.h b/Engine/Plugins/FX/Niagara/Source/NiagaraAnimNotifies/Public/AnimNotify_PlayNiagaraEffect.h
--- a/Engine/Plugins/FX/Niagara/Source/NiagaraAnimNotifies/Public/AnimNotify_PlayNiagaraEffect.h
+++ b/Engine/Plugins/FX/Niagara/Source/NiagaraAnimNotifies/Public/AnimNotify_PlayNiagaraEffect.h
@@ -34,8 +34,6 @@ public:
 	// Begin UAnimNotify interface
 	NIAGARAANIMNOTIFIES_API virtual FString GetNotifyName_Implementation() const override;
 	
-	UE_DEPRECATED(5.0, "Please use the other Notify function instead")
-	NIAGARAANIMNOTIFIES_API virtual void Notify(USkeletalMeshComponent* MeshComp, UAnimSequenceBase* Animation) override;
 	NIAGARAANIMNOTIFIES_API virtual void Notify(USkeletalMeshComponent* MeshComp, UAnimSequenceBase* Animation, const FAnimNotifyEventReference& EventReference) override;
 #if WITH_EDITOR
 	NIAGARAANIMNOTIFIES_API virtual void ValidateAssociatedAssets() override;
@@ -88,7 +86,3 @@ public:
 	UPROPERTY(EditAnywhere, BlueprintReadWrite, Category = "AnimNotify", meta = (AnimNotifyBoneName = "true"))
 	FName SocketName;
 };
-
-
-
-
diff --git a/Engine/Plugins/Runtime/GameplayAbilities/Source/GameplayAbilities/Private/AnimNotify_GameplayCue.cpp b/Engine/Plugins/Runtime/GameplayAbilities/Source/GameplayAbilities/Private/AnimNotify_GameplayCue.cpp
--- a/Engine/Plugins/Runtime/GameplayAbilities/Source/GameplayAbilities/Private/AnimNotify_GameplayCue.cpp
+++ b/Engine/Plugins/Runtime/GameplayAbilities/Source/GameplayAbilities/Private/AnimNotify_GameplayCue.cpp
@@ -74,10 +74,6 @@ UAnimNotify_GameplayCue::UAnimNotify_GameplayCue()
 {
 }
 
-void UAnimNotify_GameplayCue::Notify(USkeletalMeshComponent* MeshComp, UAnimSequenceBase* Animation)
-{
-}
-
 void UAnimNotify_GameplayCue::Notify(USkeletalMeshComponent* MeshComp, UAnimSequenceBase* Animation, const FAnimNotifyEventReference& EventReference)
 {
 	Super::Notify(MeshComp, Animation, EventReference);
@@ -113,10 +109,6 @@ UAnimNotify_GameplayCueState::UAnimNotify_GameplayCueState()
 {
 }
 
-void UAnimNotify_GameplayCueState::NotifyBegin(USkeletalMeshComponent* MeshComp, UAnimSequenceBase* Animation, float TotalDuration)
-{
-}
-
 void UAnimNotify_GameplayCueState::NotifyBegin(USkeletalMeshComponent* MeshComp, UAnimSequenceBase* Animation, float TotalDuration, const FAnimNotifyEventReference& EventReference)
 {
 	Super::NotifyBegin(MeshComp, Animation, TotalDuration, EventReference);
@@ -133,10 +125,6 @@ void UAnimNotify_GameplayCueState::NotifyBegin(USkeletalMeshComponent* MeshComp,
 #endif // #if WITH_EDITORONLY_DATA
 }
 
-void UAnimNotify_GameplayCueState::NotifyTick(USkeletalMeshComponent * MeshComp, UAnimSequenceBase * Animation, float FrameDeltaTime)
-{
-}
-
 void UAnimNotify_GameplayCueState::NotifyTick(USkeletalMeshComponent * MeshComp, UAnimSequenceBase * Animation, float FrameDeltaTime, const FAnimNotifyEventReference& EventReference)
 {
 	AActor* OwnerActor = MeshComp->GetOwner();
@@ -178,10 +166,6 @@ FString UAnimNotify_GameplayCueState::GetNotifyName_Implementation() const
 	return DisplayName;
 }
 
-void UAnimNotify_GameplayCueState::NotifyEnd(USkeletalMeshComponent* MeshComp, UAnimSequenceBase* Animation)
-{
-}
-
 void UAnimNotify_GameplayCueState::NotifyEnd(USkeletalMeshComponent* MeshComp, UAnimSequenceBase* Animation, const FAnimNotifyEventReference& EventReference)
 {
 	Super::NotifyEnd(MeshComp, Animation, EventReference);
@@ -195,4 +179,3 @@ bool UAnimNotify_GameplayCueState::CanBePlaced(UAnimSequenceBase* Animation) con
 	return (Animation && Animation->IsA(UAnimMontage::StaticClass()));
 }
 #endif // #if WITH_EDITOR
-
diff --git a/Engine/Plugins/Runtime/GameplayAbilities/Source/GameplayAbilities/Public/AnimNotify_GameplayCue.h b/Engine/Plugins/Runtime/GameplayAbilities/Source/GameplayAbilities/Public/AnimNotify_GameplayCue.h
--- a/Engine/Plugins/Runtime/GameplayAbilities/Source/GameplayAbilities/Public/AnimNotify_GameplayCue.h
+++ b/Engine/Plugins/Runtime/GameplayAbilities/Source/GameplayAbilities/Public/AnimNotify_GameplayCue.h
@@ -24,9 +24,6 @@ public:
 
 	UAnimNotify_GameplayCue();
 
-	UE_DEPRECATED(5.0, "Please use the other Notify function instead")
-	virtual void Notify(USkeletalMeshComponent* MeshComp, UAnimSequenceBase* Animation) override;
-
 	virtual void Notify(USkeletalMeshComponent* MeshComp, UAnimSequenceBase* Animation, const FAnimNotifyEventReference& EventReference) override;
 
 	FString GetNotifyName_Implementation() const override;
@@ -56,14 +53,6 @@ class UAnimNotify_GameplayCueState : public UAnimNotifyState
 public:
 
 	UAnimNotify_GameplayCueState();
-	UE_DEPRECATED(5.0, "Please use the other NotifyBegin function instead")
-	virtual void NotifyBegin(USkeletalMeshComponent* MeshComp, UAnimSequenceBase* Animation, float TotalDuration) override;
-
-	UE_DEPRECATED(5.0, "Please use the other NotifyTick function instead")
-	virtual void NotifyTick(USkeletalMeshComponent* MeshComp, UAnimSequenceBase* Animation, float FrameDeltaTime) override;
-
-	UE_DEPRECATED(5.0, "Please use the other NotifyEnd function instead")
-	virtual void NotifyEnd(USkeletalMeshComponent* MeshComp, UAnimSequenceBase* Animation) override;
 
 	virtual void NotifyBegin(USkeletalMeshComponent* MeshComp, UAnimSequenceBase* Animation, float TotalDuration, const FAnimNotifyEventReference& EventReference) override;
 	virtual void NotifyTick(USkeletalMeshComponent* MeshComp, UAnimSequenceBase* Animation, float FrameDeltaTime, const FAnimNotifyEventReference& EventReference) override;
diff --git a/Engine/Source/Runtime/Engine/Classes/Animation/AnimNotifies/AnimNotify.h b/Engine/Source/Runtime/Engine/Classes/Animation/AnimNotifies/AnimNotify.h
--- a/Engine/Source/Runtime/Engine/Classes/Animation/AnimNotifies/AnimNotify.h
+++ b/Engine/Source/Runtime/Engine/Classes/Animation/AnimNotifies/AnimNotify.h
@@ -82,8 +82,6 @@ class UAnimNotify : public UObject
 	virtual bool ShouldFireInEditor() { return bShouldFireInEditor; }
 #endif
 
-	UE_DEPRECATED(5.0, "Please use the other Notify function instead")
-	ENGINE_API virtual void Notify(USkeletalMeshComponent* MeshComp, UAnimSequenceBase* Animation);
 	ENGINE_API virtual void Notify(USkeletalMeshComponent* MeshComp, UAnimSequenceBase* Animation, const FAnimNotifyEventReference& EventReference);
 	ENGINE_API virtual void BranchingPointNotify(FBranchingPointNotifyPayload& BranchingPointPayload);
 
@@ -131,6 +129,3 @@ private:
 	/* The mesh we're currently triggering a UAnimNotify for (so we can retrieve per instance information) */
 	class USkeletalMeshComponent* MeshContext;
 };
-
-
-
diff --git a/Engine/Source/Runtime/Engine/Classes/Animation/AnimNotifies/AnimNotifyState.h b/Engine/Source/Runtime/Engine/Classes/Animation/AnimNotifies/AnimNotifyState.h
--- a/Engine/Source/Runtime/Engine/Classes/Animation/AnimNotifies/AnimNotifyState.h
+++ b/Engine/Source/Runtime/Engine/Classes/Animation/AnimNotifies/AnimNotifyState.h
@@ -61,13 +61,6 @@ class UAnimNotifyState : public UObject
 	virtual bool ShouldFireInEditor() { return bShouldFireInEditor; }
 #endif
 
-	UE_DEPRECATED(5.0, "This function is deprecated. Use the other NotifyBegin instead.")
-	ENGINE_API virtual void NotifyBegin(USkeletalMeshComponent * MeshComp, UAnimSequenceBase * Animation, float TotalDuration);
-	UE_DEPRECATED(5.0, "This function is deprecated. Use the other NotifyTick instead.")
-	ENGINE_API virtual void NotifyTick(USkeletalMeshComponent * MeshComp, UAnimSequenceBase * Animation, float FrameDeltaTime);
-	UE_DEPRECATED(5.0, "This function is deprecated. Use the other NotifyEnd instead.")
-	ENGINE_API virtual void NotifyEnd(USkeletalMeshComponent * MeshComp, UAnimSequenceBase * Animation);
-
 	ENGINE_API virtual void NotifyBegin(USkeletalMeshComponent * MeshComp, UAnimSequenceBase * Animation, float TotalDuration, const FAnimNotifyEventReference& EventReference);
 	ENGINE_API virtual void NotifyTick(USkeletalMeshComponent * MeshComp, UAnimSequenceBase * Animation, float FrameDeltaTime, const FAnimNotifyEventReference& EventReference);
 	ENGINE_API virtual void NotifyEnd(USkeletalMeshComponent * MeshComp, UAnimSequenceBase * Animation, const FAnimNotifyEventReference& EventReference);
@@ -107,6 +100,3 @@ class UAnimNotifyState : public UObject
 protected:
 	ENGINE_API UObject* GetContainingAsset() const;
 };
-
-
-
diff --git a/Engine/Source/Runtime/Engine/Classes/Animation/AnimNotifies/AnimNotifyState_TimedParticleEffect.h b/Engine/Source/Runtime/Engine/Classes/Animation/AnimNotifies/AnimNotifyState_TimedParticleEffect.h
--- a/Engine/Source/Runtime/Engine/Classes/Animation/AnimNotifies/AnimNotifyState_TimedParticleEffect.h
+++ b/Engine/Source/Runtime/Engine/Classes/Animation/AnimNotifies/AnimNotifyState_TimedParticleEffect.h
@@ -56,13 +56,6 @@ class UAnimNotifyState_TimedParticleEffect : public UAnimNotifyState
 #if WITH_EDITOR
 	ENGINE_API virtual void PreEditChange(FProperty* PropertyAboutToChange) override;
 #endif
-
-	UE_DEPRECATED(5.0, "Please use the other NotifyBegin function instead")
-	ENGINE_API virtual void NotifyBegin(class USkeletalMeshComponent * MeshComp, class UAnimSequenceBase * Animation, float TotalDuration) override;
-	UE_DEPRECATED(5.0, "Please use the other NotifyTick function instead")
-	ENGINE_API virtual void NotifyTick(class USkeletalMeshComponent * MeshComp, class UAnimSequenceBase * Animation, float FrameDeltaTime) override;
-	UE_DEPRECATED(5.0, "Please use the other NotifyEnd function instead")
-	ENGINE_API virtual void NotifyEnd(class USkeletalMeshComponent * MeshComp, class UAnimSequenceBase * Animation) override;
 	
 	ENGINE_API virtual void NotifyBegin(class USkeletalMeshComponent * MeshComp, class UAnimSequenceBase * Animation, float TotalDuration, const FAnimNotifyEventReference& EventReference) override;
 	ENGINE_API virtual void NotifyTick(class USkeletalMeshComponent * MeshComp, class UAnimSequenceBase * Animation, float FrameDeltaTime, const FAnimNotifyEventReference& EventReference) override;
diff --git a/Engine/Source/Runtime/Engine/Classes/Animation/AnimNotifies/AnimNotifyState_Trail.h b/Engine/Source/Runtime/Engine/Classes/Animation/AnimNotifies/AnimNotifyState_Trail.h
--- a/Engine/Source/Runtime/Engine/Classes/Animation/AnimNotifies/AnimNotifyState_Trail.h
+++ b/Engine/Source/Runtime/Engine/Classes/Animation/AnimNotifies/AnimNotifyState_Trail.h
@@ -76,19 +76,9 @@ class UAnimNotifyState_Trail : public UAnimNotifyState
 	/** Helper function for outside code to get PSC that we are using */
 	ENGINE_API UParticleSystemComponent* GetParticleSystemComponent(USkeletalMeshComponent* MeshComp) const;
 
-	UE_DEPRECATED(5.0, "Please use the other NotifyBegin function instead")
-	ENGINE_API virtual void NotifyBegin(class USkeletalMeshComponent * MeshComp, class UAnimSequenceBase * Animation, float TotalDuration) override;
-	UE_DEPRECATED(5.0, "Please use the other NotifyTick function instead")
-	ENGINE_API virtual void NotifyTick(class USkeletalMeshComponent * MeshComp, class UAnimSequenceBase * Animation, float FrameDeltaTime) override;
-	UE_DEPRECATED(5.0, "Please use the other NotifyEnd function instead")
-	ENGINE_API virtual void NotifyEnd(class USkeletalMeshComponent * MeshComp, class UAnimSequenceBase * Animation) override;
-
 	ENGINE_API virtual void NotifyBegin(USkeletalMeshComponent* MeshComp, UAnimSequenceBase* Animation, float TotalDuration, const FAnimNotifyEventReference& EventReference) override;
 	ENGINE_API virtual void NotifyTick(USkeletalMeshComponent* MeshComp, UAnimSequenceBase* Animation, float FrameDeltaTime, const FAnimNotifyEventReference& EventReference) override;
 	ENGINE_API virtual void NotifyEnd(USkeletalMeshComponent* MeshComp, UAnimSequenceBase* Animation, const FAnimNotifyEventReference& EventReference) override;
 
 	ENGINE_API bool ValidateInput(USkeletalMeshComponent* MeshComp, bool bReportErrors = false);
 };
-
-
-
diff --git a/Engine/Source/Runtime/Engine/Classes/Animation/AnimNotifies/AnimNotify_PauseClothingSimulation.h b/Engine/Source/Runtime/Engine/Classes/Animation/AnimNotifies/AnimNotify_PauseClothingSimulation.h
--- a/Engine/Source/Runtime/Engine/Classes/Animation/AnimNotifies/AnimNotify_PauseClothingSimulation.h
+++ b/Engine/Source/Runtime/Engine/Classes/Animation/AnimNotifies/AnimNotify_PauseClothingSimulation.h
@@ -18,12 +18,7 @@ public:
 	// Begin UAnimNotify interface
 	ENGINE_API virtual FString GetNotifyName_Implementation() const override;
 	
-	UE_DEPRECATED(5.0, "Please use the other Notify function instead")
-	ENGINE_API virtual void Notify(USkeletalMeshComponent* MeshComp, UAnimSequenceBase* Animation) override;
 	ENGINE_API virtual void Notify(USkeletalMeshComponent* MeshComp, UAnimSequenceBase* Animation, const FAnimNotifyEventReference& EventReference) override;
 	// End UAnimNotify interface
 
 };
-
-
-
diff --git a/Engine/Source/Runtime/Engine/Classes/Animation/AnimNotifies/AnimNotify_PlayParticleEffect.h b/Engine/Source/Runtime/Engine/Classes/Animation/AnimNotifies/AnimNotify_PlayParticleEffect.h
--- a/Engine/Source/Runtime/Engine/Classes/Animation/AnimNotifies/AnimNotify_PlayParticleEffect.h
+++ b/Engine/Source/Runtime/Engine/Classes/Animation/AnimNotifies/AnimNotify_PlayParticleEffect.h
@@ -32,8 +32,6 @@ public:
 	// Begin UAnimNotify interface
 	ENGINE_API virtual FString GetNotifyName_Implementation() const override;
 
-	UE_DEPRECATED(5.0, "Please use the other Notify function instead")
-	ENGINE_API virtual void Notify(USkeletalMeshComponent* MeshComp, UAnimSequenceBase* Animation) override;
 	ENGINE_API virtual void Notify(USkeletalMeshComponent* MeshComp, UAnimSequenceBase* Animation, const FAnimNotifyEventReference& EventReference) override;
 #if WITH_EDITOR
 	ENGINE_API virtual void ValidateAssociatedAssets() override;
@@ -74,6 +72,3 @@ public:
 	UPROPERTY(EditAnywhere, BlueprintReadWrite, Category = "AnimNotify")
 	FName SocketName;
 };
-
-
-
diff --git a/Engine/Source/Runtime/Engine/Classes/Animation/AnimNotifies/AnimNotify_PlaySound.h b/Engine/Source/Runtime/Engine/Classes/Animation/AnimNotifies/AnimNotify_PlaySound.h
--- a/Engine/Source/Runtime/Engine/Classes/Animation/AnimNotifies/AnimNotify_PlaySound.h
+++ b/Engine/Source/Runtime/Engine/Classes/Animation/AnimNotifies/AnimNotify_PlaySound.h
@@ -23,8 +23,6 @@ public:
 
 	// Begin UAnimNotify interface
 	ENGINE_API virtual FString GetNotifyName_Implementation() const override;
-	UE_DEPRECATED(5.0, "Please use the other Notify function instead")
-	ENGINE_API virtual void Notify(USkeletalMeshComponent* MeshComp, UAnimSequenceBase* Animation) override;
 	ENGINE_API virtual void Notify(USkeletalMeshComponent* MeshComp, UAnimSequenceBase* Animation, const FAnimNotifyEventReference& EventReference) override;
 #if WITH_EDITOR
 	ENGINE_API virtual void ValidateAssociatedAssets() override;
diff --git a/Engine/Source/Runtime/Engine/Classes/Animation/AnimNotifies/AnimNotify_ResetClothingSimulation.h b/Engine/Source/Runtime/Engine/Classes/Animation/AnimNotifies/AnimNotify_ResetClothingSimulation.h
--- a/Engine/Source/Runtime/Engine/Classes/Animation/AnimNotifies/AnimNotify_ResetClothingSimulation.h
+++ b/Engine/Source/Runtime/Engine/Classes/Animation/AnimNotifies/AnimNotify_ResetClothingSimulation.h
@@ -17,12 +17,7 @@ public:
 
 	// Begin UAnimNotify interface
 	ENGINE_API virtual FString GetNotifyName_Implementation() const override;
-	UE_DEPRECATED(5.0, "Please use the other Notify function instead")
-	ENGINE_API virtual void Notify(USkeletalMeshComponent* MeshComp, UAnimSequenceBase* Animation) override;
 	ENGINE_API virtual void Notify(USkeletalMeshComponent* MeshComp, UAnimSequenceBase* Animation, const FAnimNotifyEventReference& EventReference) override;
 	// End UAnimNotify interface
 
 };
-
-
-
diff --git a/Engine/Source/Runtime/Engine/Classes/Animation/AnimNotifies/AnimNotify_ResetDynamics.h b/Engine/Source/Runtime/Engine/Classes/Animation/AnimNotifies/AnimNotify_ResetDynamics.h
--- a/Engine/Source/Runtime/Engine/Classes/Animation/AnimNotifies/AnimNotify_ResetDynamics.h
+++ b/Engine/Source/Runtime/Engine/Classes/Animation/AnimNotifies/AnimNotify_ResetDynamics.h
@@ -17,12 +17,7 @@ public:
 
 	// Begin UAnimNotify interface
 	ENGINE_API virtual FString GetNotifyName_Implementation() const override;
-	UE_DEPRECATED(5.0, "Please use the other Notify function instead")
-	ENGINE_API virtual void Notify(USkeletalMeshComponent* MeshComp, UAnimSequenceBase* Animation) override;
 	ENGINE_API virtual void Notify(USkeletalMeshComponent* MeshComp, UAnimSequenceBase* Animation, const FAnimNotifyEventReference& EventReference) override;
 	// End UAnimNotify interface
 
 };
-
-
-
diff --git a/Engine/Source/Runtime/Engine/Classes/Animation/AnimNotifies/AnimNotify_ResumeClothingSimulation.h b/Engine/Source/Runtime/Engine/Classes/Animation/AnimNotifies/AnimNotify_ResumeClothingSimulation.h
--- a/Engine/Source/Runtime/Engine/Classes/Animation/AnimNotifies/AnimNotify_ResumeClothingSimulation.h
+++ b/Engine/Source/Runtime/Engine/Classes/Animation/AnimNotifies/AnimNotify_ResumeClothingSimulation.h
@@ -17,12 +17,7 @@ public:
 
 	// Begin UAnimNotify interface
 	ENGINE_API virtual FString GetNotifyName_Implementation() const override;
-	UE_DEPRECATED(5.0, "Please use the other Notify function instead")
-	ENGINE_API virtual void Notify(USkeletalMeshComponent* MeshComp, UAnimSequenceBase* Animation) override;
 	ENGINE_API virtual void Notify(USkeletalMeshComponent* MeshComp, UAnimSequenceBase* Animation, const FAnimNotifyEventReference& EventReference) override;
 	// End UAnimNotify interface
 
 };
-
-
-
diff --git a/Engine/Source/Runtime/Engine/Private/Animation/AnimNotify.cpp b/Engine/Source/Runtime/Engine/Private/Animation/AnimNotify.cpp
--- a/Engine/Source/Runtime/Engine/Private/Animation/AnimNotify.cpp
+++ b/Engine/Source/Runtime/Engine/Private/Animation/AnimNotify.cpp
@@ -26,17 +26,10 @@ UAnimNotify::UAnimNotify(const FObjectInitializer& ObjectInitializer)
 	bIsNativeBranchingPoint = false;
 }
 
-void UAnimNotify::Notify(class USkeletalMeshComponent* MeshComp, class UAnimSequenceBase* Animation)
-{
-}
-
 void UAnimNotify::Notify(USkeletalMeshComponent* MeshComp, UAnimSequenceBase* Animation, const FAnimNotifyEventReference& EventReference)
 {
 	USkeletalMeshComponent* PrevContext = MeshContext;
 	MeshContext = MeshComp;
-	PRAGMA_DISABLE_DEPRECATION_WARNINGS
-	Notify(MeshComp, Animation);
-	PRAGMA_ENABLE_DEPRECATION_WARNINGS
 	Received_Notify(MeshComp, Animation, EventReference);
 	MeshContext = PrevContext;
 }
@@ -112,4 +105,3 @@ UObject* UAnimNotify::GetContainingAsset() const
 	}
 	return ContainingAsset;
 }
-
diff --git a/Engine/Source/Runtime/Engine/Private/Animation/AnimNotifyState.cpp b/Engine/Source/Runtime/Engine/Private/Animation/AnimNotifyState.cpp
--- a/Engine/Source/Runtime/Engine/Private/Animation/AnimNotifyState.cpp
+++ b/Engine/Source/Runtime/Engine/Private/Animation/AnimNotifyState.cpp
@@ -25,39 +25,18 @@ UAnimNotifyState::UAnimNotifyState(const FObjectInitializer& ObjectInitializer)
 	bIsNativeBranchingPoint = false;
 }
 
-void UAnimNotifyState::NotifyBegin(class USkeletalMeshComponent * MeshComp, class UAnimSequenceBase * Animation, float TotalDuration)
-{
-}
-
-void UAnimNotifyState::NotifyTick(class USkeletalMeshComponent * MeshComp, class UAnimSequenceBase * Animation, float FrameDeltaTime)
-{
-}
-
-void UAnimNotifyState::NotifyEnd(class USkeletalMeshComponent * MeshComp, class UAnimSequenceBase * Animation)
-{
-}
-
 void UAnimNotifyState::NotifyBegin(USkeletalMeshComponent* MeshComp, UAnimSequenceBase* Animation, float TotalDuration, const FAnimNotifyEventReference& EventReference)
 {
-	PRAGMA_DISABLE_DEPRECATION_WARNINGS
-	NotifyBegin(MeshComp, Animation, TotalDuration);
-	PRAGMA_ENABLE_DEPRECATION_WARNINGS
 	Received_NotifyBegin(MeshComp, Animation, TotalDuration, EventReference);
 }
 
 void UAnimNotifyState::NotifyTick(USkeletalMeshComponent* MeshComp, UAnimSequenceBase* Animation, float FrameDeltaTime, const FAnimNotifyEventReference& EventReference)
 {
-	PRAGMA_DISABLE_DEPRECATION_WARNINGS
-	NotifyTick(MeshComp, Animation, FrameDeltaTime);
-	PRAGMA_ENABLE_DEPRECATION_WARNINGS
 	Received_NotifyTick(MeshComp, Animation, FrameDeltaTime, EventReference);
 }
 
 void UAnimNotifyState::NotifyEnd(USkeletalMeshComponent* MeshComp, UAnimSequenceBase* Animation, const FAnimNotifyEventReference& EventReference)
 {
-	PRAGMA_DISABLE_DEPRECATION_WARNINGS
-	NotifyEnd(MeshComp, Animation);
-	PRAGMA_ENABLE_DEPRECATION_WARNINGS
 	Received_NotifyEnd(MeshComp, Animation, EventReference);
 }
 
@@ -144,4 +123,3 @@ UObject* UAnimNotifyState::GetContainingAsset() const
 	}
 	return ContainingAsset;
 }
-
diff --git a/Engine/Source/Runtime/Engine/Private/Animation/AnimNotifyState_TimedParticleEffect.cpp b/Engine/Source/Runtime/Engine/Private/Animation/AnimNotifyState_TimedParticleEffect.cpp
--- a/Engine/Source/Runtime/Engine/Private/Animation/AnimNotifyState_TimedParticleEffect.cpp
+++ b/Engine/Source/Runtime/Engine/Private/Animation/AnimNotifyState_TimedParticleEffect.cpp
@@ -17,18 +17,8 @@ UAnimNotifyState_TimedParticleEffect::UAnimNotifyState_TimedParticleEffect(const
 	RotationOffset = FRotator(0.0f, 0.0f, 0.0f);
 }
 
-void UAnimNotifyState_TimedParticleEffect::NotifyBegin(USkeletalMeshComponent * MeshComp, class UAnimSequenceBase * Animation, float TotalDuration)
-{
-
-}
-
-
 void UAnimNotifyState_TimedParticleEffect::NotifyBegin(USkeletalMeshComponent * MeshComp, class UAnimSequenceBase * Animation, float TotalDuration, const FAnimNotifyEventReference& EventReference)
 {
-	// ensure deprecated path is called because a call to Super is not made
-	PRAGMA_DISABLE_DEPRECATION_WARNINGS
-	NotifyBegin(MeshComp, Animation, TotalDuration);
-	PRAGMA_ENABLE_DEPRECATION_WARNINGS
 	// Only spawn if we've got valid params
 	if(ValidateParameters(MeshComp))
 	{
@@ -37,27 +27,13 @@ void UAnimNotifyState_TimedParticleEffect::NotifyBegin(USkeletalMeshComponent *
 	Received_NotifyBegin(MeshComp, Animation, TotalDuration, EventReference);
 }
 
-void UAnimNotifyState_TimedParticleEffect::NotifyTick(USkeletalMeshComponent * MeshComp, class UAnimSequenceBase * Animation, float FrameDeltaTime)
-{
-}
-
 void UAnimNotifyState_TimedParticleEffect::NotifyTick(USkeletalMeshComponent * MeshComp, class UAnimSequenceBase * Animation, float FrameDeltaTime, const FAnimNotifyEventReference& EventReference)
 {
-	PRAGMA_DISABLE_DEPRECATION_WARNINGS
-	NotifyTick(MeshComp, Animation, FrameDeltaTime);
-	PRAGMA_ENABLE_DEPRECATION_WARNINGS
 	Received_NotifyTick(MeshComp, Animation, FrameDeltaTime, EventReference);
 }
 
-void UAnimNotifyState_TimedParticleEffect::NotifyEnd(USkeletalMeshComponent * MeshComp, class UAnimSequenceBase * Animation)
-{
-}
-
 void UAnimNotifyState_TimedParticleEffect::NotifyEnd(USkeletalMeshComponent * MeshComp, class UAnimSequenceBase * Animation, const FAnimNotifyEventReference& EventReference)
 {
-	PRAGMA_DISABLE_DEPRECATION_WARNINGS
-	NotifyEnd(MeshComp, Animation);
-	PRAGMA_ENABLE_DEPRECATION_WARNINGS
 	TArray<USceneComponent*> Children;
 	MeshComp->GetChildrenComponents(false, Children);
 
@@ -148,4 +124,3 @@ void UAnimNotifyState_TimedParticleEffect::PreEditChange(FProperty* PropertyAbou
 	}
 }
 #endif
-
diff --git a/Engine/Source/Runtime/Engine/Private/Animation/AnimNotifyState_Trail.cpp b/Engine/Source/Runtime/Engine/Private/Animation/AnimNotifyState_Trail.cpp
--- a/Engine/Source/Runtime/Engine/Private/Animation/AnimNotifyState_Trail.cpp
+++ b/Engine/Source/Runtime/Engine/Private/Animation/AnimNotifyState_Trail.cpp
@@ -79,15 +79,8 @@ float UAnimNotifyState_Trail::GetCurveWidth(USkeletalMeshComponent* MeshComp) co
 	return Width;
 }
 
-void UAnimNotifyState_Trail::NotifyBegin(class USkeletalMeshComponent* MeshComp, class UAnimSequenceBase* Animation, float TotalDuration)
-{
-}
-
 void UAnimNotifyState_Trail::NotifyBegin(class USkeletalMeshComponent * MeshComp, class UAnimSequenceBase * Animation, float TotalDuration, const FAnimNotifyEventReference& EventReference)
 {
-	PRAGMA_DISABLE_DEPRECATION_WARNINGS
-    NotifyBegin(MeshComp, Animation, TotalDuration);
-	PRAGMA_ENABLE_DEPRECATION_WARNINGS
 	bool bError = ValidateInput(MeshComp);
 
 	if (MeshComp->GetWorld()->GetNetMode() == NM_DedicatedServer)
@@ -209,15 +202,8 @@ void UAnimNotifyState_Trail::NotifyBegin(class USkeletalMeshComponent * MeshComp
 	Received_NotifyBegin(MeshComp, Animation, TotalDuration, EventReference);
 }
 
-void UAnimNotifyState_Trail::NotifyTick(class USkeletalMeshComponent* MeshComp, class UAnimSequenceBase* Animation, float FrameDeltaTime)
-{
-}
-
 void UAnimNotifyState_Trail::NotifyTick(class USkeletalMeshComponent * MeshComp, class UAnimSequenceBase * Animation, float FrameDeltaTime, const FAnimNotifyEventReference& EventReference)
 {
-	PRAGMA_DISABLE_DEPRECATION_WARNINGS
-    NotifyTick(MeshComp, Animation, FrameDeltaTime);
-	PRAGMA_ENABLE_DEPRECATION_WARNINGS
 	bool bError = ValidateInput(MeshComp, true);
 
 	if (MeshComp->GetWorld()->GetNetMode() == NM_DedicatedServer)
@@ -257,16 +243,8 @@ void UAnimNotifyState_Trail::NotifyTick(class USkeletalMeshComponent * MeshComp,
 	Received_NotifyTick(MeshComp, Animation, FrameDeltaTime, EventReference);
 }
 
-void UAnimNotifyState_Trail::NotifyEnd(class USkeletalMeshComponent* MeshComp, class UAnimSequenceBase* Animation)
-{
-}
-
 void UAnimNotifyState_Trail::NotifyEnd(class USkeletalMeshComponent * MeshComp, class UAnimSequenceBase * Animation, const FAnimNotifyEventReference& EventReference)
 {
-	PRAGMA_DISABLE_DEPRECATION_WARNINGS
-	NotifyEnd(MeshComp, Animation);
-	PRAGMA_ENABLE_DEPRECATION_WARNINGS
-
 	if (MeshComp->GetWorld()->GetNetMode() == NM_DedicatedServer)
 	{
 		return;
@@ -374,4 +352,3 @@ bool UAnimNotifyState_Trail::ValidateInput(class USkeletalMeshComponent * MeshCo
 }
 
 #undef LOCTEXT_NAMESPACE
-
diff --git a/Engine/Source/Runtime/Engine/Private/Animation/AnimNotify_PauseClothingSimulation.cpp b/Engine/Source/Runtime/Engine/Private/Animation/AnimNotify_PauseClothingSimulation.cpp
--- a/Engine/Source/Runtime/Engine/Private/Animation/AnimNotify_PauseClothingSimulation.cpp
+++ b/Engine/Source/Runtime/Engine/Private/Animation/AnimNotify_PauseClothingSimulation.cpp
@@ -14,16 +14,8 @@ UAnimNotify_PauseClothingSimulation::UAnimNotify_PauseClothingSimulation()
 #endif // WITH_EDITORONLY_DATA
 }
 
-void UAnimNotify_PauseClothingSimulation::Notify(class USkeletalMeshComponent* MeshComp, class UAnimSequenceBase* Animation)
-{
-
-}
-
 void UAnimNotify_PauseClothingSimulation::Notify(class USkeletalMeshComponent* MeshComp, class UAnimSequenceBase* Animation, const FAnimNotifyEventReference& EventReference)
 {
-	PRAGMA_DISABLE_DEPRECATION_WARNINGS
-    Notify(MeshComp, Animation);
-    PRAGMA_ENABLE_DEPRECATION_WARNINGS
 	MeshComp->SuspendClothingSimulation();
 }
 
diff --git a/Engine/Source/Runtime/Engine/Private/Animation/AnimNotify_PlayParticleEffect.cpp b/Engine/Source/Runtime/Engine/Private/Animation/AnimNotify_PlayParticleEffect.cpp
--- a/Engine/Source/Runtime/Engine/Private/Animation/AnimNotify_PlayParticleEffect.cpp
+++ b/Engine/Source/Runtime/Engine/Private/Animation/AnimNotify_PlayParticleEffect.cpp
@@ -73,15 +73,8 @@ void UAnimNotify_PlayParticleEffect::ValidateAssociatedAssets()
 }
 #endif
 
-void UAnimNotify_PlayParticleEffect::Notify(class USkeletalMeshComponent* MeshComp, class UAnimSequenceBase* Animation)
-{
-}
-
 void UAnimNotify_PlayParticleEffect::Notify(class USkeletalMeshComponent* MeshComp, class UAnimSequenceBase* Animation, const FAnimNotifyEventReference& EventReference)
 {
-	PRAGMA_DISABLE_DEPRECATION_WARNINGS
-    Notify(MeshComp, Animation);
-	PRAGMA_ENABLE_DEPRECATION_WARNINGS
 	// Don't call super to avoid unnecessary call in to blueprints
 	SpawnParticleSystem(MeshComp, Animation);
 }
@@ -131,4 +124,3 @@ UParticleSystemComponent* UAnimNotify_PlayParticleEffect::SpawnParticleSystem(cl
 
 	return ReturnComp;
 }
-
diff --git a/Engine/Source/Runtime/Engine/Private/Animation/AnimNotify_PlaySound.cpp b/Engine/Source/Runtime/Engine/Private/Animation/AnimNotify_PlaySound.cpp
--- a/Engine/Source/Runtime/Engine/Private/Animation/AnimNotify_PlaySound.cpp
+++ b/Engine/Source/Runtime/Engine/Private/Animation/AnimNotify_PlaySound.cpp
@@ -31,15 +31,8 @@ UAnimNotify_PlaySound::UAnimNotify_PlaySound()
 #endif // WITH_EDITORONLY_DATA
 }
 
-void UAnimNotify_PlaySound::Notify(USkeletalMeshComponent* MeshComp, UAnimSequenceBase* Animation)
-{
-}
-
 void UAnimNotify_PlaySound::Notify(class USkeletalMeshComponent* MeshComp, class UAnimSequenceBase* Animation, const FAnimNotifyEventReference& EventReference)
 {
-	PRAGMA_DISABLE_DEPRECATION_WARNINGS
-    Notify(MeshComp, Animation);
-	PRAGMA_ENABLE_DEPRECATION_WARNINGS
 	// Don't call super to avoid call back in to blueprints
 	if (Sound && MeshComp)
 	{
diff --git a/Engine/Source/Runtime/Engine/Private/Animation/AnimNotify_ResetClothingSimulation.cpp b/Engine/Source/Runtime/Engine/Private/Animation/AnimNotify_ResetClothingSimulation.cpp
--- a/Engine/Source/Runtime/Engine/Private/Animation/AnimNotify_ResetClothingSimulation.cpp
+++ b/Engine/Source/Runtime/Engine/Private/Animation/AnimNotify_ResetClothingSimulation.cpp
@@ -14,15 +14,8 @@ UAnimNotify_ResetClothingSimulation::UAnimNotify_ResetClothingSimulation()
 #endif // WITH_EDITORONLY_DATA
 }
 
-void UAnimNotify_ResetClothingSimulation::Notify(class USkeletalMeshComponent* MeshComp, class UAnimSequenceBase* Animation)
-{
-}
-
 void UAnimNotify_ResetClothingSimulation::Notify(class USkeletalMeshComponent* MeshComp, class UAnimSequenceBase* Animation, const FAnimNotifyEventReference& EventReference)
 {
-	PRAGMA_DISABLE_DEPRECATION_WARNINGS
-    Notify(MeshComp, Animation);
-	PRAGMA_ENABLE_DEPRECATION_WARNINGS
 	MeshComp->ForceClothNextUpdateTeleportAndReset();
 }
 
diff --git a/Engine/Source/Runtime/Engine/Private/Animation/AnimNotify_ResetDynamics.cpp b/Engine/Source/Runtime/Engine/Private/Animation/AnimNotify_ResetDynamics.cpp
--- a/Engine/Source/Runtime/Engine/Private/Animation/AnimNotify_ResetDynamics.cpp
+++ b/Engine/Source/Runtime/Engine/Private/Animation/AnimNotify_ResetDynamics.cpp
@@ -14,15 +14,8 @@ UAnimNotify_ResetDynamics::UAnimNotify_ResetDynamics()
 #endif // WITH_EDITORONLY_DATA
 }
 
-void UAnimNotify_ResetDynamics::Notify(class USkeletalMeshComponent* MeshComp, class UAnimSequenceBase* Animation)
-{
-}
-
 void UAnimNotify_ResetDynamics::Notify(class USkeletalMeshComponent* MeshComp, class UAnimSequenceBase* Animation, const FAnimNotifyEventReference& EventReference)
 {
-	PRAGMA_DISABLE_DEPRECATION_WARNINGS
-    Notify(MeshComp, Animation);
-	PRAGMA_ENABLE_DEPRECATION_WARNINGS
 	MeshComp->ResetAnimInstanceDynamics();
 }
 
diff --git a/Engine/Source/Runtime/Engine/Private/Animation/AnimNotify_ResumeClothingSimulation.cpp b/Engine/Source/Runtime/Engine/Private/Animation/AnimNotify_ResumeClothingSimulation.cpp
--- a/Engine/Source/Runtime/Engine/Private/Animation/AnimNotify_ResumeClothingSimulation.cpp
+++ b/Engine/Source/Runtime/Engine/Private/Animation/AnimNotify_ResumeClothingSimulation.cpp
@@ -14,15 +14,8 @@ UAnimNotify_ResumeClothingSimulation::UAnimNotify_ResumeClothingSimulation()
 #endif // WITH_EDITORONLY_DATA
 }
 
-void UAnimNotify_ResumeClothingSimulation::Notify(class USkeletalMeshComponent* MeshComp, class UAnimSequenceBase* Animation)
-{
-}
-
 void UAnimNotify_ResumeClothingSimulation::Notify(class USkeletalMeshComponent* MeshComp, class UAnimSequenceBase* Animation, const FAnimNotifyEventReference& EventReference)
 {
-	PRAGMA_DISABLE_DEPRECATION_WARNINGS
-    Notify(MeshComp, Animation);
-	PRAGMA_ENABLE_DEPRECATION_WARNINGS
 	MeshComp->ResumeClothingSimulation();
 }
 
