- Fix a crash if an in-use data interace is removed from code

#rnx

#Synced-CL 48463499

diff --git a/Engine/Plugins/FX/Niagara/Source/Niagara/Private/NiagaraScript.cpp b/Engine/Plugins/FX/Niagara/Source/Niagara/Private/NiagaraScript.cpp
--- a/Engine/Plugins/FX/Niagara/Source/Niagara/Private/NiagaraScript.cpp
+++ b/Engine/Plugins/FX/Niagara/Source/Niagara/Private/NiagaraScript.cpp
@@ -2725,18 +2725,25 @@ void UNiagaraScript::SetVMCompilationResults(const FNiagaraVMExecutableDataId& I
 
 		if (CachedDefaultDataInterfaces[Idx].DataInterface == nullptr)
 		{
-			// Use the CDO since we didn't have a default..
-			UObject* Obj = Info.Type.GetClass()->GetDefaultObject(true);
-			CachedDefaultDataInterfaces[Idx].DataInterface = Cast<UNiagaraDataInterface>(CopyDataInterface(CastChecked<UNiagaraDataInterface>(Obj), this));
+			if ( Info.Type.IsValid() )
+			{
+				// Use the CDO since we didn't have a default..
+				UObject* Obj = Info.Type.GetClass()->GetDefaultObject(true);
+				CachedDefaultDataInterfaces[Idx].DataInterface = Cast<UNiagaraDataInterface>(CopyDataInterface(CastChecked<UNiagaraDataInterface>(Obj), this));
 
-			if (Info.bIsPlaceholder == false)
+				if (Info.bIsPlaceholder == false)
+				{
+					UE_LOG(LogNiagara, Display, TEXT("We somehow ended up with a data interface that we couldn't match post compile. This shouldn't happen. Creating a dummy to prevent crashes. DataInterfaceInfoName:%s Object:%s"), *Info.Name.ToString(), *GetPathNameSafe(this));
+					UE_LOG(LogNiagara, Log, TEXT("Object to Name map contents:"));
+					DumpNameMap(ObjectNameMap);
+				}
+				check(CachedDefaultDataInterfaces[Idx].DataInterface != nullptr);
+			}
+			else
 			{
-				UE_LOG(LogNiagara, Display, TEXT("We somehow ended up with a data interface that we couldn't match post compile. This shouldn't happen. Creating a dummy to prevent crashes. DataInterfaceInfoName:%s Object:%s"), *Info.Name.ToString(), *GetPathNameSafe(this));
-				UE_LOG(LogNiagara, Log, TEXT("Object to Name map contents:"));
-				DumpNameMap(ObjectNameMap);
+				UE_LOG(LogNiagara, Log, TEXT("Data Interface has no class, was it deleted? DataInterfaceInfoName:%s Object:%s"), *Info.Name.ToString(), *GetPathNameSafe(this));
 			}
 		}
-		check(CachedDefaultDataInterfaces[Idx].DataInterface != nullptr);
 	}
 
 	CachedDefaultUObjects = CachedScriptVM.UObjectInfos;
