- Check for editor not being nullptr
- The merge manager can invoke the view model with -game resulting in a crash

#rnx
#Synced-CL 48612007

diff --git a/Engine/Plugins/FX/Niagara/Source/NiagaraEditor/Private/ViewModels/Stack/NiagaraStackRendererItem.cpp b/Engine/Plugins/FX/Niagara/Source/NiagaraEditor/Private/ViewModels/Stack/NiagaraStackRendererItem.cpp
--- a/Engine/Plugins/FX/Niagara/Source/NiagaraEditor/Private/ViewModels/Stack/NiagaraStackRendererItem.cpp
+++ b/Engine/Plugins/FX/Niagara/Source/NiagaraEditor/Private/ViewModels/Stack/NiagaraStackRendererItem.cpp
@@ -659,7 +659,8 @@ void UNiagaraStackRendererItem::RefreshIssues(TArray<FStackIssue>& NewIssues)
 		RendererProperties->GetRendererFeedback(GetEmitterViewModel()->GetEmitter(), Errors, Warnings, Infos);
 
 		// Query feedback from the material validation system
-		if (UMaterialEditorValidationSubsystem* MaterialValidationSubsystem = GEditor->GetEditorSubsystem<UMaterialEditorValidationSubsystem>())
+		UMaterialEditorValidationSubsystem* MaterialValidationSubsystem = GEditor ? GEditor->GetEditorSubsystem<UMaterialEditorValidationSubsystem>() : nullptr;
+		if (MaterialValidationSubsystem)
 		{
 			TArray<UMaterialInterface*> RendererMaterials;
 			RendererProperties->GetUsedMaterials(nullptr, RendererMaterials);
