- Recache the world transform on Reset regardless of mode
- Some data interaces grab this transform when creating per instance data, which meant the reused pooled components would have an invalid or old transform

#Synced-CL 48563820

diff --git a/Engine/Plugins/FX/Niagara/Source/Niagara/Private/NiagaraSystemInstance.cpp b/Engine/Plugins/FX/Niagara/Source/Niagara/Private/NiagaraSystemInstance.cpp
--- a/Engine/Plugins/FX/Niagara/Source/Niagara/Private/NiagaraSystemInstance.cpp
+++ b/Engine/Plugins/FX/Niagara/Source/Niagara/Private/NiagaraSystemInstance.cpp
@@ -113,7 +113,7 @@ FNiagaraSystemInstance::FNiagaraSystemInstance(UWorld& InWorld, UNiagaraSystem&
 	  , World(&InWorld)
 	  , System(&InSystem)
 	  , OverrideParameters(InOverrideParameters)
-	  , AttachComponent(InAttachComponent)
+	  , WeakAttachComponent(InAttachComponent)
 	  , Age(0.0f)
 	  , LastRenderTime(0.0f)
 	  , TickCount(0)
@@ -190,7 +190,7 @@ void FNiagaraSystemInstance::SetEmitterEnable(FName EmitterName, bool bNewEnable
 
 	if (FNiagaraUtilities::LogVerboseWarnings())
 	{
-		UE_LOG(LogNiagara, Warning, TEXT("SetEmitterEnable: Failed to find Emitter(%s) System(%s) Component(%s)"), *EmitterName.ToString(), *GetNameSafe(System), *GetFullNameSafe(AttachComponent.Get()));
+		UE_LOG(LogNiagara, Warning, TEXT("SetEmitterEnable: Failed to find Emitter(%s) System(%s) Component(%s)"), *EmitterName.ToString(), *GetNameSafe(System), *GetFullNameSafe(GetAttachComponent()));
 	}
 }
 
@@ -519,7 +519,7 @@ void FNiagaraSystemInstance::SetGpuComputeDebug(bool bEnableDebug)
 	if (bEnableDebug)
 	{
 		FString SystemName = System->GetName();
-		if (USceneComponent* Owner = AttachComponent.Get())
+		if (USceneComponent* Owner = GetAttachComponent())
 		{
 			SystemName.Append(TEXT("/"));
 			if (AActor* Actor = Owner->GetTypedOuter<AActor>())
@@ -567,7 +567,7 @@ void FNiagaraSystemInstance::SetWarmupSettings(int32 InWarmupTickCount, float In
 
 UActorComponent* FNiagaraSystemInstance::GetPrereqComponent() const
 {
-	UActorComponent* PrereqComponent = AttachComponent.Get();
+	UActorComponent* PrereqComponent = GetAttachComponent();
 
 	// This is to maintain legacy behavior (and perf benefit) of ticking in PrePhysics with unattached UNiagaraComponents that have no DI prereqs
 	// NOTE: This means that the system likely ticks with frame-behind transform if the component is moved, but likely doesn't manifest as an issue with local-space emitters
@@ -734,6 +734,7 @@ void FNiagaraSystemInstance::OnPooledReuse(UWorld& NewWorld)
 {
 	World = &NewWorld;
 
+	WorldTransform = FTransform();
 	FixedBounds_GT.Init();
 	FixedBounds_CNC.Init();
 	PreviousLocation.Reset();
@@ -819,11 +820,16 @@ void FNiagaraSystemInstance::Reset(EResetMode Mode)
 		}
 	}
 
-	// Set tile for LWC offset
+	// Update LWC Tile & WorldTransform as it may be used when data interaces create instance data
 	LWCTile = FVector3f::ZeroVector;
-	if (AttachComponent.IsValid() && GetSystem()->SupportsLargeWorldCoordinates())
+	if (USceneComponent* AttachComponent = GetAttachComponent())
 	{
-		LWCTile = FLargeWorldRenderScalar::GetTileFor(AttachComponent->GetComponentLocation());
+		WorldTransform = AttachComponent->GetComponentToWorld();
+		if (GetSystem()->SupportsLargeWorldCoordinates())
+		{
+			LWCTile = FLargeWorldRenderScalar::GetTileFor(AttachComponent->GetComponentLocation());
+			WorldTransform.AddToTranslation(FVector(LWCTile) * -FLargeWorldRenderScalar::GetTileSize());
+		}
 	}
 
 	// Depending on the rest mode we may need to bind or can possibly skip it
@@ -957,7 +963,7 @@ void FNiagaraSystemInstance::ResetInternal(bool bResetSimulations)
 	if (!System->IsValid())
 	{
 		SetRequestedExecutionState(ENiagaraExecutionState::Disabled);
-		UE_LOG(LogNiagara, Warning, TEXT("Failed to activate Niagara System due to invalid asset! System(%s) Component(%s)"), *System->GetName(), *GetFullNameSafe(AttachComponent.Get()));
+		UE_LOG(LogNiagara, Warning, TEXT("Failed to activate Niagara System due to invalid asset! System(%s) Component(%s)"), *System->GetName(), *GetFullNameSafe(GetAttachComponent()));
 		return;
 	}
 
@@ -1075,7 +1081,7 @@ void FNiagaraSystemInstance::ReInitInternal()
 		SetRequestedExecutionState(ENiagaraExecutionState::Disabled);
 		if ( Emitters.Num() != 0 )
 		{
-			UE_LOG(LogNiagara, Warning, TEXT("Failed to activate Niagara System due to invalid asset! System(%s) Component(%s)"), *System->GetName(), *GetFullNameSafe(AttachComponent.Get()));
+			UE_LOG(LogNiagara, Warning, TEXT("Failed to activate Niagara System due to invalid asset! System(%s) Component(%s)"), *System->GetName(), *GetFullNameSafe(GetAttachComponent()));
 		}
 		return;
 	}
@@ -1208,7 +1214,7 @@ void FNiagaraSystemInstance::BindParameters()
 	}
 	else
 	{
-		UE_LOG(LogNiagara, Warning, TEXT("OverrideParameters is null.  Component(%s) System(%s)"), *GetFullNameSafe(AttachComponent.Get()), *GetFullNameSafe(GetSystem()));
+		UE_LOG(LogNiagara, Warning, TEXT("OverrideParameters is null.  Component(%s) System(%s)"), *GetFullNameSafe(GetAttachComponent()), *GetFullNameSafe(GetSystem()));
 	}
 
 	for (const FNiagaraEmitterInstanceRef& Emitter : Emitters)
@@ -2017,7 +2023,7 @@ ETickingGroup FNiagaraSystemInstance::CalculateTickGroup() const
 			break;
 
 		case ENiagaraTickBehavior::UseComponentTickGroup:
-			if (USceneComponent* Component = AttachComponent.Get())
+			if (USceneComponent* Component = GetAttachComponent())
 			{
 				NewMinTickGroup = FMath::Clamp((ETickingGroup)Component->PrimaryComponentTick.TickGroup, NiagaraFirstTickGroup, NiagaraLastTickGroup);
 			}
@@ -2058,7 +2064,7 @@ void FNiagaraSystemInstance::SetTickBehavior(ENiagaraTickBehavior NewTickBehavio
 void FNiagaraSystemInstance::TickInstanceParameters_GameThread(float DeltaSeconds)
 {
 	// If we're associated with a scene component, update our cached transform (otherwise, assume it was previously set externally)
-	if (AttachComponent.IsValid())
+	if (USceneComponent* AttachComponent = GetAttachComponent())
 	{
 		WorldTransform = AttachComponent->GetComponentToWorld();
 		WorldTransform.AddToTranslation(FVector(LWCTile) * -FLargeWorldRenderScalar::GetTileSize());
@@ -2397,7 +2403,7 @@ void FNiagaraSystemInstance::SimCacheTick_GameThread(UNiagaraSimCache* SimCache,
 	}
 
 	TOptional<FTransform> OwnerTransform;
-	if (USceneComponent* TransformComponent = AttachComponent.Get())
+	if (USceneComponent* AttachComponent = GetAttachComponent())
 	{
 		OwnerTransform = AttachComponent->GetComponentToWorld();
 		OwnerTransform->AddToTranslation(FVector(LWCTile) * -FLargeWorldRenderScalar::GetTileSize());
@@ -2520,7 +2526,7 @@ void FNiagaraSystemInstance::WaitForConcurrentTickDoNotFinalize(bool bEnsureComp
 		if ((GWaitForAsyncStallWarnThresholdMS > 0.0f) && (StallTimeMS > GWaitForAsyncStallWarnThresholdMS))
 		{
 			//-TODO: This should be put back to a warning once EngineTests no longer cause it show up.  The reason it's triggered is because we pause in latent actions right after a TG running Niagara sims.
-			UE_LOG(LogNiagara, Log, TEXT("Niagara Effect stalled GT for %g ms. Component(%s) System(%s)"), StallTimeMS, *GetFullNameSafe(AttachComponent.Get()), *GetFullNameSafe(GetSystem()));
+			UE_LOG(LogNiagara, Log, TEXT("Niagara Effect stalled GT for %g ms. Component(%s) System(%s)"), StallTimeMS, *GetFullNameSafe(GetAttachComponent()), *GetFullNameSafe(GetSystem()));
 		}
 	}
 
@@ -2833,7 +2839,7 @@ void FNiagaraSystemInstance::FinalizeTick_GameThread(bool bEnqueueGPUTickIfNeede
 		FNiagaraSystemSimulation* Sim = SystemSimulation.Get();
 		checkf(Sim != nullptr,
 			TEXT("SystemSimulation is nullptr during Finalize and the System(%s) AttachComponent(%s) ActualExecutionState(%d) SystemInstanceIndex(%d) SystemInstanceState(%s)"),
-			*GetFullNameSafe(System), *GetFullNameSafe(AttachComponent.Get()), ActualExecutionState, SystemInstanceIndex,
+			*GetFullNameSafe(System), *GetFullNameSafe(GetAttachComponent()), ActualExecutionState, SystemInstanceIndex,
 			*StaticEnum<ENiagaraSystemInstanceState>()->GetValueAsString(SystemInstanceState)
 		);
 		
@@ -3000,7 +3006,7 @@ void FNiagaraSystemInstance::SetEmitterFixedBounds(FName EmitterName, const FBox
 	}
 
 	// Failed to find emitter
-	UE_LOG(LogNiagara, Warning, TEXT("SetEmitterFixedBounds: Failed to find Emitter(%s) System(%s) Component(%s)"), *EmitterName.ToString(), *GetNameSafe(System), *GetFullNameSafe(AttachComponent.Get()));
+	UE_LOG(LogNiagara, Warning, TEXT("SetEmitterFixedBounds: Failed to find Emitter(%s) System(%s) Component(%s)"), *EmitterName.ToString(), *GetNameSafe(System), *GetFullNameSafe(GetAttachComponent()));
 }
 
 FBox FNiagaraSystemInstance::GetEmitterFixedBounds(FName EmitterName) const
@@ -3014,7 +3020,7 @@ FBox FNiagaraSystemInstance::GetEmitterFixedBounds(FName EmitterName) const
 	}
 
 	// Failed to find emitter
-	UE_LOG(LogNiagara, Warning, TEXT("GetEmitterFixedBounds: Failed to find Emitter(%s) System(%s) Component(%s)"), *EmitterName.ToString(), *GetNameSafe(System), *GetFullNameSafe(AttachComponent.Get()));
+	UE_LOG(LogNiagara, Warning, TEXT("GetEmitterFixedBounds: Failed to find Emitter(%s) System(%s) Component(%s)"), *EmitterName.ToString(), *GetNameSafe(System), *GetFullNameSafe(GetAttachComponent()));
 	return FBox(ForceInit);
 }
 
@@ -3104,15 +3110,16 @@ const FString& FNiagaraSystemInstance::GetCrashReporterTag()const
 {
 	if(CrashReporterTag.IsEmpty())
 	{
-		UNiagaraSystem* Sys = GetSystem();
-		UNiagaraComponent* Component = Cast<UNiagaraComponent>(AttachComponent.Get());
-		USceneComponent* AttachParent = Component ? Component->GetAttachParent() : AttachComponent.Get();
+		UNiagaraComponent* AttachComponent = Cast<UNiagaraComponent>(GetAttachComponent());
+		USceneComponent* AttachParent = AttachComponent ? AttachComponent->GetAttachParent() : nullptr;
 
-		const FString& CompName = GetFullNameSafe(Component);
-		const FString& SystemName = GetFullNameSafe(Sys);
-		const FString& AttachName = GetFullNameSafe(AttachParent);
-
-		CrashReporterTag = FString::Printf(TEXT("SystemInstance | System: %s | bSolo: %s | Component: %s | AttachedTo: %s |"), *SystemName, IsSolo() ? TEXT("true") : TEXT("false"), *CompName, *AttachName);
+		CrashReporterTag = FString::Printf(
+			TEXT("SystemInstance | System: %s | bSolo: %s | Component: %s | AttachedTo: %s |"),
+			*GetFullNameSafe(GetSystem()),
+			IsSolo() ? TEXT("true") : TEXT("false"),
+			*GetFullNameSafe(AttachComponent),
+			*GetFullNameSafe(AttachParent)
+		);
 	}
 	return CrashReporterTag;
 }
diff --git a/Engine/Plugins/FX/Niagara/Source/Niagara/Public/NiagaraSystemInstance.h b/Engine/Plugins/FX/Niagara/Source/Niagara/Public/NiagaraSystemInstance.h
--- a/Engine/Plugins/FX/Niagara/Source/Niagara/Public/NiagaraSystemInstance.h
+++ b/Engine/Plugins/FX/Niagara/Source/Niagara/Public/NiagaraSystemInstance.h
@@ -270,7 +270,7 @@ public:
 
 	inline UWorld* GetWorld() const { return World; }
 	inline UNiagaraSystem* GetSystem() const { return System; }
-	inline USceneComponent* GetAttachComponent() { return AttachComponent.Get(); }
+	inline USceneComponent* GetAttachComponent() const { return WeakAttachComponent.Get(); }
 	inline FNiagaraUserRedirectionParameterStore* GetOverrideParameters() { return OverrideParameters; }
 	inline const FNiagaraUserRedirectionParameterStore* GetOverrideParameters() const { return OverrideParameters; }
 	[[nodiscard]] TArrayView<FNiagaraEmitterInstanceRef> GetEmitters() { return Emitters; }
@@ -517,7 +517,7 @@ private:
 	UWorld* World;
 	UNiagaraSystem* System = nullptr;
 	FNiagaraUserRedirectionParameterStore* OverrideParameters;
-	TWeakObjectPtr<USceneComponent> AttachComponent;
+	TWeakObjectPtr<USceneComponent> WeakAttachComponent;
 	TUniquePtr<INiagaraSceneComponentUtils> SceneComponentUtils;
 
 	FTransform WorldTransform;
