- Add check that StackNodeGroups array size is valid in UNiagaraStackRendererItem::AddMissingVariable

#rnx

#Synced-CL 48563820

diff --git a/Engine/Plugins/FX/Niagara/Source/NiagaraEditor/Private/ViewModels/Stack/NiagaraStackRendererItem.cpp b/Engine/Plugins/FX/Niagara/Source/NiagaraEditor/Private/ViewModels/Stack/NiagaraStackRendererItem.cpp
--- a/Engine/Plugins/FX/Niagara/Source/NiagaraEditor/Private/ViewModels/Stack/NiagaraStackRendererItem.cpp
+++ b/Engine/Plugins/FX/Niagara/Source/NiagaraEditor/Private/ViewModels/Stack/NiagaraStackRendererItem.cpp
@@ -237,6 +237,8 @@ bool UNiagaraStackRendererItem::AddMissingVariable(const FVersionedNiagaraEmitte
 	AssignmentGroup.StartNodes.Add(NewAssignmentNode);
 	AssignmentGroup.EndNode = NewAssignmentNode;
 
+	check(StackNodeGroups.Num() >= 2);
+
 	FNiagaraStackGraphUtilities::FStackNodeGroup& OutputGroup = StackNodeGroups[StackNodeGroups.Num() - 1];
 	FNiagaraStackGraphUtilities::FStackNodeGroup& OutputGroupPrevious = StackNodeGroups[StackNodeGroups.Num() - 2];
 	FNiagaraStackGraphUtilities::ConnectStackNodeGroup(AssignmentGroup, OutputGroupPrevious, OutputGroup);
