- Fix CSI issue

#rnx

#Synced-CL 48563820

diff --git a/Engine/Plugins/FX/Niagara/Source/Niagara/Private/DataInterface/NiagaraDataInterfacePhysicsAsset.cpp b/Engine/Plugins/FX/Niagara/Source/Niagara/Private/DataInterface/NiagaraDataInterfacePhysicsAsset.cpp
--- a/Engine/Plugins/FX/Niagara/Source/Niagara/Private/DataInterface/NiagaraDataInterfacePhysicsAsset.cpp
+++ b/Engine/Plugins/FX/Niagara/Source/Niagara/Private/DataInterface/NiagaraDataInterfacePhysicsAsset.cpp
@@ -581,6 +581,8 @@ UNiagaraDataInterfacePhysicsAsset::UNiagaraDataInterfacePhysicsAsset(FObjectInit
 
 void UNiagaraDataInterfacePhysicsAsset::ExtractSourceComponent(FNiagaraSystemInstance* SystemInstance, FTransform& BoneTransform)
 {
+	check(SystemInstance);
+
 	// Helper to scour an actor (or its parents) for a valid skeletal mesh component
 	auto FindActorSkelMeshComponent = [](AActor* Actor, bool bRecurseParents = false) -> USkeletalMeshComponent*
 	{
@@ -623,8 +625,7 @@ void UNiagaraDataInterfacePhysicsAsset::ExtractSourceComponent(FNiagaraSystemIns
 	// Track down the source component
 	TWeakObjectPtr<USkeletalMeshComponent> SourceComponent;
 	
-	
-	if (MeshUserParameter.Parameter.IsValid() && SystemInstance != nullptr)
+	if (MeshUserParameter.Parameter.IsValid())
 	{
 		// Initialize the binding and retrieve the object. If a valid object is bound, we'll try and retrieve the SkelMesh component from it.
 		// If it's not valid yet, we'll reset and do this again when/if a valid object is set on the binding
