NDCWrite DI Wizard initializes it's access context correctly.

#Synced-CL 48349740

diff --git a/Engine/Plugins/FX/Niagara/Source/Niagara/Internal/DataInterface/NiagaraDataInterfaceDataChannelWrite.h b/Engine/Plugins/FX/Niagara/Source/Niagara/Internal/DataInterface/NiagaraDataInterfaceDataChannelWrite.h
--- a/Engine/Plugins/FX/Niagara/Source/Niagara/Internal/DataInterface/NiagaraDataInterfaceDataChannelWrite.h
+++ b/Engine/Plugins/FX/Niagara/Source/Niagara/Internal/DataInterface/NiagaraDataInterfaceDataChannelWrite.h
@@ -215,6 +215,8 @@ public:
 	virtual uint32 GetGpuCountBufferEstimate() const { return 1; }
 	//UNiagaraDataInterface Interface
 
+	NIAGARA_API virtual void InitAccessContext();
+
 	NIAGARA_API void Num(FVectorVMExternalFunctionContext& Context);
 	NIAGARA_API void Allocate(FVectorVMExternalFunctionContext& Context);
 	NIAGARA_API void Write(FVectorVMExternalFunctionContext& Context, int32 FuncIdx);
diff --git a/Engine/Plugins/FX/Niagara/Source/Niagara/Private/DataInterface/NiagaraDataInterfaceDataChannelWrite.cpp b/Engine/Plugins/FX/Niagara/Source/Niagara/Private/DataInterface/NiagaraDataInterfaceDataChannelWrite.cpp
--- a/Engine/Plugins/FX/Niagara/Source/Niagara/Private/DataInterface/NiagaraDataInterfaceDataChannelWrite.cpp
+++ b/Engine/Plugins/FX/Niagara/Source/Niagara/Private/DataInterface/NiagaraDataInterfaceDataChannelWrite.cpp
@@ -1235,6 +1235,28 @@ void UNiagaraDataInterfaceDataChannelWrite::GetVMExternalFunction(const FVMExter
 	}
 }
 
+void UNiagaraDataInterfaceDataChannelWrite::InitAccessContext()
+{
+	UNiagaraDataChannel* NDC = Channel ? Channel->Get() : nullptr;
+	if (NDC)
+	{
+		const UScriptStruct* CtxStruct = NDC->GetAccessContextType().Get();
+		//Special case legacy type. We avoid creating an access context for this type here and just use the old search params path.
+		if (CtxStruct == FNDCAccessContextLegacy::StaticStruct())
+		{
+			AccessContext.Reset();
+		}
+		else if (AccessContext.IsValid() == false || CtxStruct != AccessContext.GetScriptStruct())
+		{
+			AccessContext.Init(NDC->GetAccessContextType());
+		}
+	}
+	else
+	{
+		AccessContext.Reset();
+	}
+}
+
 void UNiagaraDataInterfaceDataChannelWrite::Num(FVectorVMExternalFunctionContext& Context)
 {
 	VectorVM::FUserPtrHandler<FNDIDataChannelWriteInstanceData> InstData(Context);
diff --git a/Engine/Plugins/FX/Niagara/Source/NiagaraEditor/Private/Widgets/DataChannel/NiagaraDataChannelWizard.cpp b/Engine/Plugins/FX/Niagara/Source/NiagaraEditor/Private/Widgets/DataChannel/NiagaraDataChannelWizard.cpp
--- a/Engine/Plugins/FX/Niagara/Source/NiagaraEditor/Private/Widgets/DataChannel/NiagaraDataChannelWizard.cpp
+++ b/Engine/Plugins/FX/Niagara/Source/NiagaraEditor/Private/Widgets/DataChannel/NiagaraDataChannelWizard.cpp
@@ -642,6 +642,7 @@ TSharedRef<FModuleWizardModel> DataChannel::CreateReadNDCModuleWizardModel()
 						if (UNiagaraDataInterfaceDataChannelRead* DataInterface = Cast<UNiagaraDataInterfaceDataChannelRead>(FunctionInput->Data))
 						{
 							DataInterface->Channel = Channel;
+							DataInterface->InitAccessContext();
 							DataInterface->bReadCurrentFrame = AssetPage->Data->bReadCurrentFrame;
 							DataInterface->bUpdateSourceDataEveryTick = AssetPage->Data->bUpdateSourceDataEveryTick;
 						}
@@ -827,6 +828,7 @@ TSharedRef<FModuleWizardModel> DataChannel::CreateWriteNDCModuleWizardModel()
 						if (UNiagaraDataInterfaceDataChannelWrite* DataInterface = Cast<UNiagaraDataInterfaceDataChannelWrite>(FunctionInput->Data))
 						{
 							DataInterface->Channel = Channel;
+							DataInterface->InitAccessContext();
 							DataInterface->bPublishToGame = AssetPage->Data->bPublishToGame;
 							DataInterface->bPublishToCPU = AssetPage->Data->bPublishToCPU;
 							DataInterface->bPublishToGPU = AssetPage->Data->bPublishToGPU;
