Trim the visible API for data managers down to a minimum.
Plus a few bughawk fixes.

#Synced-CL 48349740

diff --git a/Engine/Plugins/FX/Niagara/Source/Niagara/Private/NiagaraDataManager.h b/Engine/Plugins/FX/Niagara/Source/Niagara/Private/NiagaraDataManager.h
new file mode 100644
--- /dev/null
+++ b/Engine/Plugins/FX/Niagara/Source/Niagara/Private/NiagaraDataManager.h
@@ -0,0 +1,60 @@
+// Copyright Epic Games, Inc. All Rights Reserved.
+
+#pragma once
+
+#include "CoreMinimal.h"
+#include "NiagaraTypes.h"
+#include "NiagaraGpuComputeDataManager.h"
+#include "NiagaraSharedDataUtilities.h"
+#include "NiagaraDataManager.generated.h"
+
+class FNiagaraWorldManager;
+class FNiagaraGpuComputeDispatchInterface;
+
+/** 
+A generic manager class holding some data requiring management or marshaling for Niagara.
+A counterpart for Gpu accessible data is FNiagaraGpuComputeDataManager.
+Some cases may utilize a UNiagaraDataManager and a FNiagaraGpuComputeDataManager together.
+*/
+UCLASS()
+class UNiagaraDataManager : public UObject
+{
+	GENERATED_BODY()
+
+public:
+
+	virtual void Init(FNiagaraWorldManager* WorldMan);		
+	virtual void BeginFrame();
+	virtual void EndFrame();
+	virtual void PreTick(ETickingGroup TickGroup, float DeltaSeconds){}
+
+	/** Track the connection between a NiagaraComponent and a FNiagaraSharedDataPtr. The data will remain valid at least until the valid component completes. */
+	void RegisterReferencingSystem(UNiagaraComponent* Component, FNiagaraSharedDataPtr SharedData);
+	
+	/** Callback when a Niagara component completes so we can remove any tracked data associated with it. */
+	UFUNCTION()
+	void OnReferencingSystemFinished(UNiagaraComponent* Component);
+
+protected:
+	
+	FNiagaraWorldManager* WorldManager = nullptr;
+	FNiagaraGpuComputeDispatchInterface* ComputeDispatchInterface = nullptr;
+
+	/** 
+	Map of active Niagara Components referencing FNiagaraSharedDataPtrs.
+	Allows us to maintain the lifetime of a FNiagaraSharedDataPtr until a referencing Niagara Component completes.
+	*/
+	TMap<TWeakObjectPtr<UNiagaraComponent>, TArray<FNiagaraSharedDataPtr>> NiagaraComponentToSharedData;
+};
+
+struct FNiagaraDataInterfaceSetShaderParametersContext;
+typedef FNiagaraDataInterfaceSetShaderParametersContext FNiagaraDataManagerSetShaderParametersContext;
+
+/** Base class for RT Proxies of data managers. */
+struct FNiagaraDataManager_RTProxy : public FNiagaraGpuComputeDataManager
+{
+	FNiagaraDataManager_RTProxy(FNiagaraGpuComputeDispatchInterface* InOwnerInterface)
+	: FNiagaraGpuComputeDataManager(InOwnerInterface)
+	{
+	}
+};
\ No newline at end of file
diff --git a/Engine/Plugins/FX/Niagara/Source/Niagara/Private/NiagaraDataManager_StaticMesh.cpp b/Engine/Plugins/FX/Niagara/Source/Niagara/Private/NiagaraDataManager_StaticMesh.cpp
--- a/Engine/Plugins/FX/Niagara/Source/Niagara/Private/NiagaraDataManager_StaticMesh.cpp
+++ b/Engine/Plugins/FX/Niagara/Source/Niagara/Private/NiagaraDataManager_StaticMesh.cpp
@@ -171,9 +171,14 @@ void FNiagaraDataManager_StaticMesh_RTProxy::PreTick(FNiagaraGpuComputeDataManag
 				}
 			}
 		}
-
+		
+		if(Buffer.Num() == 0)
+		{
+			Buffer.SetNum(1);
+		}
+			
 		MeshDataBuffer = CreateStructuredBuffer(Context.GetGraphBuilder(), TEXT("NiagaraSharedMeshDataBuffer"), MoveTemp(Buffer));
-		MeshDataBufferSRV = Context.GetGraphBuilder().CreateSRV(MeshDataBuffer);		
+		MeshDataBufferSRV = Context.GetGraphBuilder().CreateSRV(MeshDataBuffer);				
 	}
 }
 
@@ -218,4 +223,4 @@ void FNiagaraDataManager_StaticMesh_RTProxy::SetShaderParameters(const FNiagaraD
 			}
 		}
 	}
-}
\ No newline at end of file
+}
diff --git a/Engine/Plugins/FX/Niagara/Source/Niagara/Private/NiagaraDataManager_StaticMesh.h b/Engine/Plugins/FX/Niagara/Source/Niagara/Private/NiagaraDataManager_StaticMesh.h
new file mode 100644
--- /dev/null
+++ b/Engine/Plugins/FX/Niagara/Source/Niagara/Private/NiagaraDataManager_StaticMesh.h
@@ -0,0 +1,123 @@
+// Copyright Epic Games, Inc. All Rights Reserved.
+
+#pragma once
+
+#include "CoreMinimal.h"
+#include "NiagaraTypes.h"
+#include "MeshPassProcessor.h"
+
+#include "NiagaraDataManager.h"
+#include "NiagaraDataManager_StaticMesh.generated.h"
+
+
+class UPrimitiveComponent;
+struct FNiagaraDataInterfaceSetShaderParametersContext;
+
+/** Persistent ID allowing Niagara particles access to the static mesh info provided by the manager. */
+USTRUCT(BlueprintType)
+struct FNiagaraSharedDataID_StaticMesh : public FNiagaraID
+{
+	GENERATED_BODY()
+};
+
+//Shared data about a static mesh we are tracking for indirect access in Niagara.
+struct FNiagaraSharedData_StaticMesh : public FNiagaraSharedData
+{
+	TWeakObjectPtr<UStaticMeshComponent> SourceComponent;
+	
+	FPrimitiveComponentId PrimitiveId;
+
+	//The mesh reference generated for this primitive. 
+	FNiagaraSharedDataID_StaticMesh ID;
+};
+
+typedef TSharedPtr<FNiagaraSharedData_StaticMesh> FNiagaraSharedData_StaticMeshPtr;
+
+UCLASS()
+class UNiagaraDataManager_StaticMesh : public UNiagaraDataManager
+{
+	GENERATED_BODY()
+
+public:
+	
+	UNiagaraDataManager_StaticMesh(FObjectInitializer const& ObjectInitializer);
+
+	//UNiagaraDataManager Interface
+	virtual void EndFrame()override;
+	//END UNiagaraDataManager Interface
+	
+	/** 
+	Finds or creates shared data tracking info about the given mesh component for use later in Niagara.
+	Once no-one is holding this shared data pointer, the data will be cleaned up and invalid to access on the GPU.
+	Access is still safe but will fail.
+	You can have Niagara Components hold the reference until they complete using RegisterReferencingSystem().
+	*/
+	[[nodiscard]] FNiagaraSharedData_StaticMeshPtr GetOrCreateSharedMeshData(UStaticMeshComponent* ComponentToReference);	
+
+private:
+	
+	TMap<TWeakObjectPtr<UStaticMeshComponent>, FNiagaraSharedData_StaticMeshPtr> TrackedComponents;
+	
+	/** Currently free slots in the MeshDataBuffer. */
+	TArray<int32> FreeSlots;
+	int32 NumSlots = 0;
+	int32 IDAcquireTag = 0;
+};
+
+/** 
+Render thread proxy for UNiagaraDataManager_StaticMesh.
+Tracks link between PrimitiveIds for tracked static mesh components and associated FNiagaraSharedDataID_StaticMesh.
+The manager generates a buffer each frame allowing particles to use a FNiagaraSharedDataID_StaticMesh to access info about each static mesh component.
+*/
+struct FNiagaraDataManager_StaticMesh_RTProxy final : public FNiagaraDataManager_RTProxy
+{	
+	static FName GetManagerName()
+	{
+		static FName ManagerName("FNiagaraDataManager_StaticMesh_RTProxy");
+		return ManagerName;
+	}
+	
+	FNiagaraDataManager_StaticMesh_RTProxy(FNiagaraGpuComputeDispatchInterface* InOwnerInterface)
+	: FNiagaraDataManager_RTProxy(InOwnerInterface)
+	{
+	}
+	
+	//FNiagaraDataManager_RTProxy Interface
+	virtual void BeginFrame(FNiagaraGpuComputeDataManagerTickContext& Context) override;
+	virtual void EndFrame(FNiagaraGpuComputeDataManagerTickContext& Context) override;
+	virtual void PreTick(FNiagaraGpuComputeDataManagerTickContext& Context, ENiagaraGpuComputeTickStage::Type TickStage) override;	
+	//END FNiagaraDataManager_RTProxy Interface
+	
+	void Add(FPrimitiveComponentId PrimitiveId, FNiagaraSharedDataID_StaticMesh MeshID);
+	void Remove(FPrimitiveComponentId PrimitiveId);
+	void SetNumSlots(int32 InNumSlots){ NumSlots = InNumSlots; }
+
+	[[nodiscard]] FRDGBufferSRVRef GetGPUBufferSRV()const;
+	[[nodiscard]] int32 GetNumSlots()const{ return NumSlots; }
+	
+	[[nodiscard]] static const FShaderParametersMetadata* GetParametersMetadata();
+	static void SetShaderParameters(const FNiagaraDataManagerSetShaderParametersContext& Context);
+
+private:
+
+	TMap<FPrimitiveComponentId, FNiagaraSharedDataID_StaticMesh> PrimitiveToMeshID;
+	int32 NumSlots = 0;
+
+	FRDGBufferRef MeshDataBuffer = nullptr;
+	FRDGBufferSRVRef MeshDataBufferSRV = nullptr;
+};
+
+/** Counterpart to struct in NiagaraManagedReference_StaticMesh.ush */
+struct FNiagaraSharedDataGPU_StaticMesh
+{
+	int32 IDTag = INDEX_NONE;
+	int32 DistanceFieldIndex = INDEX_NONE;
+	int32 PrimitiveIndex = INDEX_NONE;
+	int32 Padding = INDEX_NONE;
+};
+
+BEGIN_SHADER_PARAMETER_STRUCT(FNiagaraDataManager_StaticMesh_Parameters, )
+	SHADER_PARAMETER_RDG_BUFFER_SRV(StructuredBuffer<FNiagaraSharedDataGPU_StaticMesh>, NiagaraSharedData_StaticMesh_Buffer)
+	SHADER_PARAMETER(int32, NiagaraSharedData_StaticMesh_NumSlots)
+END_SHADER_PARAMETER_STRUCT()
+
diff --git a/Engine/Plugins/FX/Niagara/Source/Niagara/Private/NiagaraSharedDataUtilities.cpp b/Engine/Plugins/FX/Niagara/Source/Niagara/Private/NiagaraSharedDataUtilities.cpp
new file mode 100644
--- /dev/null
+++ b/Engine/Plugins/FX/Niagara/Source/Niagara/Private/NiagaraSharedDataUtilities.cpp
@@ -0,0 +1,24 @@
+// Copyright Epic Games, Inc. All Rights Reserved.
+
+#include "NiagaraSharedDataUtilities.h"
+
+#include "NiagaraDataManager_StaticMesh.h"
+#include "NiagaraWorldManagerImpl.h"
+
+namespace NiagaraSharedDataUtilities
+{
+	FNiagaraSharedDataPtr GetOrCreateSharedData_StaticMesh(FNiagaraWorldManager* WorldMan, UStaticMeshComponent* StaticMeshComponent, TArrayView<UNiagaraComponent*> ReferencingComponents)
+	{
+		UNiagaraDataManager_StaticMesh& MeshDataMan = WorldMan->GetOrCreateDataManager<UNiagaraDataManager_StaticMesh>();
+		FNiagaraSharedData_StaticMeshPtr SharedMeshData = MeshDataMan.GetOrCreateSharedMeshData(StaticMeshComponent);
+		
+		for(UNiagaraComponent* ReferencingComp : ReferencingComponents)
+		{
+			MeshDataMan.RegisterReferencingSystem(ReferencingComp, SharedMeshData);
+		}
+
+		//Return the shared data pointer associated with this static mesh.
+		//This is valid as long as there are things referencing SharedMeshData
+		return SharedMeshData;
+	}
+}
diff --git a/Engine/Plugins/FX/Niagara/Source/Niagara/Private/NiagaraWorldManager.cpp b/Engine/Plugins/FX/Niagara/Source/Niagara/Private/NiagaraWorldManager.cpp
--- a/Engine/Plugins/FX/Niagara/Source/Niagara/Private/NiagaraWorldManager.cpp
+++ b/Engine/Plugins/FX/Niagara/Source/Niagara/Private/NiagaraWorldManager.cpp
@@ -28,6 +28,7 @@
 #include "SceneView.h"
 #include "UObject/UObjectIterator.h"
 #include "Logging/StructuredLog.h"
+#include "NiagaraDataManager.h"
 
 #include UE_INLINE_GENERATED_CPP_BY_NAME(NiagaraWorldManager)
 
diff --git a/Engine/Plugins/FX/Niagara/Source/Niagara/Private/NiagaraWorldManagerImpl.h b/Engine/Plugins/FX/Niagara/Source/Niagara/Private/NiagaraWorldManagerImpl.h
new file mode 100644
--- /dev/null
+++ b/Engine/Plugins/FX/Niagara/Source/Niagara/Private/NiagaraWorldManagerImpl.h
@@ -0,0 +1,28 @@
+// Copyright Epic Games, Inc. All Rights Reserved.
+
+#pragma once
+
+#include "NiagaraWorldManager.h"
+#include "NiagaraDataManager.h"
+
+template<typename T>
+T& FNiagaraWorldManager::GetOrCreateDataManager()
+{
+	FName ManagerName = T::StaticClass()->GetFName();
+	T* Ret = nullptr;
+	if(TObjectPtr<UNiagaraDataManager>* Found = DataManagers.Find(ManagerName))
+	{
+		Ret = Cast<T>(Found->Get());
+	}
+
+	if(Ret == nullptr)
+	{
+		TObjectPtr<UNiagaraDataManager>& NewMan = DataManagers.Add(ManagerName);
+		NewMan = NewObject<T>(GetTransientPackageAsObject(), ManagerName);
+		NewMan->Init(this);
+		Ret = CastChecked<T>(NewMan.Get());
+	}
+
+	check(Ret);
+	return *Ret;
+}
\ No newline at end of file
diff --git a/Engine/Plugins/FX/Niagara/Source/Niagara/Public/NiagaraDataChannelAccessContext.h b/Engine/Plugins/FX/Niagara/Source/Niagara/Public/NiagaraDataChannelAccessContext.h
--- a/Engine/Plugins/FX/Niagara/Source/Niagara/Public/NiagaraDataChannelAccessContext.h
+++ b/Engine/Plugins/FX/Niagara/Source/Niagara/Public/NiagaraDataChannelAccessContext.h
@@ -170,6 +170,9 @@ struct FNDCAccessContext : public FNDCAccessContextBase
 	
 	template<typename TAction>
 	void ForEachSpawnedSystem(TAction Func) const;
+
+	template<typename TContainerType>
+	void GetSpawnedSystems(TContainerType& Container) const;
 };
 
 /** Special legacy access context type. Used only to facilitate legacy NDCs to be usable by the new API. */
@@ -215,7 +218,7 @@ struct TNDCAccessContextType
 {
 	TNDCAccessContextType() = default;
 	[[nodiscard]] explicit TNDCAccessContextType(UScriptStruct* InStruct)
-		: AccessContextStruct(InStruct->IsChildOf(FNDCAccessContextBase::StaticStruct()) ? InStruct : nullptr)
+		: AccessContextStruct(InStruct && InStruct->IsChildOf(FNDCAccessContextBase::StaticStruct()) ? InStruct : nullptr)
 	{
 		check(AccessContextStruct);
 	}
diff --git a/Engine/Plugins/FX/Niagara/Source/Niagara/Public/NiagaraDataChannelAccessContextImpl.h b/Engine/Plugins/FX/Niagara/Source/Niagara/Public/NiagaraDataChannelAccessContextImpl.h
--- a/Engine/Plugins/FX/Niagara/Source/Niagara/Public/NiagaraDataChannelAccessContextImpl.h
+++ b/Engine/Plugins/FX/Niagara/Source/Niagara/Public/NiagaraDataChannelAccessContextImpl.h
@@ -14,4 +14,16 @@ void FNDCAccessContext::ForEachSpawnedSystem(TAction Func) const
 			Func(Comp);
 		}
 	}
+}
+
+template<typename TContainerType>
+void FNDCAccessContext::GetSpawnedSystems(TContainerType& Container)const
+{
+	for ( const FNDCSpawnedSystemRef& SysRef : SpawnedSystems )
+	{
+		if (UNiagaraComponent* Comp  = SysRef.Get())
+		{
+			Container.Add(Comp);
+		}
+	}
 }
\ No newline at end of file
diff --git a/Engine/Plugins/FX/Niagara/Source/Niagara/Public/NiagaraDataManager.h b/Engine/Plugins/FX/Niagara/Source/Niagara/Public/NiagaraDataManager.h
deleted file mode 100644
diff --git a/Engine/Plugins/FX/Niagara/Source/Niagara/Public/NiagaraDataManager_StaticMesh.h b/Engine/Plugins/FX/Niagara/Source/Niagara/Public/NiagaraDataManager_StaticMesh.h
deleted file mode 100644
diff --git a/Engine/Plugins/FX/Niagara/Source/Niagara/Public/NiagaraSharedDataUtilities.h b/Engine/Plugins/FX/Niagara/Source/Niagara/Public/NiagaraSharedDataUtilities.h
new file mode 100644
--- /dev/null
+++ b/Engine/Plugins/FX/Niagara/Source/Niagara/Public/NiagaraSharedDataUtilities.h
@@ -0,0 +1,47 @@
+// Copyright Epic Games, Inc. All Rights Reserved.
+
+#pragma once
+
+#include "NiagaraTypes.h"
+
+class UStaticMeshComponent;
+class UNiagaraComponent;
+class FNiagaraWorldManager;
+
+/** 
+Base class for generic data items held by data managers.
+Niagara Systems, NDCs and potentially external users will hold onto these when they are being used.
+Allows Data Managers to ensure lifetime of their internal data that may be being reference inside Niagara and by external users.
+*/
+struct NIAGARA_API FNiagaraSharedData : public TSharedFromThis<FNiagaraSharedData>
+{
+	virtual ~FNiagaraSharedData(){}
+
+	//Flag to mark this ref as invalid and it should be removed.
+	uint8 bInvalidated : 1 = false;
+
+	/** A generic ID that can be passed into Niagara in cases where Niagara needs to directly reference this data via a DI. */
+	FNiagaraID ID;
+};
+
+typedef TSharedPtr<FNiagaraSharedData> FNiagaraSharedDataPtr;
+
+/** Utilities for the creation and use of "Niagara Shared Data". i.e. generic data used/referenced both inside Niagara Systems and by external users. */
+namespace NiagaraSharedDataUtilities
+{
+	/**
+	 * Finds (or creates) Shared Data Object containing data about a given static mesh. 
+	 * The ID contained within can be passed into Niagara via User parameters or NDC to indirectly read data about the given static mesh.
+	 * Using NiagaraDataInterfaceStaticMeshIndirect. This can be useful in cases where binding to a specific static mesh is not possible via a traditional NiagaraDataInterfaceStaticMesh.
+	 *
+	 * @param WorldMan				The NiagaraWorldManager for the relevant world.
+	 * @param StaticMeshComponent	The Static Mesh Component to reference.
+	 * @param ReferencingComponents	A list of NiagaraComponents which will be referencing this data.
+	 * @return						A shared pointer to the data. 
+									This remains valid until 
+										* The mesh component is destroyed.
+										* There are no longer any references to the returned shared data. 
+									Inside Niagara, access using an invalid ID is safe but will fail and return an error in the DI function(s).
+	 */
+	NIAGARA_API FNiagaraSharedDataPtr GetOrCreateSharedData_StaticMesh(FNiagaraWorldManager* WorldMan, UStaticMeshComponent* StaticMeshComponent, TArrayView<UNiagaraComponent*> ReferencingComponents);
+};
\ No newline at end of file
diff --git a/Engine/Plugins/FX/Niagara/Source/Niagara/Public/NiagaraWorldManager.h b/Engine/Plugins/FX/Niagara/Source/Niagara/Public/NiagaraWorldManager.h
--- a/Engine/Plugins/FX/Niagara/Source/Niagara/Public/NiagaraWorldManager.h
+++ b/Engine/Plugins/FX/Niagara/Source/Niagara/Public/NiagaraWorldManager.h
@@ -21,7 +21,6 @@
 #include "NiagaraScalabilityManager.h"
 #include "NiagaraDebuggerCommon.h"
 #include "NiagaraDeferredMethodQueue.h"
-#include "NiagaraDataManager.h"
 
 #include "NiagaraWorldManager.generated.h"
 
@@ -33,6 +32,7 @@ class UNiagaraParameterCollectionInstance;
 class UNiagaraComponentPool;
 struct FNiagaraScalabilityState;
 class UNiagaraDataChannelHandler;
+class UNiagaraDataManager;
 
 enum class ENiagaraWorldManagerTickPass : int
 {
diff --git a/Engine/Plugins/FX/Niagara/Source/Niagara/Public/NiagaraWorldManagerImpl.h b/Engine/Plugins/FX/Niagara/Source/Niagara/Public/NiagaraWorldManagerImpl.h
deleted file mode 100644
