Update translator to use IsRapidIterationType() instead of duplicating code

#tests Ran translation over a project to confirm no functional changes
#Synced-CL Unknown

diff --git a/Engine/Plugins/FX/Niagara/Source/NiagaraEditor/Private/NiagaraHlslTranslator.cpp b/Engine/Plugins/FX/Niagara/Source/NiagaraEditor/Private/NiagaraHlslTranslator.cpp
--- a/Engine/Plugins/FX/Niagara/Source/NiagaraEditor/Private/NiagaraHlslTranslator.cpp
+++ b/Engine/Plugins/FX/Niagara/Source/NiagaraEditor/Private/NiagaraHlslTranslator.cpp
@@ -5222,7 +5222,7 @@ int32 TNiagaraHlslTranslator<GraphBridge>::GetRapidIterationParameter(const FNia
 	bool bIsCandidateForRapidIteration = false;
 	if (ActiveHistoryForFunctionCalls.InTopLevelFunctionCall(CompileOptions.TargetUsage))
 	{
-		if (Parameter.GetType() != FNiagaraTypeDefinition::GetBoolDef() && !Parameter.GetType().IsEnum() && !Parameter.IsDataInterface() && !Parameter.IsUObject())
+		if (FNiagaraStackGraphUtilities::IsRapidIterationType(Parameter.GetType()))
 		{
 			bIsCandidateForRapidIteration = true;
 		}
@@ -6170,7 +6170,7 @@ void TNiagaraHlslTranslator<GraphBridge>::ParameterMapSet(const FParamMapSetNode
 				{
 					if (Input < 0 || Input >= CompilationOutput.ScriptData.DataInterfaceInfo.Num())
 					{
-						Error(FText::Format(LOCTEXT("ParameterMapDataInterfaceNotFoundErrorFormat", "Data interface could not be found for parameter map set.  Paramter: {0}"), FText::FromName(Var.GetName())), SetNode, Inputs[i].Pin);
+						Error(FText::Format(LOCTEXT("ParameterMapDataInterfaceNotFoundErrorFormat", "Data interface could not be found for parameter map set.  Parameter: {0}"), FText::FromName(Var.GetName())), SetNode, Inputs[i].Pin);
 						continue;
 					}
 
@@ -6222,7 +6222,7 @@ void TNiagaraHlslTranslator<GraphBridge>::ParameterMapSet(const FParamMapSetNode
 				{
 					if (Input < 0 || Input >= CompilationOutput.ScriptData.UObjectInfos.Num())
 					{
-						Error(FText::Format(LOCTEXT("ParameterMapUObjectNotFoundErrorFormat", "UObject could not be found for parameter map set.  Paramter: {0}"), FText::FromName(Var.GetName())), SetNode, Inputs[i].Pin);
+						Error(FText::Format(LOCTEXT("ParameterMapUObjectNotFoundErrorFormat", "UObject could not be found for parameter map set.  Parameter: {0}"), FText::FromName(Var.GetName())), SetNode, Inputs[i].Pin);
 						continue;
 					}
 
@@ -7228,7 +7228,7 @@ void TNiagaraHlslTranslator<GraphBridge>::HandleParameterRead(int32 ParamMapHist
 			// Skip the case where the below condition is met, but it's overridden by a binding.
 			bIsCandidateForRapidIteration = false;
 		}
-		else if (InputPin != nullptr && GraphBridge::GetLinkedOutputPin(InputPin) == nullptr && Var.GetType() != FNiagaraTypeDefinition::GetBoolDef() && !Var.GetType().IsEnum() && !Var.GetType().IsDataInterface() && !Var.IsUObject())
+		else if (InputPin != nullptr && GraphBridge::GetLinkedOutputPin(InputPin) == nullptr && FNiagaraStackGraphUtilities::IsRapidIterationType(Var.GetType()))
 		{
 			bIsCandidateForRapidIteration = true;
 		}
