Add and use a GetPersistentTypeHash for the FNiagaraVariable to use when generating the name of custom hlsl function names

#tests Validated that subsequent translations would generate the same customhlsl function names
#Synced-CL Unknown

diff --git a/Engine/Plugins/FX/Niagara/Source/Niagara/Public/NiagaraTypes.h b/Engine/Plugins/FX/Niagara/Source/Niagara/Public/NiagaraTypes.h
--- a/Engine/Plugins/FX/Niagara/Source/Niagara/Public/NiagaraTypes.h
+++ b/Engine/Plugins/FX/Niagara/Source/Niagara/Public/NiagaraTypes.h
@@ -1424,10 +1424,17 @@ protected:
 	FNiagaraTypeDefinition TypeDef_DEPRECATED;
 #endif
 
+	// note that this implementation uses the hash of the variable's FName which is not consistent from run to run.  Use GetPersistentTypeHash()
+	// if the hash is being used in serialized data
 	friend inline uint32 GetTypeHash(const FNiagaraVariableBase& Var)
 	{
 		return HashCombine(GetTypeHash(Var.GetType()), GetTypeHash(Var.GetName()));
 	}
+
+	friend inline uint32 GetPersistentTypeHash(const FNiagaraVariableBase& Var)
+	{
+		return HashCombine(GetTypeHash(Var.GetType()), GetTypeHash(Var.GetName().ToString()));
+	}
 };
 
 template<>
diff --git a/Engine/Plugins/FX/Niagara/Source/NiagaraEditor/Private/NiagaraHlslTranslator.cpp b/Engine/Plugins/FX/Niagara/Source/NiagaraEditor/Private/NiagaraHlslTranslator.cpp
--- a/Engine/Plugins/FX/Niagara/Source/NiagaraEditor/Private/NiagaraHlslTranslator.cpp
+++ b/Engine/Plugins/FX/Niagara/Source/NiagaraEditor/Private/NiagaraHlslTranslator.cpp
@@ -8952,7 +8952,7 @@ void TNiagaraHlslTranslator<GraphBridge>::HandleCustomHlslNode(const FCustomHlsl
 		for (const FPin* Pin : Pins)
 		{
 			FNiagaraVariable PinVariable = GraphBridge::GetPinVariable(Pin, false /*bNeedsValue*/, ENiagaraStructConversion::UserFacing);
-			const uint32 PinVariableTypeHash = GetTypeHash(PinVariable);
+			const uint32 PinVariableTypeHash = GetPersistentTypeHash(PinVariable);
 			HashState.Update((const uint8*)&PinVariableTypeHash, sizeof(PinVariableTypeHash));
 		}
 
