- Add options to Niagara System Editor
- Currently only exposing auto play as an option here

#Synced-CL 48612007

diff --git a/Engine/Plugins/FX/Niagara/Source/NiagaraEditor/Private/Toolkits/SystemToolkitModes/NiagaraSystemToolkitMode_Default.cpp b/Engine/Plugins/FX/Niagara/Source/NiagaraEditor/Private/Toolkits/SystemToolkitModes/NiagaraSystemToolkitMode_Default.cpp
--- a/Engine/Plugins/FX/Niagara/Source/NiagaraEditor/Private/Toolkits/SystemToolkitModes/NiagaraSystemToolkitMode_Default.cpp
+++ b/Engine/Plugins/FX/Niagara/Source/NiagaraEditor/Private/Toolkits/SystemToolkitModes/NiagaraSystemToolkitMode_Default.cpp
@@ -109,6 +109,39 @@ void FNiagaraSystemToolkitMode_Default::ExtendToolbar()
 			return MenuBuilder.MakeWidget();
 		}
 
+		static TSharedRef<SWidget> FillSettingsMenu(FNiagaraSystemToolkit* Toolkit)
+		{
+			FMenuBuilder MenuBuilder(true, Toolkit->GetToolkitCommands());
+			if (FProperty* AutoPlayProperty = UNiagaraEditorSettings::StaticClass()->FindPropertyByName("bAutoPlay"))
+			{
+				MenuBuilder.AddMenuEntry(
+					AutoPlayProperty->GetDisplayNameText(),
+					AutoPlayProperty->GetToolTipText(),
+					FSlateIcon(),
+					FUIAction(
+						FExecuteAction::CreateLambda(
+							[]()
+							{
+								UNiagaraEditorSettings* EditorSettings = GetMutableDefault<UNiagaraEditorSettings>();
+								EditorSettings->SetAutoPlay(!EditorSettings->GetAutoPlay());
+							}
+						),
+						FCanExecuteAction(),
+						FIsActionChecked::CreateLambda(
+							[]() -> bool
+							{
+								return GetDefault<UNiagaraEditorSettings>()->GetAutoPlay();
+							}
+						)
+					),
+					NAME_None,
+					EUserInterfaceActionType::ToggleButton
+				);
+			}
+
+			return MenuBuilder.MakeWidget();
+		}
+
 		static void FillToolbar(FToolBarBuilder& ToolbarBuilder, FNiagaraSystemToolkit* Toolkit)
 		{
 			ToolbarBuilder.BeginSection("Apply");
@@ -285,6 +318,18 @@ void FNiagaraSystemToolkitMode_Default::ExtendToolbar()
 				ToolbarBuilder.EndSection();
 			}
 
+			ToolbarBuilder.BeginSection("Settings");
+			{
+				ToolbarBuilder.AddComboButton(
+					FUIAction(),
+					FOnGetContent::CreateStatic(Local::FillSettingsMenu, Toolkit),
+					FText(),
+					LOCTEXT("NiagaraEditorGlobalSettings", "Niagara Editor Settings"),
+					FSlateIcon(FAppStyle::GetAppStyleSetName(), "DetailsView.ViewOptions")
+				);
+			}
+			ToolbarBuilder.EndSection();
+
 			if (GbShowNiagaraDeveloperWindows != 0)
 			{
 				// a simple ticker to visually see how fast the ui is ticking and to see any hitches due to stack actions 
