Works around issues with Niagara's AsyncTasks compilation mode and how it would incorrectly handle GPU scripts in a multiplatform cook scenario.  Introduces a new cvar that will have it revert to the orginal behavior for handling those scripts which may introduce stalls in the cooker as we need to wait for the system compilation to complete.

#tests Cooked different projects and made sure that the GPU scripts are getting compiled properly in different scenarios
#Synced-CL Unknown

diff --git a/Engine/Plugins/FX/Niagara/Source/Niagara/Private/NiagaraScript.cpp b/Engine/Plugins/FX/Niagara/Source/Niagara/Private/NiagaraScript.cpp
--- a/Engine/Plugins/FX/Niagara/Source/Niagara/Private/NiagaraScript.cpp
+++ b/Engine/Plugins/FX/Niagara/Source/Niagara/Private/NiagaraScript.cpp
@@ -9,6 +9,7 @@
 #include "DataInterface/NiagaraDataInterfaceStaticMesh.h"
 #include "HAL/PlatformFileManager.h"
 #include "Interfaces/ITargetPlatform.h"
+#include "Interfaces/ITargetPlatformManagerModule.h"
 #include "Misc/Compression.h"
 #include "Misc/CoreMiscDefines.h"
 #include "Misc/FileHelper.h"
@@ -139,6 +140,14 @@ static FAutoConsoleVariableRef CVarNiagaraGpuScriptsCompiledBySystem(
 	ECVF_Default
 );
 
+static bool GNiagaraCompileGpuSupportsMultiPlatformCook = false;
+static FAutoConsoleVariableRef CVarSystemGpuSupportsMultiPlatformCook(
+	TEXT("fx.Niagara.SystemGpuSupportsMultiPlatformCook"),
+	GNiagaraCompileGpuSupportsMultiPlatformCook,
+	TEXT("If false GPU shaders for Niagara scripts will not be compiled by the system when doing a multiplatform cook. \n"),
+	ECVF_Default
+);
+
 namespace NiagaraScriptInternal
 {
 #if WITH_EDITORONLY_DATA
@@ -1628,8 +1637,8 @@ void UNiagaraScript::Serialize(FArchive& Ar)
 	if (Ar.IsCooking())
 	{
 		// Cook optimization can temporarily modify the rapid iteration parameters for the serialize call.  We don't want this to broadcast
-		// on changed messages because it will modify the change id on owning objects so we supress those messages starting here.
-		FNiagaraParameterStore::FScopedSuppressOnChanged SupressOnChanged(RapidIterationParameters);
+		// on changed messages because it will modify the change id on owning objects so we suppress those messages starting here.
+		FNiagaraParameterStore::FScopedSuppressOnChanged SuppressOnChanged(RapidIterationParameters);
 
 		UNiagaraSystem* System = GetTypedOuter<UNiagaraSystem>();
 		bool bUsesRapidIterationParams = System != nullptr && System->ShouldUseRapidIterationParameters();
@@ -3386,25 +3395,30 @@ TArray<FName> UNiagaraScript::FindShaderFormatsForCooking(const ITargetPlatform*
 
 void UNiagaraScript::BeginCacheForCookedPlatformData(const ITargetPlatform *TargetPlatform)
 {
+	// nothing to be done for cooking non-gpu compute scripts; their work is handled by the owning NiagaraSystem
+	if (Usage != ENiagaraScriptUsage::ParticleGPUComputeScript)
+	{
+		return;
+	}
+
 	if (AreGpuScriptsCompiledBySystem())
 	{
 		// if the system is responsible for generating the GPU scripts, then we don't need to do anything here
 		return;
 	}
 
+	// Commandlets like DerivedDataCacheCommandlet call BeginCacheForCookedPlatformData directly on objects. This may mean that
+	// we have not properly gotten the HLSL script generated by the time that we get here. This does the awkward work of
+	// waiting on the parent system to finish generating the HLSL before we can begin compiling it for the GPU.
+	if (UNiagaraSystem* SystemOwner = FindRootSystem())
+	{
+		SystemOwner->WaitForCompilationComplete();
+	}
+
 	TArray<FName> ShaderFormatsForCooking = FindShaderFormatsForCooking(TargetPlatform);
 
 	if (!ShaderFormatsForCooking.IsEmpty())
 	{
-		// Commandlets like DerivedDataCacheCommandlet call BeginCacheForCookedPlatformData directly on objects. This may mean that
-		// we have not properly gotten the HLSL script generated by the time that we get here. This does the awkward work of
-		// waiting on the parent system to finish generating the HLSL before we can begin compiling it for the GPU.
-		UNiagaraSystem* SystemOwner = FindRootSystem();
-		if (SystemOwner)
-		{
-			SystemOwner->WaitForCompilationComplete();
-		}
-
 		if (HasIdsRequiredForShaderCaching() == false)
 		{
 			UE_LOG(LogNiagara, Warning,
@@ -3487,8 +3501,26 @@ bool UNiagaraScript::IsCachedCookedPlatformDataLoaded(const ITargetPlatform* Tar
 // for the generation of the FNiagaraShaderMap and general compilation for GPU scripts
 bool UNiagaraScript::AreGpuScriptsCompiledBySystem()
 {
+	if (!GNiagaraGpuScriptsCompiledBySystem)
+	{
+		return false;
+	}
+
 	const ENiagaraCompilationMode CompilationMode = GetDefault<UNiagaraSettings>()->CompilationMode;
-	return GNiagaraGpuScriptsCompiledBySystem && CompilationMode == ENiagaraCompilationMode::AsyncTasks;
+	if (CompilationMode != ENiagaraCompilationMode::AsyncTasks)
+	{
+		return false;
+	}
+
+	if (!GNiagaraCompileGpuSupportsMultiPlatformCook && IsRunningCookCommandlet())
+	{
+		if (GetTargetPlatformManagerRef().GetActiveTargetPlatforms().Num() > 1)
+		{
+			return false;
+		}
+	}
+
+	return true;
 }
 
 void UNiagaraScript::CacheResourceShadersForCooking(EShaderPlatform ShaderPlatform, TArray<TUniquePtr<FNiagaraShaderScript>>& InOutCachedResources, const ITargetPlatform* TargetPlatform)
