- Fix potential race with LWC tile on render thread
- RenderData can be read on the RT while being written to on the GT, moving into DynamicData means it won't race
- In reality the LWC tile doesn't really change without going through a 'process' to handle the adjustment so likely not an issue, but was flagged by tsan

#Synced-CL 48563820

diff --git a/Engine/Plugins/FX/Niagara/Source/Niagara/Private/NiagaraComponent.cpp b/Engine/Plugins/FX/Niagara/Source/Niagara/Private/NiagaraComponent.cpp
--- a/Engine/Plugins/FX/Niagara/Source/Niagara/Private/NiagaraComponent.cpp
+++ b/Engine/Plugins/FX/Niagara/Source/Niagara/Private/NiagaraComponent.cpp
@@ -395,11 +395,6 @@ void FNiagaraSceneProxy::ReleaseUniformBuffers(bool bEmpty)
 	CustomUniformBuffers.Empty(bEmpty ? 0 : FMath::Min(ExpectedRendererCount, CurrentUBCount));
 }
 
-const FVector3f& FNiagaraSceneProxy::GetLWCRenderTile() const
-{
-	return RenderData ? RenderData->LWCRenderTile : FVector3f::ZeroVector;
-}
-
 TUniformBuffer<FPrimitiveUniformShaderParameters>* FNiagaraSceneProxy::GetCustomUniformBufferResource(FRHICommandListBase& RHICmdList, bool bHasVelocity, const FBox& InstanceBounds) const
 {
 	// Use a hash to determine if we can re-use any uniform buffer
@@ -2322,11 +2317,9 @@ void UNiagaraComponent::SendRenderDynamicData_Concurrent()
 	{
 		if (FNiagaraSystemRenderData* RenderData = NiagaraProxy->GetSystemRenderData())
 		{
-			RenderData->LWCRenderTile = SystemInstanceController->GetSystemInstance_Unsafe()->GetLWCTile();
-			
 			FNiagaraSceneProxy::FDynamicData NewProxyDynamicData;
+			NewProxyDynamicData.LWCRenderTile = SystemInstanceController->GetSystemInstance_Unsafe()->GetLWCTile();
 			NewProxyDynamicData.LODDistanceOverride = GetPreviewLODDistance();
-
 #if WITH_PARTICLE_PERF_STATS
 			NewProxyDynamicData.PerfStatsContext = GetPerfStatsContext();
 #endif
diff --git a/Engine/Plugins/FX/Niagara/Source/Niagara/Public/NiagaraSceneProxy.h b/Engine/Plugins/FX/Niagara/Source/Niagara/Public/NiagaraSceneProxy.h
--- a/Engine/Plugins/FX/Niagara/Source/Niagara/Public/NiagaraSceneProxy.h
+++ b/Engine/Plugins/FX/Niagara/Source/Niagara/Public/NiagaraSceneProxy.h
@@ -58,7 +58,7 @@ public:
 
 	inline const FMatrix& GetLocalToWorldInverse() const { return LocalToWorldInverse; }
 
-	NIAGARA_API const FVector3f& GetLWCRenderTile() const;
+	const FVector3f& GetLWCRenderTile() const { return DynamicData.LWCRenderTile; }
 
 	NIAGARA_API TUniformBuffer<FPrimitiveUniformShaderParameters>* GetCustomUniformBufferResource(FRHICommandListBase& RHICmdList, bool bHasVelocity, const FBox& InstanceBounds = FBox(ForceInitToZero)) const;
 
@@ -69,8 +69,9 @@ public:
 	/** Some proxy wide dynamic settings passed down with the emitter dynamic data. */
 	struct FDynamicData
 	{
-		bool bUseCullProxy = false;
-		float LODDistanceOverride = -1.0f;
+		FVector3f	LWCRenderTile = FVector3f::ZeroVector;
+		bool		bUseCullProxy = false;
+		float		LODDistanceOverride = -1.0f;
 #if WITH_PARTICLE_PERF_STATS
 		FParticlePerfStatsContext PerfStatsContext;
 #endif
diff --git a/Engine/Plugins/FX/Niagara/Source/Niagara/Public/NiagaraSystemRenderData.h b/Engine/Plugins/FX/Niagara/Source/Niagara/Public/NiagaraSystemRenderData.h
--- a/Engine/Plugins/FX/Niagara/Source/Niagara/Public/NiagaraSystemRenderData.h
+++ b/Engine/Plugins/FX/Niagara/Source/Niagara/Public/NiagaraSystemRenderData.h
@@ -85,7 +85,7 @@ public:
 
 	inline int32 GetNumRenderers() const { return RendererDrawOrder.Num(); }
 
-	FVector3f LWCRenderTile = FVector3f::ZeroVector;
+	//FVector3f LWCRenderTile = FVector3f::ZeroVector;
 private:
 	/** Emitter Renderers in the order they appear in the emitters. To be accessed by the GameThread */
 	TArray<FNiagaraRenderer*> EmitterRenderers_GT;
