- Fix for int vs uint on count buffer UAV

#Synced-CL 48463499

diff --git a/Engine/Plugins/FX/Niagara/Shaders/Private/Stateless/Template/GenerateCS_PostModules.ush b/Engine/Plugins/FX/Niagara/Shaders/Private/Stateless/Template/GenerateCS_PostModules.ush
--- a/Engine/Plugins/FX/Niagara/Shaders/Private/Stateless/Template/GenerateCS_PostModules.ush
+++ b/Engine/Plugins/FX/Niagara/Shaders/Private/Stateless/Template/GenerateCS_PostModules.ush
@@ -8,7 +8,7 @@ float	Common_SimulationDeltaTime;
 float	Common_SimulationInvDeltaTime;
 
 uint				Common_GPUCountBufferOffset;
-RWBuffer<int>		Common_GPUCountBuffer;
+RWBuffer<uint>		Common_GPUCountBuffer;
 
 DECLARE_SCALAR_ARRAY(uint,	SpawnInfo_NumActive, NIAGARA_MAX_GPU_SPAWN_INFOS);
 DECLARE_SCALAR_ARRAY(uint,	SpawnInfo_ParticleOffset, NIAGARA_MAX_GPU_SPAWN_INFOS);
