- Handle weak pointer going invalid in summary view

#Synced-CL 48612007

diff --git a/Engine/Plugins/FX/Niagara/Source/NiagaraEditor/Private/ViewModels/Stack/NiagaraStackInputCategory.cpp b/Engine/Plugins/FX/Niagara/Source/NiagaraEditor/Private/ViewModels/Stack/NiagaraStackInputCategory.cpp
--- a/Engine/Plugins/FX/Niagara/Source/NiagaraEditor/Private/ViewModels/Stack/NiagaraStackInputCategory.cpp
+++ b/Engine/Plugins/FX/Niagara/Source/NiagaraEditor/Private/ViewModels/Stack/NiagaraStackInputCategory.cpp
@@ -306,7 +306,12 @@ void UNiagaraStackSummaryCategory::Initialize(FRequiredEntryData InRequiredEntry
 
 FText UNiagaraStackSummaryCategory::GetDisplayName() const
 {
-	return CategoryViewModelWeakPtr.Pin()->GetCategoryName();
+	TSharedPtr<FHierarchyCategoryViewModel> CategoryViewModel = CategoryViewModelWeakPtr.Pin();
+	if (ensure(CategoryViewModel.IsValid()))
+	{
+		return CategoryViewModel->GetCategoryName();
+	}
+	return FText();
 }
 
 void UNiagaraStackSummaryCategory::RefreshChildrenInternal(const TArray<UNiagaraStackEntry*>& CurrentChildren, TArray<UNiagaraStackEntry*>& NewChildren, TArray<FStackIssue>& NewIssues)
@@ -750,12 +755,25 @@ void UNiagaraStackSummaryCategory::RefreshChildrenInternal(const TArray<UNiagara
 
 bool UNiagaraStackSummaryCategory::IsTopLevelCategory() const
 {
-	return CategoryViewModelWeakPtr.Pin()->GetData()->GetOuter()->IsA<UHierarchyRoot>();
+	TSharedPtr<FHierarchyCategoryViewModel> CategoryViewModel = CategoryViewModelWeakPtr.Pin();
+	if (ensure(CategoryViewModel.IsValid()))
+	{
+		return CategoryViewModel->GetData()->GetOuter()->IsA<UHierarchyRoot>();
+	}
+	return false;
 }
 
 FText UNiagaraStackSummaryCategory::GetTooltipText() const
 {
-	return Cast<UHierarchyCategory>(CategoryViewModelWeakPtr.Pin()->GetData())->GetTooltip();
+	TSharedPtr<FHierarchyCategoryViewModel> CategoryViewModel = CategoryViewModelWeakPtr.Pin();
+	if (ensure(CategoryViewModel.IsValid()))
+	{
+		if (const UHierarchyCategory* Category = Cast<const UHierarchyCategory>(CategoryViewModel->GetData()))
+		{
+			return Category->GetTooltip();
+		}
+	}
+	return FText();
 }
 
 int32 UNiagaraStackSummaryCategory::GetChildIndentLevel() const
