- Fix crashing in PSO precaching if a shader was not compiled / present for the platform

#tests No known repro for this at the moment

#Synced-CL 48563820

diff --git a/Engine/Plugins/FX/Niagara/Source/Niagara/Private/NiagaraEmitter.cpp b/Engine/Plugins/FX/Niagara/Source/Niagara/Private/NiagaraEmitter.cpp
--- a/Engine/Plugins/FX/Niagara/Source/Niagara/Private/NiagaraEmitter.cpp
+++ b/Engine/Plugins/FX/Niagara/Source/Niagara/Private/NiagaraEmitter.cpp
@@ -3534,7 +3534,13 @@ FGraphEventArray FVersionedNiagaraEmitterData::PrecacheComputePSOs(const UNiagar
 		FGraphEventArray PSOPrecacheEvents;
 		for (int32 i = 0; i < ShaderScript->GetNumPermutations(); ++i)
 		{
-			FRHIComputeShader* ComputeShader = ShaderScript->GetShaderGameThread(i).GetComputeShader();
+			FNiagaraShaderRef ShaderRef = ShaderScript->GetShaderGameThread(i);
+			if (ShaderRef.IsValid() == false)
+			{
+				continue;
+			}
+
+			FRHIComputeShader* ComputeShader = ShaderRef.GetComputeShader();
 			FPSOPrecacheRequestResult PSOPrecacheRequestResult = PipelineStateCache::PrecacheComputePipelineState(ComputeShader, TEXT("NiagaraEmitter"), true);
 			if (PSOPrecacheRequestResult.AsyncCompileEvent.IsValid())
 			{
@@ -3578,7 +3584,13 @@ FGraphEventArray FVersionedNiagaraEmitterData::PrecacheComputePSOs(const UNiagar
 
 				for (int32 i = 0; i < ShaderScript->GetNumPermutations(); ++i)
 				{
-					FRHIComputeShader* ComputeShader = ShaderScript->GetShaderGameThread(i).GetComputeShader();
+					FNiagaraShaderRef ShaderRef = ShaderScript->GetShaderGameThread(i);
+					if (ShaderRef.IsValid() == false)
+					{
+						continue;
+					}
+
+					FRHIComputeShader* ComputeShader = ShaderRef.GetComputeShader();
 					const EPSOPrecacheResult ShaderResult = PipelineStateCache::CheckPipelineStateInCache(ComputeShader);
 					if (ShaderResult != EPSOPrecacheResult::Complete)
 					{
@@ -3614,7 +3626,13 @@ FGraphEventArray FVersionedNiagaraEmitterData::PrecacheComputePSOs(const UNiagar
 		bool bAddToPSOReadyGraphTasks = GNiagaraEmitterComputePSOPrecacheMode == 1 && CVarPSOProxyCreationWhenPSOReady && CVarPSOProxyCreationWhenPSOReady->GetInt() != 0;
 		for (int32 i=0; i < ShaderScript->GetNumPermutations(); ++i)
 		{
-			FRHIComputeShader* ComputeShader = ShaderScript->GetShaderGameThread(i).GetComputeShader();
+			FNiagaraShaderRef ShaderRef = ShaderScript->GetShaderGameThread(i);
+			if (ShaderRef.IsValid() == false)
+			{
+				continue;
+			}
+
+			FRHIComputeShader* ComputeShader = ShaderRef.GetComputeShader();
 			FPSOPrecacheRequestResult PSOPrecacheRequestResult = PipelineStateCache::PrecacheComputePipelineState(ComputeShader, TEXT("NiagaraEmitter"), true);
 			if (PSOPrecacheRequestResult.AsyncCompileEvent != nullptr && bAddToPSOReadyGraphTasks)
 			{
