Fixed incorrect keygen when overriding Niagara System

#Synced-CL 48349740

diff --git a/Engine/Plugins/FX/Niagara/Source/Niagara/Private/NiagaraDataChannel_Map.cpp b/Engine/Plugins/FX/Niagara/Source/Niagara/Private/NiagaraDataChannel_Map.cpp
--- a/Engine/Plugins/FX/Niagara/Source/Niagara/Private/NiagaraDataChannel_Map.cpp
+++ b/Engine/Plugins/FX/Niagara/Source/Niagara/Private/NiagaraDataChannel_Map.cpp
@@ -169,9 +169,10 @@ void UNiagaraDataChannelHandler_MapBase::GenerateKey(FNDCAccessContextInst& Acce
 	const FNDCAccessContext_MapBase& TypedContext = AccessContext.GetChecked<FNDCAccessContext_MapBase>();
 
 	// We write our systems into the key so that all users can use a single NDC that has internally distinct buckets for different handler systems.
-	if(IsValid(TypedContext.SystemToSpawn))
+	const UObject* SystemToSpawn = GetSystemToSpawn(AccessContext);
+	if(IsValid(SystemToSpawn))
 	{
-		KeyWriter << TypedContext.SystemToSpawn->GetFName();
+		KeyWriter << SystemToSpawn->GetFName();
 	}
 }
 
