Refactor Niagara's audio sampling DIs (Oscilloscope & Spectrum) to be more inline with the current structure of DI and encapsulated the bulky audio data processing into a GeneratedData object so that it can be shared between multiple DI (assuming that they are compatible).

#Synced-CL Unknown

diff --git a/Engine/Plugins/FX/Niagara/Shaders/Private/NiagaraDataInterfaceAudioCommon.ush b/Engine/Plugins/FX/Niagara/Shaders/Private/NiagaraDataInterfaceAudioCommon.ush
new file mode 100644
--- /dev/null
+++ b/Engine/Plugins/FX/Niagara/Shaders/Private/NiagaraDataInterfaceAudioCommon.ush
@@ -0,0 +1,34 @@
+// Copyright Epic Games, Inc. All Rights Reserved.
+
+#pragma once
+
+float SampleInterleavedFloatBuffer(Buffer<float> InBuffer, int ChannelCount, int ChannelResolution, float NormalizedPosition, int ChannelIndex)
+{
+	float FrameIndex = saturate(NormalizedPosition) * (ChannelResolution - 1);
+	int LowerIndex = max(0, floor(FrameIndex));
+	int UpperIndex = min(LowerIndex + 1, ChannelResolution - 1);
+	ChannelIndex = clamp(ChannelIndex, 0, ChannelCount);
+
+	float Fraction = FrameIndex - LowerIndex;
+
+	float LowerValue = InBuffer.Load(LowerIndex * ChannelCount + ChannelIndex);
+	float UpperValue = InBuffer.Load(UpperIndex * ChannelCount + ChannelIndex);
+
+	return lerp(LowerValue, UpperValue, Fraction);
+}
+
+float SamplePlanarFloatBuffer(Buffer<float> InBuffer, int ChannelCount, int ChannelResolution, float NormalizedPosition, int ChannelIndex)
+{
+	float FrameIndex = saturate(NormalizedPosition) * (ChannelResolution - 1);
+	int LowerIndex = max(0, floor(FrameIndex));
+	int UpperIndex = min(LowerIndex + 1, ChannelResolution - 1);
+	ChannelIndex = clamp(ChannelIndex, 0, ChannelCount);
+
+	float Fraction = FrameIndex - LowerIndex;
+
+	float LowerValue = InBuffer.Load(ChannelIndex * ChannelCount + LowerIndex);
+	float UpperValue = InBuffer.Load(ChannelIndex * ChannelCount + UpperIndex);
+
+	return lerp(LowerValue, UpperValue, Fraction);
+}
+
diff --git a/Engine/Plugins/FX/Niagara/Shaders/Private/NiagaraDataInterfaceAudioOscilloscopeTemplate.ush b/Engine/Plugins/FX/Niagara/Shaders/Private/NiagaraDataInterfaceAudioOscilloscopeTemplate.ush
new file mode 100644
--- /dev/null
+++ b/Engine/Plugins/FX/Niagara/Shaders/Private/NiagaraDataInterfaceAudioOscilloscopeTemplate.ush
@@ -0,0 +1,21 @@
+// Copyright Epic Games, Inc. All Rights Reserved.
+
+Buffer<float> {ParameterName}_OscilloscopeBuffer;
+int {ParameterName}_ChannelCount;
+int {ParameterName}_Resolution;
+
+void SampleAudioBuffer_{ParameterName}(float In_NormalizedPosition, int In_ChannelIndex, out float Out_Val)
+{
+	Out_Val = SampleInterleavedFloatBuffer(
+		{ParameterName}_OscilloscopeBuffer,
+		{ParameterName}_ChannelCount,
+		{ParameterName}_Resolution,
+		In_NormalizedPosition,
+		In_ChannelIndex);
+}
+
+void GetAudioBufferNumChannels_{ParameterName}(out int Out_Val)
+{
+	Out_Val = {ParameterName}_ChannelCount;
+}
+
diff --git a/Engine/Plugins/FX/Niagara/Shaders/Private/NiagaraDataInterfaceAudioSpectrumTemplate.ush b/Engine/Plugins/FX/Niagara/Shaders/Private/NiagaraDataInterfaceAudioSpectrumTemplate.ush
new file mode 100644
--- /dev/null
+++ b/Engine/Plugins/FX/Niagara/Shaders/Private/NiagaraDataInterfaceAudioSpectrumTemplate.ush
@@ -0,0 +1,20 @@
+// Copyright Epic Games, Inc. All Rights Reserved.
+
+Buffer<float> {ParameterName}_SpectrumBuffer;
+int {ParameterName}_ChannelCount;
+int {ParameterName}_Resolution;
+
+void AudioSpectrum_{ParameterName}(float In_NormalizedPosition, int In_ChannelIndex, out float Out_Val)
+{
+	Out_Val = SamplePlanarFloatBuffer(
+		{ParameterName}_SpectrumBuffer,
+		{ParameterName}_ChannelCount,
+		{ParameterName}_Resolution,
+		In_NormalizedPosition,
+		In_ChannelIndex);
+}
+
+void GetNumChannels_{ParameterName}(out int Out_Val)
+{
+	Out_Val = {ParameterName}_ChannelCount;
+}
diff --git a/Engine/Plugins/FX/Niagara/Source/Niagara/Classes/NiagaraDataInterfaceAudio.h b/Engine/Plugins/FX/Niagara/Source/Niagara/Classes/NiagaraDataInterfaceAudio.h
deleted file mode 100644
diff --git a/Engine/Plugins/FX/Niagara/Source/Niagara/Classes/NiagaraDataInterfaceAudioOscilloscope.h b/Engine/Plugins/FX/Niagara/Source/Niagara/Classes/NiagaraDataInterfaceAudioOscilloscope.h
--- a/Engine/Plugins/FX/Niagara/Source/Niagara/Classes/NiagaraDataInterfaceAudioOscilloscope.h
+++ b/Engine/Plugins/FX/Niagara/Source/Niagara/Classes/NiagaraDataInterfaceAudioOscilloscope.h
@@ -1,90 +1,10 @@
 // Copyright Epic Games, Inc. All Rights Reserved.
 #pragma once
 
-#include "NiagaraCommon.h"
-#include "NiagaraShared.h"
-#include "VectorVM.h"
-#include "StaticMeshResources.h"
 #include "NiagaraDataInterface.h"
-#include "AudioDevice.h"
-#include "AudioDeviceManager.h"
-#include "DSP/MultithreadedPatching.h"
-#include "NiagaraDataInterfaceAudio.h"
-#include "NiagaraDataInterfaceAudioOscilloscope.generated.h"
-
-struct FNiagaraDataInterfaceProxyOscilloscope : public FNiagaraDataInterfaceProxy
-{
-	FNiagaraDataInterfaceProxyOscilloscope(int32 InResolution, float InScopeInMillseconds);
-
-	~FNiagaraDataInterfaceProxyOscilloscope();
-
-	void OnBeginDestroy();
-
-	/**
-	 * Sample vertical displacement from the oscilloscope buffer.
-	 * @param NormalizedPositionInBuffer Horizontal position in the Oscilloscope buffer, from 0.0 to 1.0.
-	 * @param Channel channel index.
-	 * @return Amplitude at this position.
-	 */
-	float SampleAudio(float NormalizedPositionInBuffer, int32 Channel, int32 NumFramesInBuffer, int32 NumChannelsInBuffer);
-
-	/**
-	 * @return the number of channels in the buffer.
-	 */
-	int32 GetNumChannels();
-
-	// Called when the Submix property changes.
-	void OnUpdateSubmix(USoundSubmix* Submix);
-
-	void RegisterToAllAudioDevices();
-	void UnregisterFromAllAudioDevices(USoundSubmix* Submix);
-
-	// Called when Resolution or Zoom are changed.
-	void OnUpdateResampling(int32 InResolution, float InScopeInMilliseconds);
-
-	// This function enqueues a render thread command to decimate the pop audio off of the SubmixListener, downsample it, and post it to the GPUAudioBuffer.
-	void PostAudioToGPU();
-	FReadBuffer& ComputeAndPostSRV();
-
-	// This function pops audio and downsamples it to our specified resolution. Returns the number of frames of audio in the downsampled buffer.
-	int32 DownsampleAudioToBuffer();
-
-	virtual int32 PerInstanceDataPassedToRenderThreadSize() const override
-	{
-		return 0;
-	}
+#include "Sound/SoundSubmix.h"
 
-private:
-	void OnNewDeviceCreated(Audio::FDeviceId InID);
-	void OnDeviceDestroyed(Audio::FDeviceId InID);
-
-	TMap<Audio::FDeviceId, TSharedPtr<FNiagaraSubmixListener>> SubmixListeners;
-
-	// This mixer is patched into by all instances of FNiagaraSubmixListener in the SubmixListeners map, and is consumed by DownsampleAudioToBuffer().
-	Audio::FPatchMixer PatchMixer;
-
-	USoundSubmix* SubmixRegisteredTo;
-	bool bIsSubmixListenerRegistered;
-
-	int32 Resolution;
-	float ScopeInMilliseconds;
-
-	// The buffer we downsample PoppedBuffer to based on the Resolution property.
-	Audio::FAlignedFloatBuffer PopBuffer;
-	Audio::FAlignedFloatBuffer DownsampledBuffer;
-
-	// Handle for the SRV used by the generated HLSL.
-	FReadBuffer GPUDownsampledBuffer;
-	FThreadSafeCounter NumChannelsInDownsampledBuffer;
-	
-	// Buffer read by VectorVM worker threads. This vector is guaranteed to not be mutated during the VectorVM tasks.
-	Audio::FAlignedFloatBuffer VectorVMReadBuffer;
-
-	FDelegateHandle DeviceCreatedHandle;
-	FDelegateHandle DeviceDestroyedHandle;
-
-	FCriticalSection DownsampleBufferLock;
-};
+#include "NiagaraDataInterfaceAudioOscilloscope.generated.h"
 
 /** Data Interface allowing sampling of recent audio data. */
 UCLASS(EditInlineNew, Category = "Audio", CollapseCategories, meta = (DisplayName = "Audio Oscilloscope"), MinimalAPI)
@@ -92,26 +12,19 @@ class UNiagaraDataInterfaceAudioOscilloscope : public UNiagaraDataInterface
 {
 	GENERATED_UCLASS_BODY()
 
-	BEGIN_SHADER_PARAMETER_STRUCT(FShaderParameters, )
-		SHADER_PARAMETER(int,				NumChannels)
-		SHADER_PARAMETER_SRV(Buffer<float>,	AudioBuffer)
-	END_SHADER_PARAMETER_STRUCT()
-
 public:
 	UPROPERTY(EditAnywhere, Category = "Oscilloscope")
 	TObjectPtr<USoundSubmix> Submix;
 
-	static const int32 MaxBufferResolution = 8192;
-
 	// The number of samples of audio to pass to the GPU. Audio will be resampled to fit this resolution.
 	// Increasing this number will increase the resolution of the waveform, but will also increase usage of the GPU memory bus,
 	// potentially causing issues across the application.
 	UPROPERTY(EditAnywhere, Category = "Oscilloscope", AdvancedDisplay, meta = (ClampMin = "64", ClampMax = "8192"))
-	int32 Resolution;
+	int32 Resolution = 512;
 
 	// The number of milliseconds of audio to show.
 	UPROPERTY(EditAnywhere, Category = "Oscilloscope", meta = (ClampMin = "5.0", ClampMax = "400.0"))
-	float ScopeInMilliseconds;
+	float ScopeInMilliseconds = 20.0f;
 
 	//VM function overrides:
 	NIAGARA_API void SampleAudio(FVectorVMExternalFunctionContext& Context);
@@ -127,21 +40,23 @@ public:
 #if WITH_EDITORONLY_DATA
 	NIAGARA_API virtual bool AppendCompileHash(FNiagaraCompileHashVisitor* InVisitor) const override;
 	NIAGARA_API virtual bool GetFunctionHLSL(const FNiagaraDataInterfaceGPUParamInfo& ParamInfo, const FNiagaraDataInterfaceGeneratedFunction& FunctionInfo, int FunctionInstanceIndex, FString& OutHLSL) override;
+	NIAGARA_API virtual void GetCommonHLSL(FString& OutHLSL) override;
 	NIAGARA_API virtual void GetParameterDefinitionHLSL(const FNiagaraDataInterfaceGPUParamInfo& ParamInfo, FString& OutHLSL) override;
 #endif
 	NIAGARA_API virtual void BuildShaderParameters(FNiagaraShaderParametersBuilder& ShaderParametersBuilder) const override;
 	NIAGARA_API virtual void SetShaderParameters(const FNiagaraDataInterfaceSetShaderParametersContext& Context) const override;
 
 	NIAGARA_API virtual bool Equals(const UNiagaraDataInterface* Other) const override;
-
-#if WITH_EDITOR
-	NIAGARA_API virtual void PostEditChangeProperty(struct FPropertyChangedEvent& PropertyChangedEvent) override;
-#endif // WITH_EDITOR
-
 	NIAGARA_API virtual void PostInitProperties() override;
-	NIAGARA_API virtual void BeginDestroy() override;
-	NIAGARA_API virtual void PostLoad() override;
 	
+	NIAGARA_API virtual bool InitPerInstanceData(void* PerInstanceData, FNiagaraSystemInstance* SystemInstance) override;
+	NIAGARA_API virtual void DestroyPerInstanceData(void* PerInstanceData, FNiagaraSystemInstance* SystemInstance) override;
+	NIAGARA_API virtual bool PerInstanceTick(void* PerInstanceData, FNiagaraSystemInstance* SystemInstance, float DeltaSeconds) override;
+	NIAGARA_API virtual int32 PerInstanceDataSize() const override;
+	NIAGARA_API virtual void ProvidePerInstanceDataForRenderThread(void* DataForRenderThread, void* PerInstanceData, const FNiagaraSystemInstanceID& SystemInstance) override;
+	virtual bool HasPreSimulateTick() const override { return true; }
+	virtual bool HasTickGroupPrereqs() const override { return true; }
+	NIAGARA_API virtual ETickingGroup CalculateTickGroup(const void* PerInstanceData) const override;
 
 protected:
 #if WITH_EDITORONLY_DATA
diff --git a/Engine/Plugins/FX/Niagara/Source/Niagara/Classes/NiagaraDataInterfaceAudioSpectrum.h b/Engine/Plugins/FX/Niagara/Source/Niagara/Classes/NiagaraDataInterfaceAudioSpectrum.h
--- a/Engine/Plugins/FX/Niagara/Source/Niagara/Classes/NiagaraDataInterfaceAudioSpectrum.h
+++ b/Engine/Plugins/FX/Niagara/Source/Niagara/Classes/NiagaraDataInterfaceAudioSpectrum.h
@@ -2,167 +2,38 @@
 #pragma once
 
 #include "NiagaraDataInterface.h"
-#include "DSP/SlidingWindow.h"
-#include "DSP/DeinterleaveView.h"
-#include "DSP/ConstantQ.h"
-#include "DSP/FFTAlgorithm.h"
-#include "DSP/FloatArrayMath.h"
-#include "DSP/AudioFFT.h"
-#include "NiagaraDataInterfaceAudio.h"
-#include "RHIUtilities.h"
+#include "Sound/SoundSubmix.h"
 
 #include "NiagaraDataInterfaceAudioSpectrum.generated.h"
 
-/** FNiagaraDataInterfaceProxySpectrum
- *  
- *  Proxy class for calculating spectrums on the rendering thread for Niagara.
- */
-struct FNiagaraDataInterfaceProxySpectrum : public FNiagaraDataInterfaceProxySubmix
-{
-	/** Construct an FNiagaraDataInterfaceProxySpectrum
-	 *
-	 * 	@param InMinimumFrequency - Minimum frequency in spectrum.
-	 * 	@param InMaximumFrequency - Maximum frequency in spectrum.
-	 * 	@param InNumBands - Number of bands in spectrum.
-	 * 	@param InNoiseFloorDb - Decibels level considered as silence. 
-	 * 	@param InNumSamplesToBuffer - Number of samples to buffer internally across thread boundaries.
-	 */
-	FNiagaraDataInterfaceProxySpectrum(float InMinimumFrequency, float InMaximumFrequency, int32 InNumBands, float InNoiseFloorDb, int32 InNumSamplesToBuffer);
-
-	virtual ~FNiagaraDataInterfaceProxySpectrum();
-
-	/** Sample value from the spectrum buffer.
-	 *  
-	 *  @param InNormalizedPositionInSpectrum - Value between 0 and 1 which represents the relative frequency in the spectrum.  
-	 *                                          0 corresponds to the MinimumFrequency, 1 corresponds to the MaximumFrequency.
-	 *                                          Values between 0 and 1 correspond to logarithmically spaced frequencys in-between
-	 *                                          the limits.
-	 *  @param InChannelIndex - The index of the audio channel to sample.
-	 *
-	 *  @return Amplitude at this position. 0.0 corresponds to the nosie floor, while 1.0 corresponds to full volume. Some values
-	 *          may exceed 1.0.
-	 */
-	float GetSpectrumValue(float InNormalizedPositionInSpectrum, int32 InChannelIndex);
-
-	int32 GetNumBands() const;
-
-	/** Updates the minimum and maximum frequency of the CQT on the render thread. */
-	void UpdateCQT(float InMinimumFrequency, float InMaximumFrequency, int32 InNumBands);
-
-	/** Updates the value considered silence in decibles on the render thread. */
-	void UpdateNoiseFloor(float InNoiseFloorDb);
-
-
-	/** This function enqueues a render thread command to pop audio off of the SubmixListener, transform it into a CQT, and 
-	 * post it to the GPUAudioBuffer.
-	 */
-	void PostDataToGPU();
-	FReadBuffer& ComputeAndPostSRV();
-
-	virtual int32 PerInstanceDataPassedToRenderThreadSize() const override
-	{
-		return 0;
-	}
-
-	// Calculates the CQT
-	void UpdateSpectrum();
-
-private:
-	typedef Audio::TSlidingBuffer<float> FSlidingBuffer;
-	typedef Audio::TAutoSlidingWindow<float, Audio::FAudioBufferAlignedAllocator> FSlidingWindow;
-	typedef Audio::TAutoDeinterleaveView<float, Audio::FAudioBufferAlignedAllocator> FDeinterleaveView;
-	typedef FDeinterleaveView::TChannel<Audio::FAudioBufferAlignedAllocator> FChannel;
-
-	// Updates internal objects for num channels and samplerate
-	void SetAudioFormat(int32 InNumChannels, float InSampleRate);
-
-	void SetNoiseFloor(float InNoiseFloorDb);
-
-	// Resize internal abuffers
-	void ResizeCQT(float InMinimumFrequency, float InMaximumFrequency, int32 InNumBands);
-	void ResizeSpectrumBuffer(int32 InNumChannels, int32 InNumBands);
-	void ResizeAudioTransform(float InMinimumFrequency, float InMaximumFrequency, float InSampleRate, int32 InNumBands);
-	void ResizeWindow(int32 InNumChannels, int32 InFFTSize);
-
-
-	// Get settings 
-	static void Clamp(float& InMinimumFrequency, float& InMaximumFrequency, int32& InNumBand);
-	static float GetNumOctaves(float InMinimumFrequency, float InMaximumFrequency);
-	static float GetNumBandsPerOctave(int32 InNumBands, float InNumOctaves);
-	static float GetBandwidthStretch(float InSampleRate, float InFFTSize, float InNumBandsPerOctave, float InMinimumFrequency);
-
-	static Audio::FPseudoConstantQKernelSettings GetConstantQSettings(float InMinimumFrequency, float InMaximumFrequency, int32 InNumBands, float InNumBandsPerOctave, float InBandwidthStretch);
-
-	static Audio::FFFTSettings GetFFTSettings(float InMinimumFrequency, float InSampleRate, float InNumBandsPerOctave);
-
-	float MinimumFrequency;
-	float MaximumFrequency;
-	TAtomic<int32> NumBands;
-	float NoiseFloorDb;
-
-	int32 NumChannels;
-	float SampleRate;
-	float FFTScale;
-
-	Audio::FAlignedFloatBuffer PopBuffer;
-	TUniquePtr< FSlidingBuffer > SlidingBuffer;
-
-	TArray<Audio::FAlignedFloatBuffer> ChannelSpectrumBuffers;
-
-	TUniquePtr<Audio::FContiguousSparse2DKernelTransform> CQTKernel;
-	TUniquePtr<Audio::IFFTAlgorithm> FFTAlgorithm;
-
-	Audio::FAlignedFloatBuffer InterleavedBuffer;
-	Audio::FAlignedFloatBuffer DeinterleavedBuffer;
-	Audio::FAlignedFloatBuffer FFTInputBuffer;
-	Audio::FAlignedFloatBuffer FFTOutputBuffer;
-	Audio::FAlignedFloatBuffer PowerSpectrumBuffer;
-	Audio::FAlignedFloatBuffer SpectrumBuffer;
-	Audio::FAlignedFloatBuffer WindowBuffer;
-
-	// Handle for the SRV used by the generated HLSL.
-	FReadBuffer GPUBuffer;
-
-	FCriticalSection BufferLock;
-};
+class FNDIAudio_SharedResource;
 
 /** Data Interface allowing sampling of recent audio spectrum. */
 UCLASS(EditInlineNew, Category = "Audio", meta = (DisplayName = "Audio Spectrum"), MinimalAPI)
-class UNiagaraDataInterfaceAudioSpectrum : public UNiagaraDataInterfaceAudioSubmix
+class UNiagaraDataInterfaceAudioSpectrum : public UNiagaraDataInterface
 {
 	GENERATED_UCLASS_BODY()
 
-	BEGIN_SHADER_PARAMETER_STRUCT(FShaderParameters, )
-		SHADER_PARAMETER(int,				NumChannels)
-		SHADER_PARAMETER(int,				Resolution)
-		SHADER_PARAMETER_SRV(Buffer<float>, SpectrumBuffer)
-	END_SHADER_PARAMETER_STRUCT()
-
 public:
-	// VM function names
-	static NIAGARA_API const FName GetSpectrumFunctionName;
-	static NIAGARA_API const FName GetNumChannelsFunctionName;
-
-	// Global variable prefixes
-	static NIAGARA_API const FString GetSpectrumName;
-	static NIAGARA_API const FString NumChannelsName;
-	static NIAGARA_API const FString ResolutionName;
-	
+	/** The audio submix where analyzed. */
+	UPROPERTY(EditAnywhere, Category = "Audio")
+	TObjectPtr<USoundSubmix> Submix;
+
 	/** The number of spectrum samples to pass to the GPU */
 	UPROPERTY(EditAnywhere, Category = "Spectrum", meta = (ClampMin = "16", ClampMax = "1024") )
-	int32 Resolution;
+	int32 Resolution = 512;
 
 	/** The minimum frequency represented in the spectrum. */
 	UPROPERTY(EditAnywhere, Category = "Spectrum", meta = (ClampMin = "20.0", ClampMax = "20000.0"))
-	float MinimumFrequency;
+	float MinimumFrequency = 55.0f;
 
 	/** The maximum frequency represented in the spectrum. */
 	UPROPERTY(EditAnywhere, Category = "Spectrum", meta = (ClampMin = "20.0", ClampMax = "20000.0"))
-	float MaximumFrequency;
+	float MaximumFrequency = 10000.f;
 
 	/** The decibel level considered as silence. This is used to scale the output of the spectrum. */
 	UPROPERTY(EditAnywhere, AdvancedDisplay, Category = "Spectrum", meta = (ClampMin = "-120.0", ClampMax = "0.0"))
-	float NoiseFloorDb;
+	float NoiseFloorDb = -60.0f;
 
 	//VM function overrides:
 	NIAGARA_API void GetSpectrumValue(FVectorVMExternalFunctionContext& Context);
@@ -178,20 +49,23 @@ public:
 #if WITH_EDITORONLY_DATA
 	NIAGARA_API virtual bool AppendCompileHash(FNiagaraCompileHashVisitor* InVisitor) const override;
 	NIAGARA_API virtual bool GetFunctionHLSL(const FNiagaraDataInterfaceGPUParamInfo& ParamInfo, const FNiagaraDataInterfaceGeneratedFunction& FunctionInfo, int FunctionInstanceIndex, FString& OutHLSL) override;
+	NIAGARA_API virtual void GetCommonHLSL(FString& OutHLSL) override;
 	NIAGARA_API virtual void GetParameterDefinitionHLSL(const FNiagaraDataInterfaceGPUParamInfo& ParamInfo, FString& OutHLSL) override;
 #endif
 	NIAGARA_API virtual void BuildShaderParameters(FNiagaraShaderParametersBuilder& ShaderParametersBuilder) const override;
 	NIAGARA_API virtual void SetShaderParameters(const FNiagaraDataInterfaceSetShaderParametersContext& Context) const override;
 
 	NIAGARA_API virtual bool Equals(const UNiagaraDataInterface* Other) const override;
-
-#if WITH_EDITOR
-	NIAGARA_API virtual void PostEditChangeProperty(struct FPropertyChangedEvent& PropertyChangedEvent) override;
-#endif // WITH_EDITOR
-
-
 	NIAGARA_API virtual void PostInitProperties() override;
-	NIAGARA_API virtual void PostLoad() override;
+
+	NIAGARA_API virtual bool InitPerInstanceData(void* PerInstanceData, FNiagaraSystemInstance* SystemInstance) override;
+	NIAGARA_API virtual void DestroyPerInstanceData(void* PerInstanceData, FNiagaraSystemInstance* SystemInstance) override;
+	NIAGARA_API virtual bool PerInstanceTick(void* PerInstanceData, FNiagaraSystemInstance* SystemInstance, float DeltaSeconds) override;
+	NIAGARA_API virtual int32 PerInstanceDataSize() const override;
+	NIAGARA_API virtual void ProvidePerInstanceDataForRenderThread(void* DataForRenderThread, void* PerInstanceData, const FNiagaraSystemInstanceID& SystemInstance) override;
+	virtual bool HasPreSimulateTick() const override { return true; }
+	virtual bool HasTickGroupPrereqs() const override { return true; }
+	NIAGARA_API virtual ETickingGroup CalculateTickGroup(const void* PerInstanceData) const override;
 
 protected:
 #if WITH_EDITORONLY_DATA
diff --git a/Engine/Plugins/FX/Niagara/Source/Niagara/Private/NiagaraDataInterfaceAudio.cpp b/Engine/Plugins/FX/Niagara/Source/Niagara/Private/NiagaraDataInterfaceAudio.cpp
deleted file mode 100644
diff --git a/Engine/Plugins/FX/Niagara/Source/Niagara/Private/NiagaraDataInterfaceAudioOscilloscope.cpp b/Engine/Plugins/FX/Niagara/Source/Niagara/Private/NiagaraDataInterfaceAudioOscilloscope.cpp
--- a/Engine/Plugins/FX/Niagara/Source/Niagara/Private/NiagaraDataInterfaceAudioOscilloscope.cpp
+++ b/Engine/Plugins/FX/Niagara/Source/Niagara/Private/NiagaraDataInterfaceAudioOscilloscope.cpp
@@ -1,192 +1,205 @@
 // Copyright Epic Games, Inc. All Rights Reserved.
 
 #include "NiagaraDataInterfaceAudioOscilloscope.h"
-#include "AudioDevice.h"
-#include "NiagaraTypes.h"
-#include "AudioDeviceManager.h"
-#include "AudioResampler.h"
 #include "NiagaraCompileHashVisitor.h"
+#include "NiagaraGeneratedDataAudioSampling.h"
 #include "NiagaraRenderer.h"
-#include "NiagaraDataInterfaceAudio.h"
 #include "NiagaraShaderParametersBuilder.h"
-#include "RenderingThread.h"
+#include "NiagaraSystemInstance.h"
+#include "NiagaraTypes.h"
+#include "NiagaraWorldManager.h"
+#include "RenderGraphBuilder.h"
+#include "RenderGraphUtils.h"
+#include "SystemTextures.h"
+#include "UnifiedBuffer.h"
 
 #include UE_INLINE_GENERATED_CPP_BY_NAME(NiagaraDataInterfaceAudioOscilloscope)
 
 #define LOCTEXT_NAMESPACE "NiagaraDataInterfaceGridAudioOscilloscope"
 
-
-// Global VM function names, also used by the shaders code generation methods.
-static const FName SampleAudioBufferFunctionName("SampleAudioBuffer");
-static const FName GetAudioBufferNumChannelsFunctionName("GetAudioBufferNumChannels");
-
-// Global variable prefixes, used in HLSL parameter declarations.
-static const FString AudioBufferName(TEXT("_AudioBuffer"));
-static const FString NumChannelsName(TEXT("_NumChannels"));
-
-FNiagaraDataInterfaceProxyOscilloscope::FNiagaraDataInterfaceProxyOscilloscope(int32 InResolution, float InScopeInMillseconds)
-	: PatchMixer()
-	, SubmixRegisteredTo(nullptr)
-	, bIsSubmixListenerRegistered(false)
-	, Resolution(InResolution)
-	, ScopeInMilliseconds(InScopeInMillseconds)
-	, NumChannelsInDownsampledBuffer(0)
+namespace UE::Niagara::DataInterfaceAudioOscilloscope::Private
 {
-	DeviceCreatedHandle = FAudioDeviceManagerDelegates::OnAudioDeviceCreated.AddRaw(this, &FNiagaraDataInterfaceProxyOscilloscope::OnNewDeviceCreated);
-	DeviceDestroyedHandle = FAudioDeviceManagerDelegates::OnAudioDeviceDestroyed.AddRaw(this, &FNiagaraDataInterfaceProxyOscilloscope::OnDeviceDestroyed);
+	// Global VM function names, also used by the shaders code generation methods.
+	static const FName SampleAudioBufferFunctionName("SampleAudioBuffer");
+	static const FName GetAudioBufferNumChannelsFunctionName("GetAudioBufferNumChannels");
 
-	VectorVMReadBuffer.Reset();
-	VectorVMReadBuffer.AddZeroed(Resolution * AUDIO_MIXER_MAX_OUTPUT_CHANNELS);
-}
+	const TCHAR* CommonShaderFile = TEXT("/Plugin/FX/Niagara/Private/NiagaraDataInterfaceAudioCommon.ush");
+	const TCHAR* TemplateShaderFile = TEXT("/Plugin/FX/Niagara/Private/NiagaraDataInterfaceAudioOscilloscopeTemplate.ush");
 
-FNiagaraDataInterfaceProxyOscilloscope::~FNiagaraDataInterfaceProxyOscilloscope()
-{
-	check(IsInRenderingThread());
-	GPUDownsampledBuffer.Release();
-}
+	BEGIN_SHADER_PARAMETER_STRUCT(FShaderParameters, )
+		SHADER_PARAMETER(int, ChannelCount)
+		SHADER_PARAMETER(int, Resolution)
+		SHADER_PARAMETER_RDG_BUFFER_SRV(Buffer<float>, OscilloscopeBuffer)
+	END_SHADER_PARAMETER_STRUCT()
 
-void FNiagaraDataInterfaceProxyOscilloscope::RegisterToAllAudioDevices()
-{
-	if (FAudioDeviceManager* DeviceManager = FAudioDeviceManager::Get())
+	struct FInstanceData_GameThread
 	{
-		// Register a new submix listener for every audio device that currently exists.
-		DeviceManager->IterateOverAllDevices([&](Audio::FDeviceId DeviceId, FAudioDevice* InDevice)
-		{
-			check(!SubmixListeners.Contains(DeviceId));
-			const int32 NumSamplesToPop = Align(FMath::FloorToInt(ScopeInMilliseconds / 1000.0f * InDevice->GetSampleRate()) * AUDIO_MIXER_MAX_OUTPUT_CHANNELS, 4);
-			SubmixListeners.Emplace(DeviceId, MakeShared<FNiagaraSubmixListener>(PatchMixer, NumSamplesToPop, DeviceId, SubmixRegisteredTo));
-			SubmixListeners[DeviceId]->RegisterToSubmix();
-		});
-	}
+		FNDIAudio_SharedResourceHandle SharedResource;
 
-	bIsSubmixListenerRegistered = true;
-}
+		TWeakObjectPtr<USoundSubmix> SoundSubmix;
+		TArray<float> DownsampledBuffer;
+		int32 DownsampledResolution = 0;
+		float ScopeInMilliseconds = 0.0f;
+		bool bSoundSubmixAssigned = false;
+	};
 
-void FNiagaraDataInterfaceProxyOscilloscope::UnregisterFromAllAudioDevices(USoundSubmix* Submix)
-{
-	if (FAudioDeviceManager* DeviceManager = FAudioDeviceManager::Get())
+	struct FGameToRenderInstanceData
 	{
-		// Register a new submix listener for every audio device that currently exists.
-		DeviceManager->IterateOverAllDevices([&](Audio::FDeviceId DeviceId, FAudioDevice* InDevice)
-		{
-			check(SubmixListeners.Contains(DeviceId));
-			SubmixListeners.Remove(DeviceId);
-		});
-	}
+		TArray<float> DataToUpload;
 
-	ensure(SubmixListeners.Num() == 0);
-	SubmixListeners.Reset();
-}
+		int32 ChannelCount = 0;
+		int32 ChannelResolution = 0;
+	};
 
-void FNiagaraDataInterfaceProxyOscilloscope::OnUpdateSubmix(USoundSubmix* Submix)
-{
-	if (bIsSubmixListenerRegistered)
+	struct FInstanceData_RenderThread
 	{
-		UnregisterFromAllAudioDevices(Submix);
-	}
-
-	SubmixRegisteredTo = Submix;
+		TArray<float> DataToUpload;
+		TRefCountPtr<FRDGPooledBuffer> PooledBuffer;
 
-	RegisterToAllAudioDevices();
-	
-}
+		int32 ChannelCount = 0;
+		int32 ChannelResolution = 0;
+	};
 
-void FNiagaraDataInterfaceProxyOscilloscope::OnNewDeviceCreated(Audio::FDeviceId InID)
-{
-	if (bIsSubmixListenerRegistered)
+	struct FNDIProxy : public FNiagaraDataInterfaceProxy
 	{
-		check(!SubmixListeners.Contains(InID));
+		virtual int32 PerInstanceDataPassedToRenderThreadSize() const override { return sizeof(FGameToRenderInstanceData); }
 
-		FAudioDevice * DeviceHandle = FAudioDeviceManager::Get()->GetAudioDeviceRaw(InID);
-		if (ensure(DeviceHandle))
+		static void ProvidePerInstanceDataForRenderThread(void* InDataForRenderThread, void* PerInstanceData, const FNiagaraSystemInstanceID& SystemInstance)
 		{
-			const int32 NumSamplesToPop = Align(FMath::FloorToInt(ScopeInMilliseconds / 1000.0f * DeviceHandle->GetSampleRate()) * AUDIO_MIXER_MAX_OUTPUT_CHANNELS, 4);
-			SubmixListeners.Emplace(InID, MakeShared<FNiagaraSubmixListener>(PatchMixer, NumSamplesToPop, InID, SubmixRegisteredTo));
-			SubmixListeners[InID]->RegisterToSubmix();
+			FInstanceData_GameThread* InstanceData = static_cast<FInstanceData_GameThread*>(PerInstanceData);
+			FGameToRenderInstanceData* DataForRenderThread = new(InDataForRenderThread) FGameToRenderInstanceData();
+			const FNDIAudio_SharedResource& AudioResource = InstanceData->SharedResource.ReadResource();
+
+			DataForRenderThread->DataToUpload = AudioResource.ReadAudioBuffer(
+				DataForRenderThread->ChannelCount, DataForRenderThread->ChannelResolution);
 		}
-	}
-}
 
-void FNiagaraDataInterfaceProxyOscilloscope::OnDeviceDestroyed(Audio::FDeviceId InID)
-{
-	if (bIsSubmixListenerRegistered)
-	{
-		if (SubmixListeners.Contains(InID))
+		virtual void ConsumePerInstanceDataFromGameThread(void* PerInstanceData, const FNiagaraSystemInstanceID& InstanceID) override
 		{
-			SubmixListeners.Remove(InID);
+			FGameToRenderInstanceData* InstanceDataFromGT = static_cast<FGameToRenderInstanceData*>(PerInstanceData);
+			FInstanceData_RenderThread& InstanceData = SystemInstancesToInstanceData_RT.FindOrAdd(InstanceID);
+
+			InstanceData.DataToUpload = MoveTemp(InstanceDataFromGT->DataToUpload);
+			InstanceData.ChannelCount = InstanceDataFromGT->ChannelCount;
+			InstanceData.ChannelResolution = InstanceDataFromGT->ChannelResolution;
+
+			InstanceDataFromGT->~FGameToRenderInstanceData();
 		}
+
+		virtual void PreStage(const FNDIGpuComputePreStageContext& Context)
+		{
+			const FNiagaraSystemInstanceID SystemInstanceID = Context.GetSystemInstanceID();
+			FInstanceData_RenderThread& InstanceData = SystemInstancesToInstanceData_RT.FindChecked(SystemInstanceID);
+			if (InstanceData.DataToUpload.Num() > 0)
+			{
+				FRDGBuilder& GraphBuilder = Context.GetGraphBuilder();
+
+				const void* InitialData = InstanceData.DataToUpload.GetData();
+				const uint64 InitialDataSize = InstanceData.DataToUpload.Num() * InstanceData.DataToUpload.GetTypeSize();
+				const uint64 BufferSize = Align(InitialDataSize, 16);
+
+				FRDGBufferDesc BufferDesc = FRDGBufferDesc::CreateByteAddressDesc(IntCastChecked<uint32>(BufferSize));
+				ResizeBufferIfNeeded(GraphBuilder, InstanceData.PooledBuffer, BufferDesc, TEXT("NiagaraAudioOscilloscopeBuffer"));
+
+				GraphBuilder.QueueBufferUpload(
+					GraphBuilder.RegisterExternalBuffer(InstanceData.PooledBuffer),
+					InitialData,
+					InitialDataSize,
+					[Data = MoveTemp(InstanceData.DataToUpload)](const void*) {}
+				);
+
+				InstanceData.DataToUpload.Empty();
+			}
+		}
+
+		TMap<FNiagaraSystemInstanceID, FInstanceData_RenderThread> SystemInstancesToInstanceData_RT;
+	};
+
+	static float SampleInterleavedAudioBuffer(TConstArrayView<float> Samples, int32 ChannelCount, int32 FrameCount, int32 ChannelIndex, float Position)
+	{
+		float FrameIndex = FMath::Clamp(Position, 0.0f, 1.0f) * (FrameCount - 1);
+		int32 LowerIndex = FMath::Max(0, FMath::FloorToInt(FrameIndex));
+		int32 UpperIndex = FMath::Min(LowerIndex + 1, FrameCount - 1);
+
+		float Fraction = FrameIndex - float(LowerIndex);
+
+		float LowerValue = Samples[LowerIndex * ChannelCount + ChannelIndex];
+		float UpperValue = Samples[UpperIndex * ChannelCount + ChannelIndex];
+
+		return FMath::Lerp(LowerValue, UpperValue, Fraction);
 	}
 }
 
 UNiagaraDataInterfaceAudioOscilloscope::UNiagaraDataInterfaceAudioOscilloscope(FObjectInitializer const& ObjectInitializer)
 	: Super(ObjectInitializer)
-	, Submix(nullptr)
-	, Resolution(512)
-	, ScopeInMilliseconds(20.0f)
-	
 {
-	Proxy = TUniquePtr<FNiagaraDataInterfaceProxyOscilloscope>(new FNiagaraDataInterfaceProxyOscilloscope(Resolution, ScopeInMilliseconds));
-}
-
-float FNiagaraDataInterfaceProxyOscilloscope::SampleAudio(float NormalizedPositionInBuffer, int32 ChannelIndex, int32 NumSamplesInBuffer, int32 NumChannelsInBuffer)
-{
-	if (NumSamplesInBuffer == 0 || NumChannelsInBuffer == 0)
-	{
-		return 0.0f;
-	}
-
-	NormalizedPositionInBuffer = FMath::Clamp(NormalizedPositionInBuffer, 0.0f, 1.0f - SMALL_NUMBER);
-	ChannelIndex = FMath::Clamp(ChannelIndex, 0, NumChannelsInBuffer);
-
-	float FrameIndex = NormalizedPositionInBuffer * (NumSamplesInBuffer / NumChannelsInBuffer);
-	int32 LowerFrameIndex = FMath::FloorToInt(FrameIndex);
-	float LowerFrameAmplitude = VectorVMReadBuffer[LowerFrameIndex * NumChannelsInBuffer + ChannelIndex];
-	int32 HigherFrameIndex = LowerFrameIndex < NumSamplesInBuffer ? LowerFrameIndex + 1 : LowerFrameIndex;
-	float HigherFrameAmplitude = VectorVMReadBuffer[HigherFrameIndex * NumChannelsInBuffer + ChannelIndex];
-	float Fraction = HigherFrameIndex - FrameIndex;
-	return FMath::Lerp<float>(LowerFrameAmplitude, HigherFrameAmplitude, Fraction);
+	Proxy.Reset(new UE::Niagara::DataInterfaceAudioOscilloscope::Private::FNDIProxy());
 }
 
 void UNiagaraDataInterfaceAudioOscilloscope::SampleAudio(FVectorVMExternalFunctionContext& Context)
 {
-	const int32 NumSamplesInDownsampledBuffer = GetProxyAs<FNiagaraDataInterfaceProxyOscilloscope>()->DownsampleAudioToBuffer();
+	using namespace UE::Niagara::DataInterfaceAudioOscilloscope::Private;
+
+	const int32 InstanceCount = Context.GetNumInstances();
 
+	VectorVM::FUserPtrHandler<FInstanceData_GameThread> InstData(Context);
 	VectorVM::FExternalFuncInputHandler<float> InNormalizedPos(Context);
 	VectorVM::FExternalFuncInputHandler<int32> InChannel(Context);
 	VectorVM::FExternalFuncRegisterHandler<float> OutAmplitude(Context);
 
-	const int32 NumChannelsInDownsampledBuffer = GetProxyAs<FNiagaraDataInterfaceProxyOscilloscope>()->GetNumChannels();
-
-	for (int32 InstanceIdx = 0; InstanceIdx < Context.GetNumInstances(); ++InstanceIdx)
+	const FNDIAudio_SharedResource& SharedAudio = InstData->SharedResource.ReadResource();
 	{
-		float Position = InNormalizedPos.Get();
-		int32 Channel = InChannel.Get();
-		*OutAmplitude.GetDest() = GetProxyAs<FNiagaraDataInterfaceProxyOscilloscope>()->SampleAudio(Position, Channel, NumSamplesInDownsampledBuffer, NumChannelsInDownsampledBuffer);
-
-		InNormalizedPos.Advance();
-		InChannel.Advance();
-		OutAmplitude.Advance();
+		SharedAudio.WaitForAudio();
 	}
-}
+	int32 ChannelCount = 0;
+	int32 ChannelResolution = 0;
+	TConstArrayView<float> AudioBuffer = SharedAudio.ReadAudioBuffer(ChannelCount, ChannelResolution);
 
-int32 FNiagaraDataInterfaceProxyOscilloscope::GetNumChannels()
-{
-	return NumChannelsInDownsampledBuffer.GetValue();
+	if (ChannelCount && ChannelResolution)
+	{
+		for (int32 InstanceIdx = 0; InstanceIdx < InstanceCount; ++InstanceIdx)
+		{
+			const float Position = InNormalizedPos.GetAndAdvance();
+			const int32 ChannelIndex = InChannel.GetAndAdvance();
+			*OutAmplitude.GetDestAndAdvance() = SampleInterleavedAudioBuffer(
+				AudioBuffer,
+				ChannelCount,
+				ChannelResolution,
+				FMath::Clamp(ChannelIndex, 0, ChannelCount - 1),
+				Position);
+		}
+	}
+	else
+	{
+		for (int32 InstanceIdx = 0; InstanceIdx < Context.GetNumInstances(); ++InstanceIdx)
+		{
+			InNormalizedPos.GetAndAdvance();
+			InChannel.GetAndAdvance();
+			*OutAmplitude.GetDestAndAdvance() = 0.0f;
+		}
+	}
 }
 
 void UNiagaraDataInterfaceAudioOscilloscope::GetNumChannels(FVectorVMExternalFunctionContext& Context)
 {
+	const int32 InstanceCount = Context.GetNumInstances();
+
+	VectorVM::FUserPtrHandler<UE::Niagara::DataInterfaceAudioOscilloscope::Private::FInstanceData_GameThread> InstData(Context);
 	VectorVM::FExternalFuncRegisterHandler<int32> OutChannel(Context);
 
-	for (int32 InstanceIdx = 0; InstanceIdx < Context.GetNumInstances(); ++InstanceIdx)
+	const FNDIAudio_SharedResource& SharedAudio = InstData->SharedResource.ReadResource();
+	const int32 ChannelCount = SharedAudio.GetChannelCount();
+
+	for (int32 InstanceIdx = 0; InstanceIdx < InstanceCount; ++InstanceIdx)
 	{
-		*OutChannel.GetDestAndAdvance() = GetProxyAs<FNiagaraDataInterfaceProxyOscilloscope>()->GetNumChannels();
+		*OutChannel.GetDestAndAdvance() = ChannelCount;
 	}
 }
 
 #if WITH_EDITORONLY_DATA
 void UNiagaraDataInterfaceAudioOscilloscope::GetFunctionsInternal(TArray<FNiagaraFunctionSignature>& OutFunctions) const
 {
+	using namespace UE::Niagara::DataInterfaceAudioOscilloscope::Private;
 	Super::GetFunctionsInternal(OutFunctions);
 
 	{
@@ -220,6 +233,8 @@ DEFINE_NDI_DIRECT_FUNC_BINDER(UNiagaraDataInterfaceAudioOscilloscope, GetNumChan
 
 void UNiagaraDataInterfaceAudioOscilloscope::GetVMExternalFunction(const FVMExternalFunctionBindingInfo& BindingInfo, void* InstanceData, FVMExternalFunction &OutFunc)
 {
+	using namespace UE::Niagara::DataInterfaceAudioOscilloscope::Private;
+
 	if (BindingInfo.Name == SampleAudioBufferFunctionName)
 	{
 		//TNDIParamBinder<0, int32, NDI_FUNC_BINDER(UNiagaraDataInterfaceAudioOscilloscope, SampleAudio)>::Bind(this, BindingInfo, InstanceData, OutFunc);
@@ -239,152 +254,95 @@ void UNiagaraDataInterfaceAudioOscilloscope::GetVMExternalFunction(const FVMExte
 #if WITH_EDITORONLY_DATA
 bool UNiagaraDataInterfaceAudioOscilloscope::AppendCompileHash(FNiagaraCompileHashVisitor* InVisitor) const
 {
-	bool bSuccess = Super::AppendCompileHash(InVisitor);
-	bSuccess &= InVisitor->UpdateShaderParameters<FShaderParameters>();
-	return bSuccess;
+	using namespace UE::Niagara::DataInterfaceAudioOscilloscope::Private;
+
+	if (!Super::AppendCompileHash(InVisitor))
+	{
+		return false;
+	}
+	InVisitor->UpdateShaderFile(CommonShaderFile);
+	InVisitor->UpdateShaderFile(TemplateShaderFile);
+	InVisitor->UpdateShaderParameters<FShaderParameters>();
+	return true;
 }
 
 bool UNiagaraDataInterfaceAudioOscilloscope::GetFunctionHLSL(const FNiagaraDataInterfaceGPUParamInfo& ParamInfo, const FNiagaraDataInterfaceGeneratedFunction& FunctionInfo, int FunctionInstanceIndex, FString& OutHLSL)
 {
-	bool ParentRet = Super::GetFunctionHLSL(ParamInfo, FunctionInfo, FunctionInstanceIndex, OutHLSL);
-	if (ParentRet)
-	{
-		return true;
-	}
-	else if (FunctionInfo.DefinitionName == SampleAudioBufferFunctionName)
+	using namespace UE::Niagara::DataInterfaceAudioOscilloscope::Private;
+
+	if (Super::GetFunctionHLSL(ParamInfo, FunctionInfo, FunctionInstanceIndex, OutHLSL))
 	{
-		// See SampleBuffer(float InNormalizedPosition, int32 InChannel)
-		/**
-		 * float FrameIndex = NormalizedPositionInBuffer * DownsampledBuffer.Num();
-		 * int32 LowerFrameIndex = FMath::FloorToInt(FrameIndex);
-		 * float LowerFrameAmplitude = DownsampledBuffer[LowerFrameIndex * NumChannelsInDownsampledBuffer + ChannelIndex];
-		 * int32 HigherFrameIndex = FMath::CeilToInt(FrameIndex);
-		 * float HigherFrameAmplitude = DownsampledBuffer[HigherFrameIndex * NumChannelsInDownsampledBuffer + ChannelIndex];
-		 * float Fraction = HigherFrameIndex - FrameIndex;
-		 * return FMath::Lerp<float>(LowerFrameAmplitude, HigherFrameAmplitude, Fraction);
-		 */
-		static const TCHAR *FormatBounds = TEXT(R"(
-			void {FunctionName}(float In_NormalizedPosition, int In_ChannelIndex, out float Out_Val)
-			{
-				float FrameIndex = In_NormalizedPosition * {AudioBufferNumSamples} / {ChannelCount};
-				int LowerIndex = floor(FrameIndex);
-				int UpperIndex =  LowerIndex < {AudioBufferNumSamples} ? LowerIndex + 1.0 : LowerIndex;
-				float Fraction = FrameIndex - LowerIndex;
-				float LowerValue = {AudioBuffer}.Load(LowerIndex * {ChannelCount} + In_ChannelIndex);
-				float UpperValue = {AudioBuffer}.Load(UpperIndex * {ChannelCount} + In_ChannelIndex);
-				Out_Val = lerp(LowerValue, UpperValue, Fraction);
-			}
-		)");
-		TMap<FString, FStringFormatArg> ArgsBounds = {
-			{TEXT("FunctionName"), FStringFormatArg(FunctionInfo.InstanceName)},
-			{TEXT("ChannelCount"), FStringFormatArg(ParamInfo.DataInterfaceHLSLSymbol + NumChannelsName)},
-			{TEXT("AudioBufferNumSamples"), FStringFormatArg(Resolution)},
-			{TEXT("AudioBuffer"), FStringFormatArg(ParamInfo.DataInterfaceHLSLSymbol + AudioBufferName)},
-		};
-		OutHLSL += FString::Format(FormatBounds, ArgsBounds);
 		return true;
 	}
-	else if (FunctionInfo.DefinitionName == GetAudioBufferNumChannelsFunctionName)
+
+	if ((FunctionInfo.DefinitionName == SampleAudioBufferFunctionName) ||
+		(FunctionInfo.DefinitionName == GetAudioBufferNumChannelsFunctionName))
 	{
-		// See GetNumChannels()
-		static const TCHAR *FormatBounds = TEXT(R"(
-			void {FunctionName}(out int Out_Val)
-			{
-				Out_Val = {ChannelCount};
-			}
-		)");
-		TMap<FString, FStringFormatArg> ArgsBounds = {
-			{TEXT("FunctionName"), FStringFormatArg(FunctionInfo.InstanceName)},
-			{TEXT("ChannelCount"), FStringFormatArg(ParamInfo.DataInterfaceHLSLSymbol + NumChannelsName)},
-		};
-		OutHLSL += FString::Format(FormatBounds, ArgsBounds);
 		return true;
 	}
-	else
-	{
-		return false;
-	}
+
+	return false;
 }
 
-void UNiagaraDataInterfaceAudioOscilloscope::GetParameterDefinitionHLSL(const FNiagaraDataInterfaceGPUParamInfo& ParamInfo, FString& OutHLSL)
+void UNiagaraDataInterfaceAudioOscilloscope::GetCommonHLSL(FString& OutHLSL)
 {
-	Super::GetParameterDefinitionHLSL(ParamInfo, OutHLSL);
-
-	OutHLSL.Appendf(TEXT("int %s%s;\n"), *ParamInfo.DataInterfaceHLSLSymbol, *NumChannelsName);
-	OutHLSL.Appendf(TEXT("Buffer<float> %s%s;\n"), *ParamInfo.DataInterfaceHLSLSymbol, *AudioBufferName);
+	Super::GetCommonHLSL(OutHLSL);
+	OutHLSL.Appendf(TEXT("#include \"%s\"\n"), UE::Niagara::DataInterfaceAudioOscilloscope::Private::CommonShaderFile);
+}
 
+void UNiagaraDataInterfaceAudioOscilloscope::GetParameterDefinitionHLSL(const FNiagaraDataInterfaceGPUParamInfo& ParamInfo, FString& OutHLSL)
+{
+	const TMap<FString, FStringFormatArg> TemplateArgs =
+	{
+		{TEXT("ParameterName"),	ParamInfo.DataInterfaceHLSLSymbol},
+	};
+	AppendTemplateHLSL(OutHLSL, UE::Niagara::DataInterfaceAudioOscilloscope::Private::TemplateShaderFile, TemplateArgs);
 }
+
 #endif
 
 void UNiagaraDataInterfaceAudioOscilloscope::BuildShaderParameters(FNiagaraShaderParametersBuilder& ShaderParametersBuilder) const
 {
-	ShaderParametersBuilder.AddNestedStruct<FShaderParameters>();
+	ShaderParametersBuilder.AddNestedStruct<UE::Niagara::DataInterfaceAudioOscilloscope::Private::FShaderParameters>();
 }
 
 void UNiagaraDataInterfaceAudioOscilloscope::SetShaderParameters(const FNiagaraDataInterfaceSetShaderParametersContext& Context) const
 {
-	FNiagaraDataInterfaceProxyOscilloscope& DIProxy = Context.GetProxy<FNiagaraDataInterfaceProxyOscilloscope>();
-	FReadBuffer& AudioBufferSRV = DIProxy.ComputeAndPostSRV();
+	using namespace UE::Niagara::DataInterfaceAudioOscilloscope::Private;
+	
+	FNDIProxy& DIProxy = Context.GetProxy<FNDIProxy>();
+	FInstanceData_RenderThread& InstanceData = DIProxy.SystemInstancesToInstanceData_RT.FindChecked(Context.GetSystemInstanceID());
 
-	FShaderParameters* ShaderParameters = Context.GetParameterNestedStruct<FShaderParameters>();
-	ShaderParameters->NumChannels = DIProxy.GetNumChannels();
-	ShaderParameters->AudioBuffer = FNiagaraRenderer::GetSrvOrDefaultFloat(AudioBufferSRV.SRV);
-}
+	FRDGBuilder& GraphBuilder = Context.GetGraphBuilder();
 
-void FNiagaraDataInterfaceProxyOscilloscope::OnUpdateResampling(int32 InResolution, float InScopeInMilliseconds)
-{
-	Resolution = InResolution;
-	ScopeInMilliseconds = InScopeInMilliseconds;
+	FShaderParameters* ShaderParameters = Context.GetParameterNestedStruct<FShaderParameters>();
+	ShaderParameters->ChannelCount = InstanceData.ChannelCount;
+	ShaderParameters->Resolution = InstanceData.ChannelResolution;
 
-	const int32 NumSamplesInBuffer = Resolution * NumChannelsInDownsampledBuffer.GetValue();
-	if (NumSamplesInBuffer)
+	if (InstanceData.PooledBuffer)
 	{
-		ENQUEUE_RENDER_COMMAND(FUpdateDIAudioBuffer)(
-			[&, NumSamplesInBuffer](FRHICommandListImmediate& RHICmdList)
-		{
-			if (GPUDownsampledBuffer.NumBytes > 0)
-			{
-				GPUDownsampledBuffer.Release();
-			}
-
-			GPUDownsampledBuffer.Initialize(RHICmdList, TEXT("FNiagaraDataInterfaceProxyOscilloscope_GPUDownsampledBuffer"), sizeof(float), NumSamplesInBuffer, EPixelFormat::PF_R32_FLOAT, BUF_Static);
-		});
-
-		VectorVMReadBuffer.SetNumZeroed(Resolution * AUDIO_MIXER_MAX_OUTPUT_CHANNELS);
+		FRDGBufferRef RDGBuffer = GraphBuilder.RegisterExternalBuffer(InstanceData.PooledBuffer);
+		ShaderParameters->OscilloscopeBuffer = GraphBuilder.CreateSRV(RDGBuffer, PF_R32_FLOAT);
+	}
+	else
+	{
+		FRDGBufferRef RDGBuffer = GSystemTextures.GetDefaultBuffer(GraphBuilder, sizeof(float), 0.0f);
+		ShaderParameters->OscilloscopeBuffer = GraphBuilder.CreateSRV(RDGBuffer, PF_R32_FLOAT);
 	}
 }
 
-FReadBuffer& FNiagaraDataInterfaceProxyOscilloscope::ComputeAndPostSRV()
+bool UNiagaraDataInterfaceAudioOscilloscope::Equals(const UNiagaraDataInterface* Other) const
 {
-	// Copy to GPUDownsampledBuffer:
-	PostAudioToGPU();
-	return GPUDownsampledBuffer;
-}
+	bool bIsEqual = Super::Equals(Other);
 
-#if WITH_EDITOR
-void UNiagaraDataInterfaceAudioOscilloscope::PostEditChangeProperty(struct FPropertyChangedEvent& PropertyChangedEvent)
-{
-	Super::PostEditChangeProperty(PropertyChangedEvent);
-	
-	static FName SubmixFName = GET_MEMBER_NAME_CHECKED(UNiagaraDataInterfaceAudioOscilloscope, Submix);
-	static FName ResolutionFName = GET_MEMBER_NAME_CHECKED(UNiagaraDataInterfaceAudioOscilloscope, Resolution);
-	static FName ScopeInMillisecondsFName = GET_MEMBER_NAME_CHECKED(UNiagaraDataInterfaceAudioOscilloscope, ScopeInMilliseconds);
+	const UNiagaraDataInterfaceAudioOscilloscope* OtherOscilloscope = CastChecked<const UNiagaraDataInterfaceAudioOscilloscope>(Other);
 
-	// Regenerate on save any compressed sound formats or if analysis needs to be re-done
-	if (FProperty* PropertyThatChanged = PropertyChangedEvent.Property)
-	{
-		const FName& Name = PropertyThatChanged->GetFName();
-		if (Name == SubmixFName)
-		{
-			GetProxyAs<FNiagaraDataInterfaceProxyOscilloscope>()->OnUpdateSubmix(Submix);
-		}
-		else if (Name == ResolutionFName || Name == ScopeInMillisecondsFName)
-		{
-			GetProxyAs<FNiagaraDataInterfaceProxyOscilloscope>()->OnUpdateResampling(Resolution, ScopeInMilliseconds);
-		}
-	}
+	bIsEqual &= OtherOscilloscope->Submix == Submix;
+	bIsEqual &= OtherOscilloscope->Resolution == Resolution;
+	bIsEqual &= OtherOscilloscope->ScopeInMilliseconds == ScopeInMilliseconds;
+
+	return bIsEqual;
 }
-#endif //WITH_EDITOR
 
 void UNiagaraDataInterfaceAudioOscilloscope::PostInitProperties()
 {
@@ -395,167 +353,126 @@ void UNiagaraDataInterfaceAudioOscilloscope::PostInitProperties()
 		ENiagaraTypeRegistryFlags Flags = ENiagaraTypeRegistryFlags::AllowAnyVariable | ENiagaraTypeRegistryFlags::AllowParameter;
 		FNiagaraTypeRegistry::Register(FNiagaraTypeDefinition(GetClass()), Flags);
 	}
-	else if (!HasAnyFlags(RF_NeedLoad) && !HasAnyFlags(RF_Transactional))
-	{
-		GetProxyAs<FNiagaraDataInterfaceProxyOscilloscope>()->OnUpdateResampling(Resolution, ScopeInMilliseconds);
-		GetProxyAs<FNiagaraDataInterfaceProxyOscilloscope>()->OnUpdateSubmix(Submix);
-	}
 }
 
-void FNiagaraDataInterfaceProxyOscilloscope::OnBeginDestroy()
+bool UNiagaraDataInterfaceAudioOscilloscope::PerInstanceTick(void* PerInstanceData, FNiagaraSystemInstance* SystemInstance, float DeltaSeconds)
 {
-	FAudioDeviceManagerDelegates::OnAudioDeviceCreated.Remove(DeviceCreatedHandle);
-	FAudioDeviceManagerDelegates::OnAudioDeviceDestroyed.Remove(DeviceDestroyedHandle);
-	DeviceCreatedHandle.Reset();
-	DeviceDestroyedHandle.Reset();
+	using namespace UE::Niagara::DataInterfaceAudioOscilloscope::Private;
+
+	FInstanceData_GameThread* Inst = reinterpret_cast<FInstanceData_GameThread*>(PerInstanceData);
 
-	if (bIsSubmixListenerRegistered)
+	// reset our per instance data if the submix has changed/invalidated
+	if (Inst->bSoundSubmixAssigned && !Inst->SoundSubmix.IsValid())
 	{
-		UnregisterFromAllAudioDevices(SubmixRegisteredTo);
+		return true;
 	}
-}
-
-void UNiagaraDataInterfaceAudioOscilloscope::BeginDestroy()
-{
-	GetProxyAs<FNiagaraDataInterfaceProxyOscilloscope>()->OnBeginDestroy();
-
-	Super::BeginDestroy();
-}
 
-void UNiagaraDataInterfaceAudioOscilloscope::PostLoad()
-{
-	Super::PostLoad();
-	GetProxyAs<FNiagaraDataInterfaceProxyOscilloscope>()->OnUpdateResampling(Resolution, ScopeInMilliseconds);
-	GetProxyAs<FNiagaraDataInterfaceProxyOscilloscope>()->OnUpdateSubmix(Submix);
-}
+	if (Inst->SoundSubmix != Submix)
+	{
+		return true;
+	}
 
-bool UNiagaraDataInterfaceAudioOscilloscope::Equals(const UNiagaraDataInterface* Other) const
-{
-	const UNiagaraDataInterfaceAudioOscilloscope* CastedOther = Cast<const UNiagaraDataInterfaceAudioOscilloscope>(Other);
-	return Super::Equals(Other)
-		&& (CastedOther->Submix == Submix)
-		&& (CastedOther->Resolution == Resolution)
-		&& FMath::IsNearlyEqual(CastedOther->ScopeInMilliseconds, ScopeInMilliseconds);
+	return false;
 }
 
-bool UNiagaraDataInterfaceAudioOscilloscope::CopyToInternal(UNiagaraDataInterface* Destination) const
+bool UNiagaraDataInterfaceAudioOscilloscope::InitPerInstanceData(void* PerInstanceData, FNiagaraSystemInstance* SystemInstance)
 {
-	Super::CopyToInternal(Destination);
+	using namespace UE::Niagara::DataInterfaceAudioOscilloscope::Private;
+	
+	FInstanceData_GameThread* Inst = new (PerInstanceData) FInstanceData_GameThread();
 
-	UNiagaraDataInterfaceAudioOscilloscope* CastedDestination = Cast<UNiagaraDataInterfaceAudioOscilloscope>(Destination);
+	Inst->SoundSubmix = MakeWeakObjectPtr(Submix);
+	Inst->ScopeInMilliseconds = ScopeInMilliseconds;
+	Inst->DownsampledResolution = Resolution;
+	Inst->bSoundSubmixAssigned = Submix != nullptr;
 
-	if (CastedDestination)
+	if (FNiagaraWorldManager* WorldManager = SystemInstance->GetWorldManager())
 	{
-		CastedDestination->Submix = Submix;
-		CastedDestination->Resolution = Resolution;
-		CastedDestination->ScopeInMilliseconds = ScopeInMilliseconds;
+		FNDIAudio_GeneratedData::FResourceDesc AudioDesc;
+		AudioDesc.DeviceId = INDEX_NONE;
+		AudioDesc.Submix = Submix;
+		AudioDesc.bUseLatestAudio = true;
+		AudioDesc.bResampleAudio = true;
+		AudioDesc.AudioResamplingResolution = Resolution;
+		AudioDesc.SamplingMethod = FNDIAudio_SharedResource::ESamplingWindowMethod::ByTime;
+		AudioDesc.SamplingWindowInMilliseconds = ScopeInMilliseconds;
+
+		if (UWorld* World = WorldManager->GetWorld())
+		{
+			FAudioDeviceHandle AudioDeviceHandle = World->GetAudioDevice();
+			if (AudioDeviceHandle.IsValid())
+			{
+				AudioDesc.DeviceId = AudioDeviceHandle.GetDeviceID();
+			}
+		}
 
-		CastedDestination->GetProxyAs<FNiagaraDataInterfaceProxyOscilloscope>()->OnUpdateResampling(Resolution, ScopeInMilliseconds);
-		CastedDestination->GetProxyAs<FNiagaraDataInterfaceProxyOscilloscope>()->OnUpdateSubmix(Submix);
+		FNDIAudio_SharedResourceUsage ResourceUsage;
+		{
+			bool bInstanceRequiresCpu = false;
+			bool bInstanceRequiresGpu = false;
+			SystemInstance->EvaluateBoundFunction(SampleAudioBufferFunctionName, bInstanceRequiresCpu, bInstanceRequiresGpu);
+			FNDIAudio_SharedResourceUsage::UpdateUsageFlags(bInstanceRequiresCpu, bInstanceRequiresGpu, ResourceUsage.AudioUsage);
+		}
+
+		FNDIAudio_GeneratedData& SpectrumGeneratedData = WorldManager->EditGeneratedData<FNDIAudio_GeneratedData>();
+		Inst->SharedResource = SpectrumGeneratedData.GetSharedResource(ResourceUsage, AudioDesc);
 	}
-	
+
 	return true;
 }
 
-void FNiagaraDataInterfaceProxyOscilloscope::PostAudioToGPU()
+void UNiagaraDataInterfaceAudioOscilloscope::DestroyPerInstanceData(void* PerInstanceData, FNiagaraSystemInstance* SystemInstance)
 {
-	ENQUEUE_RENDER_COMMAND(FUpdateDIAudioBuffer)(
-		[&](FRHICommandListImmediate& RHICmdList)
-	{
-		DownsampleAudioToBuffer();
-		size_t BufferSize = DownsampledBuffer.Num() * sizeof(float);
-		if (BufferSize != 0 && !GPUDownsampledBuffer.NumBytes)
-		{
-			GPUDownsampledBuffer.Initialize(RHICmdList, TEXT("GPUDownsampledBuffer"), sizeof(float), Resolution * NumChannelsInDownsampledBuffer.GetValue(), EPixelFormat::PF_R32_FLOAT, BUF_Static);
-		}
+	using namespace UE::Niagara::DataInterfaceAudioOscilloscope::Private;
 
-		if (GPUDownsampledBuffer.NumBytes > 0)
-		{
-			float *BufferData = static_cast<float*>(RHICmdList.LockBuffer(GPUDownsampledBuffer.Buffer, 0, BufferSize, EResourceLockMode::RLM_WriteOnly));
-			FScopeLock ScopeLock(&DownsampleBufferLock);
-			FPlatformMemory::Memcpy(BufferData, DownsampledBuffer.GetData(), BufferSize);
-			RHICmdList.UnlockBuffer(GPUDownsampledBuffer.Buffer);
-		}
-	});
+	FInstanceData_GameThread* InstData = reinterpret_cast<FInstanceData_GameThread*>(PerInstanceData);
+
+	InstData->~FInstanceData_GameThread();
 }
 
-int32 FNiagaraDataInterfaceProxyOscilloscope::DownsampleAudioToBuffer()
+
+int32 UNiagaraDataInterfaceAudioOscilloscope::PerInstanceDataSize() const
 {
-	FScopeLock ScopeLock(&DownsampleBufferLock);
+	return sizeof(UE::Niagara::DataInterfaceAudioOscilloscope::Private::FInstanceData_GameThread);
+}
 
-	// Get the number of channels from the first listener we find.
-	// If the num channels hasn't been set, either SubmixListeners is empty
-	// or we haven't pushed audio to any of them yet.
-	float SourceSampleRate = 0.0f;
-	int32 NumChannels = 0;
+void UNiagaraDataInterfaceAudioOscilloscope::ProvidePerInstanceDataForRenderThread(void* DataForRenderThread, void* PerInstanceData, const FNiagaraSystemInstanceID& SystemInstance)
+{
+	UE::Niagara::DataInterfaceAudioOscilloscope::Private::FNDIProxy::ProvidePerInstanceDataForRenderThread(DataForRenderThread, PerInstanceData, SystemInstance);
+}
 
-	for (auto& Listener : SubmixListeners)
-	{
-		NumChannels = Listener.Value->GetNumChannels();
-		SourceSampleRate = Listener.Value->GetSampleRate();
+ETickingGroup UNiagaraDataInterfaceAudioOscilloscope::CalculateTickGroup(const void* PerInstanceData) const
+{
+	using namespace UE::Niagara::DataInterfaceAudioOscilloscope::Private;
 
-		if (NumChannels != 0)
+	if (const FInstanceData_GameThread* Inst = reinterpret_cast<const FInstanceData_GameThread*>(PerInstanceData))
+	{
+		if (Inst->SharedResource.Usage.IsValid())
 		{
-			break;
+			return static_cast<ETickingGroup>(FNDIAudio_GeneratedData::GeneratedDataTickGroup + 1);
 		}
 	}
 
-	if (NumChannels == 0 || FMath::IsNearlyZero(SourceSampleRate))
-	{
-		return 0;
-	}
-
-	check(NumChannels > 0);
-	NumChannelsInDownsampledBuffer.Set(NumChannels);
-
-	// Get the number of frames of audio in the original sample rate to show.
-	int32 NumFramesToPop = Align(FMath::FloorToInt(SourceSampleRate * (ScopeInMilliseconds / 1000.0f)), 4);
-
-	// If we have enough frames buffered to pop, update the PopBuffer. Otherwise, use the previous frames buffer.
-	int32 NumSamplesToPop = NumFramesToPop * NumChannels;
+	return NiagaraFirstTickGroup;
+}
 
-	if (PopBuffer.Num() != NumSamplesToPop)
+bool UNiagaraDataInterfaceAudioOscilloscope::CopyToInternal(UNiagaraDataInterface* Destination) const
+{
+	if (!Super::CopyToInternal(Destination))
 	{
-		PopBuffer.Reset();
-		PopBuffer.AddZeroed(NumSamplesToPop);
+		return false;
 	}
 
-	if (PatchMixer.MaxNumberOfSamplesThatCanBePopped() >= NumSamplesToPop)
+	UNiagaraDataInterfaceAudioOscilloscope* CastedDestination = Cast<UNiagaraDataInterfaceAudioOscilloscope>(Destination);
+
+	if (CastedDestination)
 	{
-		PatchMixer.PopAudio(PopBuffer.GetData(), NumSamplesToPop, true);
+		CastedDestination->Submix = Submix;
+		CastedDestination->Resolution = Resolution;
+		CastedDestination->ScopeInMilliseconds = ScopeInMilliseconds;
 	}
-
-	// Downsample buffer in place.
-	float SampleRateRatio = ((float)Resolution) / ((float)NumFramesToPop);
-	float DestinationSampleRate = SourceSampleRate * SampleRateRatio;
-
-	Audio::FResamplingParameters ResampleParameters = {
-		Audio::EResamplingMethod::Linear,
-		NumChannels,
-		SourceSampleRate,
-		DestinationSampleRate,
-		PopBuffer
-	};
-
-	int32 DownsampleBufferSize = Audio::GetOutputBufferSize(ResampleParameters);
-	DownsampledBuffer.Reset();
-	DownsampledBuffer.AddZeroed(DownsampleBufferSize);
-
-
-	Audio::FResamplerResults ResampleResults;
-	ResampleResults.OutBuffer = &DownsampledBuffer;
-
-	bool bResampleResult = Audio::Resample(ResampleParameters, ResampleResults);
-	check(bResampleResult);
-
-	// here we may have an extra sample or two due to roundoff.
-	DownsampledBuffer.SetNumZeroed(Resolution * NumChannels);
-
-	check(DownsampledBuffer.Num() <= VectorVMReadBuffer.Num());
-	FMemory::Memcpy(VectorVMReadBuffer.GetData(), DownsampledBuffer.GetData(), VectorVMReadBuffer.Num());
-
-	return Resolution;
+	
+	return true;
 }
 
 #undef LOCTEXT_NAMESPACE
-
diff --git a/Engine/Plugins/FX/Niagara/Source/Niagara/Private/NiagaraDataInterfaceAudioSpectrum.cpp b/Engine/Plugins/FX/Niagara/Source/Niagara/Private/NiagaraDataInterfaceAudioSpectrum.cpp
--- a/Engine/Plugins/FX/Niagara/Source/Niagara/Private/NiagaraDataInterfaceAudioSpectrum.cpp
+++ b/Engine/Plugins/FX/Niagara/Source/Niagara/Private/NiagaraDataInterfaceAudioSpectrum.cpp
@@ -1,609 +1,229 @@
 // Copyright Epic Games, Inc. All Rights Reserved.
 
 #include "NiagaraDataInterfaceAudioSpectrum.h"
-#include "DSP/AudioFFT.h"
-#include "NiagaraTypes.h"
 #include "NiagaraCompileHashVisitor.h"
+#include "NiagaraGeneratedDataAudioSampling.h"
 #include "NiagaraRenderer.h"
 #include "NiagaraShaderParametersBuilder.h"
-
-#include "DSP/DeinterleaveView.h"
-#include "DSP/SlidingWindow.h"
-#include "DSP/ConstantQ.h"
-#include "DSP/FFTAlgorithm.h"
-#include "DSP/FloatArrayMath.h"
-#include "RenderingThread.h"
+#include "NiagaraSystemInstance.h"
+#include "NiagaraTypes.h"
+#include "NiagaraWorldManager.h"
+#include "RenderGraphBuilder.h"
+#include "RenderGraphUtils.h"
+#include "SystemTextures.h"
+#include "UnifiedBuffer.h"
 
 #include UE_INLINE_GENERATED_CPP_BY_NAME(NiagaraDataInterfaceAudioSpectrum)
 
 #define LOCTEXT_NAMESPACE "NiagaraDataInterfaceGridAudioSpectrum"
 
-FNiagaraDataInterfaceProxySpectrum::FNiagaraDataInterfaceProxySpectrum(float InMinimumFrequency, float InMaximumFrequency, int32 InNumBands, float InNoiseFloorDb, int32 InNumSamplesToBuffer)
-
-	: FNiagaraDataInterfaceProxySubmix(InNumSamplesToBuffer)
-	, MinimumFrequency(0.f)
-	, MaximumFrequency(0.f)
-	, NumBands(0)
-	, NoiseFloorDb(0.f)
-	, NumChannels(0)
-	, SampleRate(0.f)
-	, FFTScale(1.f)
-{
-	int32 InitialNumChannels = 2;
-	float InitialSampleRate = 48000.f;
-
-	SetAudioFormat(InitialNumChannels, InitialSampleRate);
-
-	ResizeCQT(InMinimumFrequency, InMaximumFrequency, InNumBands);
-
-	SetNoiseFloor(InNoiseFloorDb);
-}
-
-FNiagaraDataInterfaceProxySpectrum::~FNiagaraDataInterfaceProxySpectrum()
+namespace UE::Niagara::DataInterfaceAudioSpectrum::Private
 {
-	check(IsInRenderingThread());
-	GPUBuffer.Release();
-}
-
-float FNiagaraDataInterfaceProxySpectrum::GetSpectrumValue(float InNormalizedPositionInSpectrum, int32 InChannelIndex)
-{
-	check(InChannelIndex >= 0);
-
-	// Check if channel index is within channel spectrum buffers
-	if (InChannelIndex >= ChannelSpectrumBuffers.Num())
-	{
-		return 0.f;
-	}
-
-	const Audio::FAlignedFloatBuffer& Buffer = ChannelSpectrumBuffers[InChannelIndex];
-	const int32 MaxIndex = Buffer.Num() - 1;
-
-	if (MaxIndex < 0)
-	{
-		return 0.0f;
-	}
+	// Global VM function names, also used by the shaders code generation methods.
+	static const FName AudioSpectrumFunctionName("AudioSpectrum");
+	static const FName GetNumChannelsFunctionName("GetNumChannels");
 
-	InNormalizedPositionInSpectrum = FMath::Clamp(InNormalizedPositionInSpectrum, 0.0f, 1.0f);
+	const TCHAR* CommonShaderFile = TEXT("/Plugin/FX/Niagara/Private/NiagaraDataInterfaceAudioCommon.ush");
+	const TCHAR* TemplateShaderFile = TEXT("/Plugin/FX/Niagara/Private/NiagaraDataInterfaceAudioSpectrumTemplate.ush");
 
-	float Index = InNormalizedPositionInSpectrum * Buffer.Num();
+	BEGIN_SHADER_PARAMETER_STRUCT(FShaderParameters, )
+		SHADER_PARAMETER(int, ChannelCount)
+		SHADER_PARAMETER(int, Resolution)
+		SHADER_PARAMETER_RDG_BUFFER_SRV(Buffer<float>, SpectrumBuffer)
+	END_SHADER_PARAMETER_STRUCT()
 
-	int32 LowerIndex = FMath::FloorToInt(Index);
-	LowerIndex = FMath::Clamp(LowerIndex, 0, MaxIndex);
+	constexpr int32 DefaultSamplesToBufferCount = 16384;
 
-	int32 UpperIndex = FMath::CeilToInt(Index);
-	UpperIndex = FMath::Clamp(UpperIndex, 0, MaxIndex);
-
-	float Alpha = Index - LowerIndex;
-
-	return FMath::Lerp<float>(Buffer[LowerIndex], Buffer[UpperIndex], Alpha);
-}
-
-int32 FNiagaraDataInterfaceProxySpectrum::GetNumBands() const
-{
-	check(IsInRenderingThread());
-	return NumBands;
-}
-
-
-void FNiagaraDataInterfaceProxySpectrum::UpdateCQT(float InMinimumFrequency, float InMaximumFrequency, int32 InNumBands)
-{
-	ENQUEUE_RENDER_COMMAND(FUpdateSpectrumBuffer)(
-		[&, InMinimumFrequency, InMaximumFrequency, InNumBands](FRHICommandListImmediate& RHICmdList)
+	struct FInstanceData_GameThread
 	{
-		ResizeCQT(InMinimumFrequency, InMaximumFrequency, InNumBands);
-	});
-}
+		FNDIAudio_SharedResourceHandle SharedResource;
 
-void FNiagaraDataInterfaceProxySpectrum::ResizeCQT(float InMinimumFrequency, float InMaximumFrequency, int32 InNumBands)
-{
-	Clamp(InMinimumFrequency, InMaximumFrequency, InNumBands);
-
-	// Check which values were changed
-	bool bUpdatedMinimumFrequency = InMinimumFrequency != MinimumFrequency;
-	bool bUpdatedMaximumFrequency = InMaximumFrequency != MaximumFrequency;
-	bool bUpdatedNumBands = InNumBands != NumBands;
-
-	bool bUpdatedValues = bUpdatedMinimumFrequency || bUpdatedMaximumFrequency || bUpdatedNumBands;
+		// if null, we should be using the default submix from each of the devices
+		TWeakObjectPtr<USoundSubmix> SoundSubmix;
+		int32 Resolution = 0;
+		float MinimumFrequency = 0.0f;
+		float MaximumFrequency = 0.0f;
+		float NoiseFloorDb = 0.0f;
 
-	if (bUpdatedNumBands)
-	{
-		ResizeSpectrumBuffer(NumChannels, InNumBands);
-	}
+		bool bSoundSubmixAssigned = false;
+	};
 
-	if (bUpdatedValues)
+	struct FGameToRenderInstanceData
 	{
-		int32 OriginalFFTSize = FFTAlgorithm.IsValid() ? FFTAlgorithm->Size() : 0;
-
-		ResizeAudioTransform(InMinimumFrequency, InMaximumFrequency, SampleRate, InNumBands);
-
-		int32 NewFFTSize = FFTAlgorithm.IsValid() ? FFTAlgorithm->Size() : 0;
+		TArray<float> DataToUpload;
 
-		if (NewFFTSize != OriginalFFTSize)
-		{
-			ResizeWindow(NumChannels, NewFFTSize);
-		}
-
-	}
-
-	// Copy values internally
-	MinimumFrequency = InMinimumFrequency;
-	MaximumFrequency = InMaximumFrequency;
-	NumBands = InNumBands;
-}
+		int32 ChannelCount = 0;
+		int32 ChannelResolution = 0;
+	};
 
-void FNiagaraDataInterfaceProxySpectrum::UpdateNoiseFloor(float InNoiseFloorDb)
-{
-	ENQUEUE_RENDER_COMMAND(FUpdateSpectrumBuffer)(
-		[&, InNoiseFloorDb](FRHICommandListImmediate& RHICmdList)
+	struct FInstanceData_RenderThread
 	{
-		SetNoiseFloor(InNoiseFloorDb);
-	});
-}
-
-
-void FNiagaraDataInterfaceProxySpectrum::SetNoiseFloor(float InNoiseFloorDb)
-{
-	NoiseFloorDb = FMath::Clamp(InNoiseFloorDb, -200.f, 10.f);
-}
-
-FReadBuffer& FNiagaraDataInterfaceProxySpectrum::ComputeAndPostSRV()
-{
-	PostDataToGPU();
-	return GPUBuffer;
-}
+		TArray<float> DataToUpload;
+		TRefCountPtr<FRDGPooledBuffer> PooledBuffer;
 
+		int32 ChannelCount = 0;
+		int32 ChannelResolution = 0;
+	};
 
-void FNiagaraDataInterfaceProxySpectrum::PostDataToGPU()
-{
-	ENQUEUE_RENDER_COMMAND(FUpdateSpectrumBuffer)(
-		[&](FRHICommandListImmediate& RHICmdList)
+	struct FNDIProxy : public FNiagaraDataInterfaceProxy
 	{
-		UpdateSpectrum();
-
-		// Resize GPU data if needed
-		const int32 NumSamplesInBuffer = NumBands * NumChannels;
-		const int32 NumBytesInChannelBuffer = NumBands * sizeof(float);
-		const int32 NumBytesInBuffer = NumBytesInChannelBuffer * NumChannels;
-		
-		if (NumBytesInBuffer != GPUBuffer.NumBytes)
-		{
-			if (GPUBuffer.NumBytes > 0)
-			{
-				GPUBuffer.Release();
-			}
+		virtual int32 PerInstanceDataPassedToRenderThreadSize() const override { return sizeof(FGameToRenderInstanceData); }
 
-			if (NumSamplesInBuffer > 0)
-			{
-				GPUBuffer.Initialize(RHICmdList, TEXT("FNiagaraDataInterfaceProxySpectrum_GPUBuffer"), sizeof(float), NumSamplesInBuffer, EPixelFormat::PF_R32_FLOAT, BUF_Static);
-			}
-		}
-
-		// Copy to GPU data
-		if (GPUBuffer.NumBytes > 0)
+		static void ProvidePerInstanceDataForRenderThread(void* InDataForRenderThread, void* PerInstanceData, const FNiagaraSystemInstanceID& SystemInstance)
 		{
-			float *BufferData = static_cast<float*>(RHICmdList.LockBuffer(GPUBuffer.Buffer, 0, NumBytesInBuffer, EResourceLockMode::RLM_WriteOnly));
+			FInstanceData_GameThread* InstanceData = static_cast<FInstanceData_GameThread*>(PerInstanceData);
+			FGameToRenderInstanceData* DataForRenderThread = new(InDataForRenderThread) FGameToRenderInstanceData();
 
-			FScopeLock ScopeLock(&BufferLock);
-
-			for (int32 ChannelIndex = 0; ChannelIndex < NumChannels; ChannelIndex++)
+			if (InstanceData->SharedResource)
 			{
-				
-				int32 Pos = ChannelIndex * NumBands;
-				FPlatformMemory::Memcpy(&BufferData[Pos], ChannelSpectrumBuffers[ChannelIndex].GetData(), NumBytesInChannelBuffer);
-			}
-
-			RHICmdList.UnlockBuffer(GPUBuffer.Buffer);
-		}
-	});
-}
-
-void FNiagaraDataInterfaceProxySpectrum::UpdateSpectrum()
-{
-	FScopeLock ScopeLock(&BufferLock);
-
-	int32 SourceNumChannels = GetNumChannels();
-	float SourceSampleRate = GetSampleRate();
-
-	// Perform internal update if new samplerate or channel count encountered.
-	if ((SourceNumChannels != NumChannels) || (SourceSampleRate != SampleRate))
-	{
-		// This will update values of NumChannels and SampleRate member variables.
-		SetAudioFormat(SourceNumChannels, SourceSampleRate);
-	}
-
-	// Set data to zero if we are skipping update due to bad samplerate or channel count.
-	if (NumChannels == 0 || FMath::IsNearlyZero(SampleRate))
-	{
-		for (int32 ChannelIndex = 0; ChannelIndex < ChannelSpectrumBuffers.Num(); ChannelIndex++)
-		{
-			Audio::FAlignedFloatBuffer& Buffer = ChannelSpectrumBuffers[ChannelIndex];
-			if (Buffer.Num() > 0)
-			{
-				FMemory::Memset(Buffer.GetData(), 0, sizeof(float) * Buffer.Num());
-			}
-		}
-		return;
-	}
-
-	
-	// Grab audio from buffer.
-	int32 NumSamplesToPop = GetNumSamplesAvailable();
-	NumSamplesToPop -= (NumSamplesToPop % NumChannels);
-	PopBuffer.Reset();
-	
-	if (NumSamplesToPop > 0)
-	{
-		
-		PopBuffer.AddUninitialized(NumSamplesToPop);
-
-		bool bUseLatestAudio = false;
-		int32 NumPopped = PopAudio(PopBuffer.GetData(), NumSamplesToPop, bUseLatestAudio);
-	}
-	
-	// If we're in a bad internal state, give up here.	
-	if (!FFTAlgorithm.IsValid() || !CQTKernel.IsValid())
-	{
-		return;
-	}
+				const FNDIAudio_SharedResource& AudioResource = InstanceData->SharedResource.ReadResource();
 
-	// Run sliding window over available audio.
-	FSlidingWindow SlidingWindow(*SlidingBuffer, PopBuffer, InterleavedBuffer);
+				// we need to be sure that the CPU tasks have been finished
+				AudioResource.WaitForAudio();
 
-	int32 NumWindows = 0;
-	for (Audio::FAlignedFloatBuffer& InterleavedWindow : SlidingWindow)
-	{
-		if (0 == NumWindows)
-		{
-			// Zero output on first window.  We need to wait for at least one window to process
-			// because we do not want to zero out data if no windows are available.
-			for (int32 ChannelIndex = 0; ChannelIndex < NumChannels; ChannelIndex++)
-			{
-				FMemory::Memset(ChannelSpectrumBuffers[ChannelIndex].GetData(), 0, sizeof(float) * NumBands);
+				DataForRenderThread->DataToUpload = AudioResource.ReadSpectrumBuffer(
+					DataForRenderThread->ChannelCount, DataForRenderThread->ChannelResolution);
 			}
 		}
 
-		NumWindows++;
-
-		// Run deinterleave view over window.
-		FDeinterleaveView DeinterleaveView(InterleavedWindow, DeinterleavedBuffer, NumChannels);
-		for(FChannel Channel : DeinterleaveView)
+		virtual void ConsumePerInstanceDataFromGameThread(void* PerInstanceData, const FNiagaraSystemInstanceID& InstanceID) override
 		{
-			Audio::ArrayMultiplyInPlace(WindowBuffer, Channel.Values);
+			FGameToRenderInstanceData* InstanceDataFromGT = static_cast<FGameToRenderInstanceData*>(PerInstanceData);
+			FInstanceData_RenderThread& InstanceData = SystemInstancesToInstanceData_RT.FindOrAdd(InstanceID);
 
-			// Perform FFT
-			FMemory::Memcpy(FFTInputBuffer.GetData(), Channel.Values.GetData(), sizeof(float) * Channel.Values.Num());
+			InstanceData.DataToUpload = MoveTemp(InstanceDataFromGT->DataToUpload);
+			InstanceData.ChannelCount = InstanceDataFromGT->ChannelCount;
+			InstanceData.ChannelResolution = InstanceDataFromGT->ChannelResolution;
 
-			FFTAlgorithm->ForwardRealToComplex(FFTInputBuffer.GetData(), FFTOutputBuffer.GetData());
-
-			// Transflate FFTOutput to power spectrum
-			Audio::ArrayComplexToPower(FFTOutputBuffer, PowerSpectrumBuffer);
-
-			// Take CQT of power spectrum
-			CQTKernel->TransformArray(PowerSpectrumBuffer, SpectrumBuffer);
-
-			// Accumulate power spectrum CQT output.
-			Audio::ArrayAddInPlace(SpectrumBuffer, ChannelSpectrumBuffers[Channel.ChannelIndex]);
+			InstanceDataFromGT->~FGameToRenderInstanceData();
 		}
-	}
-
-	if (NumWindows > 0)
-	{
-		// Two scalings are applied. One for FFT, another for number of windows calculated.
-		float LinearScale = FFTScale / static_cast<float>(NumWindows);
-
-		// This scaling moves output to roughly 0 to 1 range.
-		float DbScale = 1.f / FMath::Max(1.f, -NoiseFloorDb);
 
-		// Apply scaling for each channel.
-		for (int32 ChannelIndex = 0; ChannelIndex < NumChannels; ChannelIndex++)
+		virtual void PreStage(const FNDIGpuComputePreStageContext& Context)
 		{
-			Audio::FAlignedFloatBuffer& Buffer = ChannelSpectrumBuffers[ChannelIndex];
-
-			Audio::ArrayMultiplyByConstantInPlace(Buffer, LinearScale);
-
-			Audio::ArrayPowerToDecibelInPlace(Buffer, NoiseFloorDb);
-
-			Audio::ArraySubtractByConstantInPlace(Buffer, NoiseFloorDb);
-
-			Audio::ArrayMultiplyByConstantInPlace(Buffer, DbScale);
-		}
-	}
-}
-
-void FNiagaraDataInterfaceProxySpectrum::SetAudioFormat(int32 InNumChannels, float InSampleRate)
-{
-	if (InSampleRate != SampleRate)
-	{
-		ResizeAudioTransform(MinimumFrequency, MaximumFrequency, InSampleRate, NumBands);
-	}
-
-	if (InNumChannels != NumChannels)
-	{
-		ResizeSpectrumBuffer(InNumChannels, NumBands);
-	}
+			const FNiagaraSystemInstanceID SystemInstanceID = Context.GetSystemInstanceID();
+			FInstanceData_RenderThread& InstanceData = SystemInstancesToInstanceData_RT.FindChecked(SystemInstanceID);
+			if (InstanceData.DataToUpload.Num() > 0)
+			{
+				FRDGBuilder& GraphBuilder = Context.GetGraphBuilder();
 
-	if ((InSampleRate != SampleRate) || (InNumChannels != NumChannels))
-	{
-		int32 FFTSize = FFTAlgorithm.IsValid() ? FFTAlgorithm->Size() : 0;
+				const void* InitialData = InstanceData.DataToUpload.GetData();
+				const uint64 InitialDataSize = InstanceData.DataToUpload.Num() * InstanceData.DataToUpload.GetTypeSize();
+				const uint64 BufferSize = Align(InitialDataSize, 16);
 
-		ResizeWindow(InNumChannels, FFTSize);
-	}
+				FRDGBufferDesc BufferDesc = FRDGBufferDesc::CreateByteAddressDesc(IntCastChecked<uint32>(BufferSize));
+				ResizeBufferIfNeeded(GraphBuilder, InstanceData.PooledBuffer, BufferDesc, TEXT("NiagaraAudioSpectrumBuffer"));
 
-	NumChannels = InNumChannels;
-	SampleRate = InSampleRate;
-}
+				GraphBuilder.QueueBufferUpload(
+					GraphBuilder.RegisterExternalBuffer(InstanceData.PooledBuffer),
+					InitialData,
+					InitialDataSize,
+					[Data = MoveTemp(InstanceData.DataToUpload)](const void*) {}
+				);
 
-void FNiagaraDataInterfaceProxySpectrum::ResizeSpectrumBuffer(int32 InNumChannels, int32 InNumBands)
-{
-	for (int32 ChannelIndex = ChannelSpectrumBuffers.Num(); ChannelIndex < InNumChannels; ChannelIndex++)
-	{
-		ChannelSpectrumBuffers.AddDefaulted();
-	}
-
-	for (int32 ChannelIndex = 0; ChannelIndex < InNumChannels; ChannelIndex++)
-	{
-		Audio::FAlignedFloatBuffer& Buffer = ChannelSpectrumBuffers[ChannelIndex];
-		Buffer.Reset();
-		if (InNumBands > 0)
-		{
-			Buffer.AddUninitialized(InNumBands);
+				InstanceData.DataToUpload.Empty();
+			}
 		}
-	}
-}
-
-void FNiagaraDataInterfaceProxySpectrum::ResizeWindow(int32 InNumChannels, int32 InFFTSize)
-{
-	int32 NumWindowFrames = FMath::Max(1, InFFTSize);
-
-	NumWindowFrames = FMath::Min(NumWindowFrames, 1024);
-
-	WindowBuffer.Reset();
-	WindowBuffer.AddUninitialized(NumWindowFrames);
-
-	Audio::GenerateBlackmanWindow(WindowBuffer.GetData(), NumWindowFrames, 1, true);
-
-	int32 NumWindowSamples = NumWindowFrames;
-
-	if (InNumChannels > 0)
-	{
-		NumWindowSamples *= InNumChannels;
-	}
-
-	// 50% overlap
-	int32 NumHopSamples = FMath::Max(1, NumWindowSamples / 2);
-
-	SlidingBuffer = MakeUnique<FSlidingBuffer>(NumWindowSamples, NumHopSamples);
-
-}
-
-void FNiagaraDataInterfaceProxySpectrum::ResizeAudioTransform(float InMinimumFrequency, float InMaximumFrequency, float InSampleRate, int32 InNumBands)
-{
-	FFTAlgorithm.Reset();
-	CQTKernel.Reset();
-	FFTInputBuffer.Reset();
-	FFTOutputBuffer.Reset();
-	PowerSpectrumBuffer.Reset();
-	SpectrumBuffer.Reset();
-
-	FFTScale = 1.f;
-
-	if (InSampleRate <= 0.f)
-	{
-		return;
-	}
-
-	if (InNumBands < 1)
-	{
-		return;
-	}
-
-	float NumOctaves = GetNumOctaves(InMinimumFrequency, InMaximumFrequency);
-	float NumBandsPerOctave = GetNumBandsPerOctave(InNumBands, NumOctaves);
-
-	Audio::FFFTSettings FFTSettings = GetFFTSettings(InMinimumFrequency, InSampleRate, NumBandsPerOctave);
-	if (!Audio::FFFTFactory::AreFFTSettingsSupported(FFTSettings))
-	{
-		UE_LOG(LogNiagara, Warning, TEXT("FFT settings are not supported"));
-		return;
-	}
-
-	FFTAlgorithm = Audio::FFFTFactory::NewFFTAlgorithm(FFTSettings);
-
-	if (!FFTAlgorithm.IsValid())
-	{
-		UE_LOG(LogNiagara, Error, TEXT("Failed to create FFT"));
-		return;
-	}
 
-	FFTInputBuffer.AddZeroed(FFTAlgorithm->NumInputFloats());
-	FFTOutputBuffer.AddUninitialized(FFTAlgorithm->NumOutputFloats());
-	PowerSpectrumBuffer.AddUninitialized(FFTOutputBuffer.Num() / 2);
-	SpectrumBuffer.AddUninitialized(InNumBands);
-
-	float FFTSize = static_cast<float>(FFTAlgorithm->Size());
-	check(FFTSize > 0.f);
+		TMap<FNiagaraSystemInstanceID, FInstanceData_RenderThread> SystemInstancesToInstanceData_RT;
+	};
 
-	// We want to have all FFT implementations
-	// return the same scaling so that the energy conservatino property of the fourier transform
-	// is supported.  This scaling factor is applied to the power spectrum, so we square the 
-	// scaling we would have performed on the magnitude spectrum.
-	switch (FFTAlgorithm->ForwardScaling())
+	static float SamplePlanarAudioBuffer(TConstArrayView<float> Samples, int32 ChannelCount, int32 FrameCount, int32 ChannelIndex, float Position)
 	{
-		case Audio::EFFTScaling::MultipliedByFFTSize:
-			FFTScale = 1.f / (FFTSize * FFTSize);
-			break;
-			
-		case Audio::EFFTScaling::MultipliedBySqrtFFTSize:
-			FFTScale = 1.f / FFTSize;
-			break;
-
-		case Audio::EFFTScaling::DividedByFFTSize:
-			FFTScale = FFTSize * FFTSize;
-			break;
-
-		case Audio::EFFTScaling::DividedBySqrtFFTSize:
-			FFTScale = FFTSize;
-			break;
-
-		default:
-			FFTScale = 1.f;
-			break;
-	}
+		float FrameIndex = FMath::Clamp(Position, 0.0f, 1.0f) * (FrameCount - 1);
+		int32 LowerIndex = FMath::Max(0, FMath::FloorToInt(FrameIndex));
+		int32 UpperIndex = FMath::Min(LowerIndex + 1, FrameCount - 1);
 
-	float BandwidthStretch = GetBandwidthStretch(InSampleRate, FFTSize, NumBandsPerOctave, InMinimumFrequency);
+		float Fraction = FrameIndex - float(LowerIndex);
 
-	Audio::FPseudoConstantQKernelSettings CQTSettings = GetConstantQSettings(InMinimumFrequency, InMaximumFrequency, InNumBands, NumBandsPerOctave, BandwidthStretch);
-	CQTKernel = Audio::NewPseudoConstantQKernelTransform(CQTSettings, FFTAlgorithm->Size(), InSampleRate);
+		float LowerValue = Samples[ChannelIndex * FrameCount + LowerIndex];
+		float UpperValue = Samples[ChannelIndex * FrameCount + UpperIndex];
 
-	if (!CQTKernel.IsValid())
-	{
-		UE_LOG(LogNiagara, Error, TEXT("Failed to create CQT kernel."));
-		return;
+		return FMath::Lerp(LowerValue, UpperValue, Fraction);
 	}
 }
 
-void FNiagaraDataInterfaceProxySpectrum::Clamp(float& InMinimumFrequency, float& InMaximumFrequency, int32& InNumBands) 
-{
-	const int32 MinimumSupportedNumBands = 1;
-	const int32 MaximumSupportedNumBands = 16384;
-	const float MinimumSupportedFrequency = 20.f;
-	const float MaximumSupportedFrequency = 20000.f;
-
-	InMinimumFrequency = FMath::Max(InMinimumFrequency, MinimumSupportedFrequency);
-	InMaximumFrequency = FMath::Min(InMaximumFrequency, MaximumSupportedFrequency);
-	InMaximumFrequency = FMath::Max(InMinimumFrequency, InMaximumFrequency);
-
-	InNumBands = FMath::Clamp(InNumBands, MinimumSupportedNumBands, MaximumSupportedNumBands);
-}
-
-float FNiagaraDataInterfaceProxySpectrum::GetNumOctaves(float InMinimumFrequency, float InMaximumFrequency) 
-{
-	const float MinimumSupportedNumOctaves = 0.01f;
-
-	float NumOctaves = FMath::Log2(InMaximumFrequency) - FMath::Log2(InMinimumFrequency);
-	NumOctaves = FMath::Max(MinimumSupportedNumOctaves, NumOctaves);
-
-	return NumOctaves;
-}
-
-float FNiagaraDataInterfaceProxySpectrum::GetNumBandsPerOctave(int32 InNumBands, float InNumOctaves)
-{
-	InNumOctaves = FMath::Max(0.01f, InNumOctaves);
-	const float MinimumSupportedNumBandsPerOctave = 0.01f;
-
-	float NumBandsPerOctave = FMath::Max(MinimumSupportedNumBandsPerOctave, static_cast<float>(InNumBands) / InNumOctaves);
-
-	return NumBandsPerOctave;
-}
-
-float FNiagaraDataInterfaceProxySpectrum::GetBandwidthStretch(float InSampleRate, float InFFTSize, float InNumBandsPerOctave, float InMinimumFrequency)
-{
-	InFFTSize = FMath::Max(1.f, InFFTSize);
-
-	float MinimumFrequencySpacing = (FMath::Pow(2.f, 1.f / InNumBandsPerOctave) - 1.f) * InMinimumFrequency;
-	MinimumFrequencySpacing = FMath::Max(0.01f, MinimumFrequencySpacing);
-
-	float FFTBinSpacing = InSampleRate / InFFTSize;
-	float BandwidthStretch = FFTBinSpacing / MinimumFrequencySpacing;
-
-	return FMath::Clamp(BandwidthStretch, 0.5f, 2.f);
-}
-
-Audio::FPseudoConstantQKernelSettings FNiagaraDataInterfaceProxySpectrum::GetConstantQSettings(float InMinimumFrequency, float InMaximumFrequency, int32 InNumBands, float InNumBandsPerOctave, float InBandwidthStretch)
+UNiagaraDataInterfaceAudioSpectrum::UNiagaraDataInterfaceAudioSpectrum(FObjectInitializer const& ObjectInitializer)
+	: Super(ObjectInitializer)
 {
-	Audio::FPseudoConstantQKernelSettings CQTKernelSettings;
-
-	CQTKernelSettings.NumBands = InNumBands;
-	CQTKernelSettings.KernelLowestCenterFreq = InMinimumFrequency; 
-	CQTKernelSettings.NumBandsPerOctave = InNumBandsPerOctave;
-	CQTKernelSettings.BandWidthStretch = InBandwidthStretch;
-	CQTKernelSettings.Normalization = Audio::EPseudoConstantQNormalization::EqualEnergy;
-
-	return CQTKernelSettings;
+	Proxy.Reset(new UE::Niagara::DataInterfaceAudioSpectrum::Private::FNDIProxy());
 }
 
-Audio::FFFTSettings FNiagaraDataInterfaceProxySpectrum::GetFFTSettings(float InMinimumFrequency, float InSampleRate, float InNumBandsPerOctave)
+void UNiagaraDataInterfaceAudioSpectrum::GetSpectrumValue(FVectorVMExternalFunctionContext& Context)
 {
-	const int32 MaximumSupportedLog2FFTSize = 13;
-	const int32 MinimumSupportedLog2FFTSize = 8;
+	using namespace UE::Niagara::DataInterfaceAudioSpectrum::Private;
 
-	InNumBandsPerOctave = FMath::Max(0.01f, InNumBandsPerOctave);
+	const int32 InstanceCount = Context.GetNumInstances();
 
-	float MinimumFrequencySpacing = (FMath::Pow(2.f, 1.f / InNumBandsPerOctave) - 1.f) * InMinimumFrequency;
-
-	MinimumFrequencySpacing = FMath::Max(0.01f, MinimumFrequencySpacing);
-
-	int32 DesiredFFTSize = FMath::CeilToInt(3.f * InSampleRate / MinimumFrequencySpacing);
+	VectorVM::FUserPtrHandler<FInstanceData_GameThread> InstData(Context);
+	VectorVM::FExternalFuncInputHandler<float> InNormalizedPos(Context);
+	VectorVM::FExternalFuncInputHandler<int32> InChannel(Context);
+	VectorVM::FExternalFuncRegisterHandler<float> OutValue(Context);
 
-	int32 Log2FFTSize = MinimumSupportedLog2FFTSize;
-	while ((DesiredFFTSize > (1 << Log2FFTSize)) && (Log2FFTSize < MaximumSupportedLog2FFTSize))
+	if (InstData->SharedResource)
 	{
-		Log2FFTSize++;
-	}
-
-	Audio::FFFTSettings FFTSettings;
+		const FNDIAudio_SharedResource& SharedAudio = InstData->SharedResource.ReadResource();
 
-	FFTSettings.Log2Size = Log2FFTSize;
-    FFTSettings.bArrays128BitAligned = true;
-    FFTSettings.bEnableHardwareAcceleration = true;
+		{
+			SharedAudio.WaitForAudio();
+		}
 
-	return FFTSettings;
-}
+		int32 ChannelCount = 0;
+		int32 ChannelResolution = 0;
+		TConstArrayView<float> SpectrumBuffer = SharedAudio.ReadSpectrumBuffer(ChannelCount, ChannelResolution);
 
+		if (ChannelCount && ChannelResolution)
+		{
+			for (int32 InstanceIdx = 0; InstanceIdx < InstanceCount; ++InstanceIdx)
+			{
+				const float Position = InNormalizedPos.GetAndAdvance();
+				const int32 ChannelIndex = InChannel.GetAndAdvance();
+				*OutValue.GetDestAndAdvance() = SamplePlanarAudioBuffer(
+					SpectrumBuffer,
+					ChannelCount,
+					ChannelResolution,
+					FMath::Clamp(ChannelIndex, 0, ChannelCount - 1),
+					Position);
+			}
 
-UNiagaraDataInterfaceAudioSpectrum::UNiagaraDataInterfaceAudioSpectrum(FObjectInitializer const& ObjectInitializer)
-	: Super(ObjectInitializer)
-	, Resolution(512)
-	, MinimumFrequency(55.f)
-	, MaximumFrequency(10000.f)
-	, NoiseFloorDb(-60.f)
-	
-{
-	if (!HasAnyFlags(RF_ClassDefaultObject))
-	{
-		Proxy = MakeUnique<FNiagaraDataInterfaceProxySpectrum>(MinimumFrequency, MaximumFrequency, Resolution, NoiseFloorDb, 16384);
+			return;
+		}
 	}
-}
-
-
-void UNiagaraDataInterfaceAudioSpectrum::GetSpectrumValue(FVectorVMExternalFunctionContext& Context)
-{
-	GetProxyAs<FNiagaraDataInterfaceProxySpectrum>()->UpdateSpectrum();
-
-	VectorVM::FExternalFuncInputHandler<float> InNormalizedPos(Context);
-	VectorVM::FExternalFuncInputHandler<int32> InChannel(Context);
-	VectorVM::FExternalFuncRegisterHandler<float> OutValue(Context);
 
+	// fallback defaults
 	for (int32 InstanceIdx = 0; InstanceIdx < Context.GetNumInstances(); ++InstanceIdx)
 	{
-		float Position = InNormalizedPos.Get();
-		int32 Channel = InChannel.Get();
-		*OutValue.GetDest() = GetProxyAs<FNiagaraDataInterfaceProxySpectrum>()->GetSpectrumValue(Position, Channel);
-
-		InNormalizedPos.Advance();
-		InChannel.Advance();
-		OutValue.Advance();
+		InNormalizedPos.GetAndAdvance();
+		InChannel.GetAndAdvance();
+		*OutValue.GetDestAndAdvance() = 0.0f;
 	}
 }
 
 void UNiagaraDataInterfaceAudioSpectrum::GetNumChannels(FVectorVMExternalFunctionContext& Context)
 {
+	const int32 InstanceCount = Context.GetNumInstances();
+
+	VectorVM::FUserPtrHandler<UE::Niagara::DataInterfaceAudioSpectrum::Private::FInstanceData_GameThread> InstData(Context);
 	VectorVM::FExternalFuncRegisterHandler<int32> OutChannel(Context);
 
-	for (int32 InstanceIdx = 0; InstanceIdx < Context.GetNumInstances(); ++InstanceIdx)
+	const int32 ChannelCount = InstData->SharedResource
+		? InstData->SharedResource.ReadResource().GetChannelCount()
+		: 0;
+
+	for (int32 InstanceIdx = 0; InstanceIdx < InstanceCount; ++InstanceIdx)
 	{
-		*OutChannel.GetDestAndAdvance() = GetProxyAs<FNiagaraDataInterfaceProxySpectrum>()->GetNumChannels();
+		*OutChannel.GetDestAndAdvance() = ChannelCount;
 	}
 }
 
 #if WITH_EDITORONLY_DATA
 void UNiagaraDataInterfaceAudioSpectrum::GetFunctionsInternal(TArray<FNiagaraFunctionSignature>& OutFunctions) const
 {
+	using namespace UE::Niagara::DataInterfaceAudioSpectrum::Private;
 	Super::GetFunctionsInternal(OutFunctions);
 
 	{
 		FNiagaraFunctionSignature GetSpectrumSignature;
-		GetSpectrumSignature.Name = GetSpectrumFunctionName;
+		GetSpectrumSignature.Name = AudioSpectrumFunctionName;
 		GetSpectrumSignature.Inputs.Add(FNiagaraVariable(GetClass(), TEXT("Spectrum")));
 		GetSpectrumSignature.Inputs.Add(FNiagaraVariable(FNiagaraTypeDefinition::GetFloatDef(), TEXT("NormalizedPositionInSpectrum")));
 		GetSpectrumSignature.Inputs.Add(FNiagaraVariable(FNiagaraTypeDefinition::GetIntDef(), TEXT("ChannelIndex")));
@@ -632,7 +252,9 @@ DEFINE_NDI_DIRECT_FUNC_BINDER(UNiagaraDataInterfaceAudioSpectrum, GetNumChannels
 
 void UNiagaraDataInterfaceAudioSpectrum::GetVMExternalFunction(const FVMExternalFunctionBindingInfo& BindingInfo, void* InstanceData, FVMExternalFunction &OutFunc)
 {
-	if (BindingInfo.Name == GetSpectrumFunctionName)
+	using namespace UE::Niagara::DataInterfaceAudioSpectrum::Private;
+
+	if (BindingInfo.Name == AudioSpectrumFunctionName)
 	{
 		NDI_FUNC_BINDER(UNiagaraDataInterfaceAudioSpectrum, GetSpectrumValue)::Bind(this, OutFunc);
 	}
@@ -649,112 +271,81 @@ void UNiagaraDataInterfaceAudioSpectrum::GetVMExternalFunction(const FVMExternal
 #if WITH_EDITORONLY_DATA
 bool UNiagaraDataInterfaceAudioSpectrum::AppendCompileHash(FNiagaraCompileHashVisitor* InVisitor) const
 {
-	bool bSuccess = Super::AppendCompileHash(InVisitor);
-	bSuccess &= InVisitor->UpdateShaderParameters<FShaderParameters>();
-	return bSuccess;
+	using namespace UE::Niagara::DataInterfaceAudioSpectrum::Private;
+
+	if (!Super::AppendCompileHash(InVisitor))
+	{
+		return false;
+	}
+	InVisitor->UpdateShaderFile(CommonShaderFile);
+	InVisitor->UpdateShaderFile(TemplateShaderFile);
+	InVisitor->UpdateShaderParameters<FShaderParameters>();
+	return true;
 }
 
 bool UNiagaraDataInterfaceAudioSpectrum::GetFunctionHLSL(const FNiagaraDataInterfaceGPUParamInfo& ParamInfo, const FNiagaraDataInterfaceGeneratedFunction& FunctionInfo, int FunctionInstanceIndex, FString& OutHLSL)
 {
-	bool ParentRet = Super::GetFunctionHLSL(ParamInfo, FunctionInfo, FunctionInstanceIndex, OutHLSL);
-	if (ParentRet)
-	{
-		return true;
-	}
-	else if (FunctionInfo.DefinitionName == GetSpectrumFunctionName)
+	using namespace UE::Niagara::DataInterfaceAudioSpectrum::Private;
+
+	if (Super::GetFunctionHLSL(ParamInfo, FunctionInfo, FunctionInstanceIndex, OutHLSL))
 	{
-		/**
-		 * float FrameIndex = NormalizedPositionInBuffer * SpectrumBuffer.Num();
-		 * int32 LowerFrameIndex = FMath::FloorToInt(FrameIndex);
-		 * float LowerFrameAmplitude = SpectrumBuffer[LowerFrameIndex * NumChannels + ChannelIndex];
-		 * int32 HigherFrameIndex = FMath::CeilToInt(FrameIndex);
-		 * float HigherFrameAmplitude = SpectrumBuffer[HigherFrameIndex * NumChannels + ChannelIndex];
-		 * float Fraction = HigherFrameIndex - FrameIndex;
-		 * return FMath::Lerp<float>(LowerFrameAmplitude, HigherFrameAmplitude, Fraction);
-		 */
-		static const TCHAR *FormatBounds = TEXT(R"(
-			void {FunctionName}(float In_NormalizedPosition, int In_ChannelIndex, out float Out_Val)
-			{
-				float FrameIndex = In_NormalizedPosition * {SpectrumResolution};
-				int MaxIndex = {SpectrumResolution} - 1;
-				int LowerIndex = floor(FrameIndex);
-				LowerIndex = LowerIndex < {SpectrumResolution} ? LowerIndex : MaxIndex;
-				LowerIndex = LowerIndex >= 0 ? LowerIndex : 0;
-				int UpperIndex =  LowerIndex < MaxIndex ? LowerIndex + 1 : LowerIndex;
-				float Fraction = FrameIndex - LowerIndex;
-				Fraction = Fraction > 1.0 ? 1.0 : Fraction;
-				Fraction = Fraction < 0.0 ? 0.0 : Fraction;
-				float LowerValue = {SpectrumBuffer}.Load(In_ChannelIndex * {SpectrumResolution} + LowerIndex);
-				float UpperValue = {SpectrumBuffer}.Load(In_ChannelIndex * {SpectrumResolution} + UpperIndex);
-
-				Out_Val = lerp(LowerValue, UpperValue, Fraction);
-			}
-		)");
-
-		TMap<FString, FStringFormatArg> ArgsBounds = {
-			{TEXT("FunctionName"), FStringFormatArg(FunctionInfo.InstanceName)},
-			{TEXT("ChannelCount"), FStringFormatArg(ParamInfo.DataInterfaceHLSLSymbol + NumChannelsName)},
-			{TEXT("SpectrumResolution"), FStringFormatArg(ParamInfo.DataInterfaceHLSLSymbol + ResolutionName)},
-			{TEXT("SpectrumBuffer"), FStringFormatArg(ParamInfo.DataInterfaceHLSLSymbol + GetSpectrumName)},
-		};
-		OutHLSL += FString::Format(FormatBounds, ArgsBounds);
 		return true;
 	}
-	else if (FunctionInfo.DefinitionName == GetNumChannelsFunctionName)
+
+	if ((FunctionInfo.DefinitionName == AudioSpectrumFunctionName) ||
+		(FunctionInfo.DefinitionName == GetNumChannelsFunctionName))
 	{
-		// See GetNumChannels()
-		static const TCHAR *FormatBounds = TEXT(R"(
-			void {FunctionName}(out int Out_Val)
-			{
-				Out_Val = {ChannelCount};
-			}
-		)");
-		TMap<FString, FStringFormatArg> ArgsBounds = {
-			{TEXT("FunctionName"), FStringFormatArg(FunctionInfo.InstanceName)},
-			{TEXT("ChannelCount"), FStringFormatArg(ParamInfo.DataInterfaceHLSLSymbol + NumChannelsName)},
-		};
-		OutHLSL += FString::Format(FormatBounds, ArgsBounds);
 		return true;
 	}
-	else
-	{
-		return false;
-	}
+
+	return false;
+}
+
+void UNiagaraDataInterfaceAudioSpectrum::GetCommonHLSL(FString& OutHLSL)
+{
+	Super::GetCommonHLSL(OutHLSL);
+	OutHLSL.Appendf(TEXT("#include \"%s\"\n"), UE::Niagara::DataInterfaceAudioSpectrum::Private::CommonShaderFile);
 }
 
 void UNiagaraDataInterfaceAudioSpectrum::GetParameterDefinitionHLSL(const FNiagaraDataInterfaceGPUParamInfo& ParamInfo, FString& OutHLSL)
 {
-	Super::GetParameterDefinitionHLSL(ParamInfo, OutHLSL);
-
-	static const TCHAR* FormatDeclarations = TEXT(R"(				
-		Buffer<float> {GetSpectrumName};
-		int {NumChannelsName};
-		int {ResolutionName};
-
-	)");
-	TMap<FString, FStringFormatArg> ArgsDeclarations = {
-		{ TEXT("GetSpectrumName"),    FStringFormatArg(ParamInfo.DataInterfaceHLSLSymbol + GetSpectrumName) },
-		{ TEXT("NumChannelsName"),    FStringFormatArg(ParamInfo.DataInterfaceHLSLSymbol + NumChannelsName) },
-		{ TEXT("ResolutionName"),     FStringFormatArg(ParamInfo.DataInterfaceHLSLSymbol + ResolutionName) },
+	const TMap<FString, FStringFormatArg> TemplateArgs =
+	{
+		{TEXT("ParameterName"),	ParamInfo.DataInterfaceHLSLSymbol},
 	};
-	OutHLSL += FString::Format(FormatDeclarations, ArgsDeclarations);
+	AppendTemplateHLSL(OutHLSL, UE::Niagara::DataInterfaceAudioSpectrum::Private::TemplateShaderFile, TemplateArgs);
 }
+
 #endif
 
 void UNiagaraDataInterfaceAudioSpectrum::BuildShaderParameters(FNiagaraShaderParametersBuilder& ShaderParametersBuilder) const
 {
-	ShaderParametersBuilder.AddNestedStruct<FShaderParameters>();
+	ShaderParametersBuilder.AddNestedStruct<UE::Niagara::DataInterfaceAudioSpectrum::Private::FShaderParameters>();
 }
 
 void UNiagaraDataInterfaceAudioSpectrum::SetShaderParameters(const FNiagaraDataInterfaceSetShaderParametersContext& Context) const
 {
-	FNiagaraDataInterfaceProxySpectrum& DIProxy = Context.GetProxy<FNiagaraDataInterfaceProxySpectrum>();
-	FReadBuffer& SpectrumBuffer = DIProxy.ComputeAndPostSRV();
+	using namespace UE::Niagara::DataInterfaceAudioSpectrum::Private;
+
+	FNDIProxy& DIProxy = Context.GetProxy<FNDIProxy>();
+	FInstanceData_RenderThread& InstanceData = DIProxy.SystemInstancesToInstanceData_RT.FindChecked(Context.GetSystemInstanceID());
+
+	FRDGBuilder& GraphBuilder = Context.GetGraphBuilder();
 
 	FShaderParameters* ShaderParameters = Context.GetParameterNestedStruct<FShaderParameters>();
-	ShaderParameters->NumChannels = DIProxy.GetNumChannels();
-	ShaderParameters->Resolution = DIProxy.GetNumBands();
-	ShaderParameters->SpectrumBuffer = FNiagaraRenderer::GetSrvOrDefaultFloat(SpectrumBuffer.SRV);
+	ShaderParameters->ChannelCount = InstanceData.ChannelCount;
+	ShaderParameters->Resolution = InstanceData.ChannelResolution;
+
+	if (InstanceData.PooledBuffer)
+	{
+		FRDGBufferRef RDGBuffer = GraphBuilder.RegisterExternalBuffer(InstanceData.PooledBuffer);
+		ShaderParameters->SpectrumBuffer = GraphBuilder.CreateSRV(RDGBuffer, PF_R32_FLOAT);
+	}
+	else
+	{
+		FRDGBufferRef RDGBuffer = GSystemTextures.GetDefaultBuffer(GraphBuilder, sizeof(float), 0.0f);
+		ShaderParameters->SpectrumBuffer = GraphBuilder.CreateSRV(RDGBuffer, PF_R32_FLOAT);
+	}
 }
 
 bool UNiagaraDataInterfaceAudioSpectrum::Equals(const UNiagaraDataInterface* Other) const
@@ -763,6 +354,7 @@ bool UNiagaraDataInterfaceAudioSpectrum::Equals(const UNiagaraDataInterface* Oth
 
 	const UNiagaraDataInterfaceAudioSpectrum* OtherSpectrum = CastChecked<const UNiagaraDataInterfaceAudioSpectrum>(Other);
 
+	bIsEqual &= OtherSpectrum->Submix == Submix;
 	bIsEqual &= OtherSpectrum->Resolution == Resolution;
 	bIsEqual &= OtherSpectrum->MinimumFrequency == MinimumFrequency;
 	bIsEqual &= OtherSpectrum->MaximumFrequency == MaximumFrequency;
@@ -771,59 +363,119 @@ bool UNiagaraDataInterfaceAudioSpectrum::Equals(const UNiagaraDataInterface* Oth
 	return bIsEqual;
 }
 
-#if WITH_EDITOR
-void UNiagaraDataInterfaceAudioSpectrum::PostEditChangeProperty(struct FPropertyChangedEvent& PropertyChangedEvent)
+void UNiagaraDataInterfaceAudioSpectrum::PostInitProperties()
 {
-	Super::PostEditChangeProperty(PropertyChangedEvent);
-	
-	static const FName ResolutionFName = GET_MEMBER_NAME_CHECKED(UNiagaraDataInterfaceAudioSpectrum, Resolution);
-	static const FName MinimumFrequencyFName = GET_MEMBER_NAME_CHECKED(UNiagaraDataInterfaceAudioSpectrum, MinimumFrequency);
-	static const FName MaximumFrequencyFName = GET_MEMBER_NAME_CHECKED(UNiagaraDataInterfaceAudioSpectrum, MaximumFrequency);
-	static const FName NoiseFloorFName = GET_MEMBER_NAME_CHECKED(UNiagaraDataInterfaceAudioSpectrum, NoiseFloorDb);
+	Super::PostInitProperties();
 
-	if (!HasAnyFlags(RF_ClassDefaultObject))
+	if (HasAnyFlags(RF_ClassDefaultObject))
 	{
-		if (FProperty* PropertyThatChanged = PropertyChangedEvent.Property)
-		{
-			const FName& Name = PropertyThatChanged->GetFName();
-			if (Name == ResolutionFName || Name == MaximumFrequencyFName || Name == MinimumFrequencyFName)
-			{
-				GetProxyAs<FNiagaraDataInterfaceProxySpectrum>()->UpdateCQT(MinimumFrequency, MaximumFrequency, Resolution);
-			}
-			else if (Name == NoiseFloorFName)
-			{
-				GetProxyAs<FNiagaraDataInterfaceProxySpectrum>()->UpdateNoiseFloor(NoiseFloorDb);
-			}
-		}
+		ENiagaraTypeRegistryFlags Flags = ENiagaraTypeRegistryFlags::AllowAnyVariable | ENiagaraTypeRegistryFlags::AllowParameter;
+		FNiagaraTypeRegistry::Register(FNiagaraTypeDefinition(GetClass()), Flags);
 	}
 }
-#endif //WITH_EDITOR
 
-void UNiagaraDataInterfaceAudioSpectrum::PostInitProperties()
+bool UNiagaraDataInterfaceAudioSpectrum::PerInstanceTick(void* PerInstanceData, FNiagaraSystemInstance* SystemInstance, float DeltaSeconds)
 {
-	Super::PostInitProperties();
+	using namespace UE::Niagara::DataInterfaceAudioSpectrum::Private;
 
-	if (HasAnyFlags(RF_ClassDefaultObject))
+	FInstanceData_GameThread* Inst = reinterpret_cast<FInstanceData_GameThread*>(PerInstanceData);
+
+	// reset our per instance data if the submix has changed/invalidated
+	if (Inst->bSoundSubmixAssigned && !Inst->SoundSubmix.IsValid())
 	{
-		ENiagaraTypeRegistryFlags Flags = ENiagaraTypeRegistryFlags::AllowAnyVariable | ENiagaraTypeRegistryFlags::AllowParameter;
-		FNiagaraTypeRegistry::Register(FNiagaraTypeDefinition(GetClass()), Flags);
+		return true;
 	}
-	else
+
+	if (Inst->SoundSubmix != Submix)
+	{
+		return true;
+	}
+
+	return false;
+}
+
+bool UNiagaraDataInterfaceAudioSpectrum::InitPerInstanceData(void* PerInstanceData, FNiagaraSystemInstance* SystemInstance)
+{
+	using namespace UE::Niagara::DataInterfaceAudioSpectrum::Private;
+
+	FInstanceData_GameThread* Inst = new (PerInstanceData) FInstanceData_GameThread();
+
+	Inst->SoundSubmix = MakeWeakObjectPtr(Submix);
+	Inst->Resolution = Resolution;
+	Inst->MinimumFrequency = MinimumFrequency;
+	Inst->MaximumFrequency = MaximumFrequency;
+	Inst->NoiseFloorDb = NoiseFloorDb;
+	Inst->bSoundSubmixAssigned = Submix != nullptr;
+
+	if (FNiagaraWorldManager* WorldManager = SystemInstance->GetWorldManager())
 	{
-		GetProxyAs<FNiagaraDataInterfaceProxySpectrum>()->UpdateCQT(MinimumFrequency, MaximumFrequency, Resolution);
-		GetProxyAs<FNiagaraDataInterfaceProxySpectrum>()->UpdateNoiseFloor(NoiseFloorDb);
+		FNDIAudio_GeneratedData::FResourceDesc AudioDesc;
+		AudioDesc.DeviceId = INDEX_NONE;
+		AudioDesc.Submix = Submix;
+		AudioDesc.bUseLatestAudio = false;
+		AudioDesc.bGenerateSpectrum = true;
+		AudioDesc.MinimumFrequency = MinimumFrequency;
+		AudioDesc.MaximumFrequency = MaximumFrequency;
+		AudioDesc.NoiseFloorDb = NoiseFloorDb;
+		AudioDesc.SpectrumSamplingResolution = Resolution;
+		AudioDesc.SamplingMethod = FNDIAudio_SharedResource::ESamplingWindowMethod::ByCount;
+		AudioDesc.SamplingWindowCount = DefaultSamplesToBufferCount;
+		AudioDesc.bContinuousSampling = true;
+
+		if (UWorld* World = WorldManager->GetWorld())
+		{
+			FAudioDeviceHandle AudioDeviceHandle = World->GetAudioDevice();
+			if (AudioDeviceHandle.IsValid())
+			{
+				AudioDesc.DeviceId = AudioDeviceHandle.GetDeviceID();
+			}
+		}
+
+		FNDIAudio_SharedResourceUsage ResourceUsage;
+		{
+			bool bInstanceRequiresCpu = false;
+			bool bInstanceRequiresGpu = false;
+			SystemInstance->EvaluateBoundFunction(AudioSpectrumFunctionName, bInstanceRequiresCpu, bInstanceRequiresGpu);
+			FNDIAudio_SharedResourceUsage::UpdateUsageFlags(bInstanceRequiresCpu, bInstanceRequiresGpu, ResourceUsage.SpectrumUsage);
+		}
+
+		FNDIAudio_GeneratedData& SpectrumGeneratedData = WorldManager->EditGeneratedData<FNDIAudio_GeneratedData>();
+		Inst->SharedResource = SpectrumGeneratedData.GetSharedResource(ResourceUsage, AudioDesc);
 	}
+
+	return true;
 }
 
-void UNiagaraDataInterfaceAudioSpectrum::PostLoad()
+void UNiagaraDataInterfaceAudioSpectrum::DestroyPerInstanceData(void* PerInstanceData, FNiagaraSystemInstance* SystemInstance)
 {
-	Super::PostLoad();
+	using namespace UE::Niagara::DataInterfaceAudioSpectrum::Private;
+
+	FInstanceData_GameThread* InstData = reinterpret_cast<FInstanceData_GameThread*>(PerInstanceData);
 
-	if (!HasAnyFlags(RF_ClassDefaultObject))
+	InstData->~FInstanceData_GameThread();
+}
+
+int32 UNiagaraDataInterfaceAudioSpectrum::PerInstanceDataSize() const
+{
+	return sizeof(UE::Niagara::DataInterfaceAudioSpectrum::Private::FInstanceData_GameThread);
+}
+
+void UNiagaraDataInterfaceAudioSpectrum::ProvidePerInstanceDataForRenderThread(void* DataForRenderThread, void* PerInstanceData, const FNiagaraSystemInstanceID& SystemInstance)
+{
+	UE::Niagara::DataInterfaceAudioSpectrum::Private::FNDIProxy::ProvidePerInstanceDataForRenderThread(DataForRenderThread, PerInstanceData, SystemInstance);
+}
+
+ETickingGroup UNiagaraDataInterfaceAudioSpectrum::CalculateTickGroup(const void* PerInstanceData) const
+{
+	using namespace UE::Niagara::DataInterfaceAudioSpectrum::Private;
+
+	const FInstanceData_GameThread* Inst = reinterpret_cast<const FInstanceData_GameThread*>(PerInstanceData);
+	if (Inst->SharedResource.Usage.IsValid())
 	{
-		GetProxyAs<FNiagaraDataInterfaceProxySpectrum>()->UpdateCQT(MinimumFrequency, MaximumFrequency, Resolution);
-		GetProxyAs<FNiagaraDataInterfaceProxySpectrum>()->UpdateNoiseFloor(NoiseFloorDb);
+		return static_cast<ETickingGroup>(FNDIAudio_GeneratedData::GeneratedDataTickGroup + 1);
 	}
+
+	return NiagaraFirstTickGroup;
 }
 
 bool UNiagaraDataInterfaceAudioSpectrum::CopyToInternal(UNiagaraDataInterface* Destination) const
@@ -837,31 +489,14 @@ bool UNiagaraDataInterfaceAudioSpectrum::CopyToInternal(UNiagaraDataInterface* D
 
 	if (CastedDestination)
 	{
+		CastedDestination->Submix = Submix;
 		CastedDestination->Resolution = Resolution;
 		CastedDestination->MinimumFrequency = MinimumFrequency;
 		CastedDestination->MaximumFrequency = MaximumFrequency;
 		CastedDestination->NoiseFloorDb = NoiseFloorDb;
-
-		if (!CastedDestination->HasAnyFlags(RF_ClassDefaultObject))
-		{
-			CastedDestination->GetProxyAs<FNiagaraDataInterfaceProxySpectrum>()->UpdateCQT(MinimumFrequency, MaximumFrequency, Resolution);
-			CastedDestination->GetProxyAs<FNiagaraDataInterfaceProxySpectrum>()->UpdateNoiseFloor(NoiseFloorDb);
-		}
 	}
 	
 	return true;
 }
 
-
-// Global VM function names, also used by the shaders code generation methods.
-const FName UNiagaraDataInterfaceAudioSpectrum::GetSpectrumFunctionName("AudioSpectrum");
-const FName UNiagaraDataInterfaceAudioSpectrum::GetNumChannelsFunctionName("GetNumChannels");
-
-// Global variable prefixes, used in HLSL parameter declarations.
-const FString UNiagaraDataInterfaceAudioSpectrum::GetSpectrumName(TEXT("_SpectrumBuffer"));
-const FString UNiagaraDataInterfaceAudioSpectrum::NumChannelsName(TEXT("_NumChannels"));
-const FString UNiagaraDataInterfaceAudioSpectrum::ResolutionName(TEXT("_Resolution"));
-
-
 #undef LOCTEXT_NAMESPACE
-
diff --git a/Engine/Plugins/FX/Niagara/Source/Niagara/Private/NiagaraGeneratedDataAudioSampling.cpp b/Engine/Plugins/FX/Niagara/Source/Niagara/Private/NiagaraGeneratedDataAudioSampling.cpp
new file mode 100644
--- /dev/null
+++ b/Engine/Plugins/FX/Niagara/Source/Niagara/Private/NiagaraGeneratedDataAudioSampling.cpp
@@ -0,0 +1,1058 @@
+// Copyright Epic Games, Inc. All Rights Reserved.
+
+#include "NiagaraGeneratedDataAudioSampling.h"
+
+#include "Algo/RemoveIf.h"
+#include "AudioDevice.h"
+#include "AudioDeviceManager.h"
+#include "DSP/AudioFFT.h"
+#include "DSP/ConstantQ.h"
+#include "DSP/DeinterleaveView.h"
+#include "DSP/FFTAlgorithm.h"
+#include "DSP/MultithreadedPatching.h"
+#include "DSP/SlidingWindow.h"
+#include "ISubmixBufferListener.h"
+#include "NiagaraWorldManager.h"
+#include "Sound/SoundSubmix.h"
+
+namespace UE::Niagara::GeneratedDataAudioSampling::Private
+{
+
+static Audio::FPseudoConstantQKernelSettings GetConstantQSettings(float InMinimumFrequency, float InMaximumFrequency, int32 InNumBands, float InNumBandsPerOctave, float InBandwidthStretch)
+{
+	Audio::FPseudoConstantQKernelSettings CQTKernelSettings;
+
+	CQTKernelSettings.NumBands = InNumBands;
+	CQTKernelSettings.KernelLowestCenterFreq = InMinimumFrequency;
+	CQTKernelSettings.NumBandsPerOctave = InNumBandsPerOctave;
+	CQTKernelSettings.BandWidthStretch = InBandwidthStretch;
+	CQTKernelSettings.Normalization = Audio::EPseudoConstantQNormalization::EqualEnergy;
+
+	return CQTKernelSettings;
+}
+
+static Audio::FFFTSettings GetFFTSettings(float InMinimumFrequency, float InSampleRate, float InNumBandsPerOctave)
+{
+	const int32 MaximumSupportedLog2FFTSize = 13;
+	const int32 MinimumSupportedLog2FFTSize = 8;
+
+	InNumBandsPerOctave = FMath::Max(0.01f, InNumBandsPerOctave);
+
+	float MinimumFrequencySpacing = (FMath::Pow(2.f, 1.f / InNumBandsPerOctave) - 1.f) * InMinimumFrequency;
+
+	MinimumFrequencySpacing = FMath::Max(0.01f, MinimumFrequencySpacing);
+
+	int32 DesiredFFTSize = FMath::CeilToInt(3.f * InSampleRate / MinimumFrequencySpacing);
+
+	int32 Log2FFTSize = MinimumSupportedLog2FFTSize;
+	while ((DesiredFFTSize > (1 << Log2FFTSize)) && (Log2FFTSize < MaximumSupportedLog2FFTSize))
+	{
+		Log2FFTSize++;
+	}
+
+	Audio::FFFTSettings FFTSettings;
+
+	FFTSettings.Log2Size = Log2FFTSize;
+	FFTSettings.bArrays128BitAligned = true;
+	FFTSettings.bEnableHardwareAcceleration = true;
+
+	return FFTSettings;
+}
+
+class FSubmixListener : public ISubmixBufferListener
+{
+public:
+
+	/** Construct an FSubmixListener
+		*
+		* @param InMixer - Input FPatchMixer. Submix audio will be mixed into this patch mixer.
+		* @param InNumSamplesToBuffer - Number of samples to hold in patch input. Low values may cause
+		*                               data to be overwritten by threads that produce audio. High
+		*                               values require more memory.
+		*/
+	FSubmixListener(Audio::FPatchMixer& InMixer, int32 InMaxChannelCount, int32 InNumSamplesToBuffer, int32 InMinimumSamplesForActivation, Audio::FDeviceId InDeviceId, USoundSubmix* InSoundSubmix)
+		: MaxChannelCount(InMaxChannelCount)
+		, MixerInput(InMixer.AddNewInput(InNumSamplesToBuffer, 1.0f))
+		, AudioDeviceId(InDeviceId)
+		, Submix(InSoundSubmix)
+	{
+		InitializedEvent.Emplace(UE_SOURCE_LOCATION);
+		SamplesTillInitialized = InMinimumSamplesForActivation;
+	}
+
+	virtual ~FSubmixListener() override
+	{
+		check(!bIsRegistered);
+	}
+
+	FSubmixListener() = delete;
+	FSubmixListener(const FSubmixListener& Other) = delete;
+	FSubmixListener(FSubmixListener&& Other) = delete;
+
+	void RegisterToSubmix()
+	{
+		if (FAudioDevice* AudioDevice = FAudioDeviceManager::Get()->GetAudioDeviceRaw(AudioDeviceId))
+		{
+			bIsRegistered = true;
+
+			USoundSubmix* SubmixToRegister = Submix.IsValid() ? Submix.Get() : &AudioDevice->GetMainSubmixObject();
+			AudioDevice->RegisterSubmixBufferListener(AsShared(), *SubmixToRegister);
+		}
+	}
+
+	void UnregisterFromSubmix()
+	{
+		if (bIsRegistered)
+		{
+			bIsRegistered = false;
+
+			if (IsInGameThread())
+			{
+				if (FAudioDevice* AudioDevice = FAudioDeviceManager::Get()->GetAudioDeviceRaw(AudioDeviceId))
+				{
+					USoundSubmix* SubmixToUnregister = Submix.IsValid() ? Submix.Get() : &AudioDevice->GetMainSubmixObject();
+					AudioDevice->UnregisterSubmixBufferListener(AsShared(), *SubmixToUnregister);
+				}
+			}
+			else
+			{
+				FAudioThread::RunCommandOnAudioThread([Self = AsShared(), AudioDeviceId = AudioDeviceId, Submix = Submix]()
+				{
+					if (FAudioDevice* AudioDevice = FAudioDeviceManager::Get()->GetAudioDeviceRaw(AudioDeviceId))
+					{
+						USoundSubmix* SubmixToUnregister = Submix.IsValid() ? Submix.Get() : &AudioDevice->GetMainSubmixObject();
+						AudioDevice->UnregisterSubmixBufferListener(Self, *SubmixToUnregister);
+					}
+				});
+			}
+		}
+	}
+
+	/** Returns the current sample rate of the current submix. */
+	float GetSampleRate() const
+	{
+		return static_cast<float>(SubmixSampleRate);
+	}
+
+	/** Returns the number of channels of the current submix. */
+	int32 GetChannelCount() const
+	{
+		return ActiveChannelCount;
+	}
+
+	Audio::FDeviceId GetAudioDeviceId() const
+	{
+		return AudioDeviceId;
+	}
+
+	// Begin ISubmixBufferListener overrides
+	virtual const FString& GetListenerName() const override
+	{
+		static const FString ListenerName = TEXT("NiagaraAudioSamplingListener");
+		return ListenerName;
+	}
+
+	virtual void OnNewSubmixBuffer(const USoundSubmix* OwningSubmix, float* AudioData, int32 NumSamples, int32 NumChannels, const int32 SampleRate, double AudioClock) override
+	{
+		SubmixChannelCount = NumChannels;
+		SubmixSampleRate = SampleRate;
+
+		if (MaxChannelCount && MaxChannelCount < SubmixChannelCount)
+		{
+			// strip out the unneeded audio data based on our limited number of channels
+			const int32 FrameCount = NumSamples / NumChannels;
+			const int32 StrippedSampleCount = FrameCount * MaxChannelCount;
+
+			TArray<float> StrippedData;
+			StrippedData.Reserve(StrippedSampleCount);
+			for (int32 FrameIt = 0; FrameIt < FrameCount; ++FrameIt)
+			{
+				int32 ChannelIt = 0;
+				for (; ChannelIt < MaxChannelCount; ++ChannelIt)
+				{
+					StrippedData.Add(*AudioData);
+					++AudioData;
+				}
+
+				for (; ChannelIt < NumChannels; ++ChannelIt)
+				{
+					++AudioData;
+				}
+			}
+
+			ActiveChannelCount = MaxChannelCount;
+			MixerInput.PushAudio(StrippedData.GetData(), StrippedData.Num());
+		}
+		else
+		{
+			MixerInput.PushAudio(AudioData, NumSamples);
+			ActiveChannelCount = SubmixChannelCount;
+		}
+
+		const int32 PrevSamplesRemaining = SamplesTillInitialized.fetch_sub(NumSamples, std::memory_order_relaxed);
+		if (PrevSamplesRemaining > 0 && ((PrevSamplesRemaining - NumSamples) <= 0))
+		{
+			InitializedEvent->Trigger();
+		}
+	}
+	// End ISubmixBufferListener overrides
+
+	using FOptionalEvent = TOptional<UE::Tasks::FTaskEvent>;
+	FOptionalEvent GetInitializedEvent()
+	{
+		return InitializedEvent;
+	}
+	void ReleaseInitializedEvent()
+	{
+		InitializedEvent.Reset();
+	}
+
+private:
+	// properties pulled from the Submix
+	int32 SubmixChannelCount = 0;
+	int32 SubmixSampleRate = 0;
+
+	// limits to the number of channels we need to worry about specified at creation time (0 implies no limit)
+	int32 MaxChannelCount = 0;
+
+	// holds the number of channels that we are actively using (taking account of both SubmixChannelCount and MaxChannelCount)
+	int32 ActiveChannelCount = 0;
+
+	Audio::FPatchInput MixerInput;
+
+	Audio::FDeviceId AudioDeviceId = INDEX_NONE;
+	TWeakObjectPtr<USoundSubmix> Submix = nullptr;
+	bool bIsRegistered = false;
+
+	FOptionalEvent InitializedEvent;
+	std::atomic<int32> SamplesTillInitialized;
+};
+
+struct FCollectAudioContext
+{
+	int32 AudioSampleCount = 0;
+	int32 SpectrumResolution = 0;
+	int32 ResampledAudioResolution = 0;
+	int32 ChannelCount = 0;
+	float NoiseFloorDb = 0.0f;
+	float SampleRate = 0.0f;
+	float MinimumFrequency = 0.0f;
+	float MaximumFrequency = 0.0f;
+
+	bool bContinuousSampling = false;
+	bool bUseLatestAudio = false;
+};
+
+class FSpectrumBuilder
+{
+public:
+	using FSlidingBuffer = Audio::TSlidingBuffer<float>;
+	using FSlidingWindow = Audio::TAutoSlidingWindow<float, Audio::FAudioBufferAlignedAllocator>;
+	using FDeinterleaveView = Audio::TAutoDeinterleaveView<float, Audio::FAudioBufferAlignedAllocator>;
+	using FChannel = FDeinterleaveView::TChannel<Audio::FAudioBufferAlignedAllocator>;
+
+	void Build(const FCollectAudioContext& CollectContext, const Audio::FAlignedFloatBuffer& AudioBuffer)
+	{
+		TRACE_CPUPROFILER_EVENT_SCOPE(NiagaraAudioSampling::GenerateSpectrum);
+
+		SetAudioFormat(CollectContext);
+
+		ResetSpectrumBuffers(CollectContext);
+
+		// Set data to zero if we are skipping update due to bad samplerate or channel count.
+		if (CollectContext.ChannelCount == 0 || FMath::IsNearlyZero(CollectContext.SampleRate))
+		{
+			return;
+		}
+
+		// If we're in a bad internal state, give up here.	
+		if (!FFTAlgorithm.IsValid() || !CQTKernel.IsValid())
+		{
+			return;
+		}
+
+		// Run sliding window over available audio.
+		FSlidingWindow SlidingWindow(*SlidingBuffer, AudioBuffer, InterleavedBuffer);
+
+		const int32 NumBands = CollectContext.SpectrumResolution;
+		const float NoiseFloorDb = CollectContext.NoiseFloorDb;
+
+		auto MakeChannelBufferView = [&CollectContext](Audio::FAlignedFloatBuffer& PlanarBuffer, int32 ChannelIndex) -> TArrayView<float>
+		{
+			return MakeArrayView(PlanarBuffer.GetData() + CollectContext.SpectrumResolution * ChannelIndex, CollectContext.SpectrumResolution);
+		};
+
+		int32 NumWindows = 0;
+		for (Audio::FAlignedFloatBuffer& InterleavedWindow : SlidingWindow)
+		{
+			NumWindows++;
+
+			// Run deinterleave view over window.
+			FDeinterleaveView DeinterleaveView(InterleavedWindow, DeinterleavedBuffer, CollectContext.ChannelCount);
+			for (FChannel Channel : DeinterleaveView)
+			{
+				Audio::ArrayMultiplyInPlace(WindowBuffer, Channel.Values);
+
+				// Perform FFT
+				FMemory::Memcpy(FFTInputBuffer.GetData(), Channel.Values.GetData(), sizeof(float) * Channel.Values.Num());
+
+				FFTAlgorithm->ForwardRealToComplex(FFTInputBuffer.GetData(), FFTOutputBuffer.GetData());
+
+				// Transflate FFTOutput to power spectrum
+				Audio::ArrayComplexToPower(FFTOutputBuffer, PowerSpectrumBuffer);
+
+				// Take CQT of power spectrum
+				CQTKernel->TransformArray(PowerSpectrumBuffer, SpectrumBuffer);
+
+				// Accumulate power spectrum CQT output.
+				Audio::ArrayAddInPlace(SpectrumBuffer, MakeChannelBufferView(PlanarChannelSpectrumBuffers, Channel.ChannelIndex));//[Channel.ChannelIndex]);
+			}
+		}
+
+		if (NumWindows > 0)
+		{
+			// Two scalings are applied. One for FFT, another for number of windows calculated.
+			float LinearScale = FFTScale / static_cast<float>(NumWindows);
+
+			// This scaling moves output to roughly 0 to 1 range.
+			float DbScale = 1.f / FMath::Max(1.f, -NoiseFloorDb);
+
+			// Apply scaling for each channel.
+			for (int32 ChannelIndex = 0; ChannelIndex < CollectContext.ChannelCount; ChannelIndex++)
+			{
+				//Audio::FAlignedFloatBuffer& Buffer = ChannelSpectrumBuffers[ChannelIndex];
+				TArrayView<float> BufferView = MakeChannelBufferView(PlanarChannelSpectrumBuffers, ChannelIndex);
+
+				Audio::ArrayMultiplyByConstantInPlace(BufferView, LinearScale);
+
+				Audio::ArrayPowerToDecibelInPlace(BufferView, NoiseFloorDb);
+
+				Audio::ArraySubtractByConstantInPlace(BufferView, NoiseFloorDb);
+
+				Audio::ArrayMultiplyByConstantInPlace(BufferView, DbScale);
+			}
+		}
+	}
+
+	// Updates internal objects for num channels and samplerate
+	void SetAudioFormat(const FCollectAudioContext& CollectContext)
+	{
+		if (CollectContext.SampleRate != SampleRate)
+		{
+			ResizeAudioTransform(CollectContext);
+		}
+
+		if (CollectContext.ChannelCount != ChannelCount)
+		{
+			ResetSpectrumBuffers(CollectContext);
+		}
+
+		if ((CollectContext.SampleRate != SampleRate) || (CollectContext.ChannelCount != ChannelCount))
+		{
+			int32 FFTSize = FFTAlgorithm.IsValid() ? FFTAlgorithm->Size() : 0;
+
+			ResizeWindow(CollectContext.ChannelCount, FFTSize);
+
+			ChannelCount = CollectContext.ChannelCount;
+			SampleRate = CollectContext.SampleRate;
+		}
+	}
+
+	void ResizeAudioTransform(const FCollectAudioContext& CollectContext)
+	{
+		using namespace UE::Niagara::GeneratedDataAudioSampling::Private;
+
+		FFTAlgorithm.Reset();
+		CQTKernel.Reset();
+		FFTInputBuffer.Reset();
+		FFTOutputBuffer.Reset();
+		PowerSpectrumBuffer.Reset();
+		SpectrumBuffer.Reset();
+
+		FFTScale = 1.f;
+
+		if (CollectContext.SampleRate <= 0.f)
+		{
+			return;
+		}
+
+		if (CollectContext.SpectrumResolution < 1)
+		{
+			return;
+		}
+
+		float NumOctaves = GetNumOctaves(CollectContext.MinimumFrequency, CollectContext.MaximumFrequency);
+		float NumBandsPerOctave = GetNumBandsPerOctave(CollectContext.SpectrumResolution, NumOctaves);
+
+		Audio::FFFTSettings FFTSettings = GetFFTSettings(CollectContext.MinimumFrequency, CollectContext.SampleRate, NumBandsPerOctave);
+		if (!Audio::FFFTFactory::AreFFTSettingsSupported(FFTSettings))
+		{
+			UE_LOG(LogNiagara, Warning, TEXT("FFT settings are not supported"));
+			return;
+		}
+
+		FFTAlgorithm = Audio::FFFTFactory::NewFFTAlgorithm(FFTSettings);
+
+		if (!FFTAlgorithm.IsValid())
+		{
+			UE_LOG(LogNiagara, Error, TEXT("Failed to create FFT"));
+			return;
+		}
+
+		FFTInputBuffer.AddZeroed(FFTAlgorithm->NumInputFloats());
+		FFTOutputBuffer.AddUninitialized(FFTAlgorithm->NumOutputFloats());
+		PowerSpectrumBuffer.AddUninitialized(FFTOutputBuffer.Num() / 2);
+		SpectrumBuffer.AddUninitialized(CollectContext.SpectrumResolution);
+
+		float FFTSize = static_cast<float>(FFTAlgorithm->Size());
+		check(FFTSize > 0.f);
+
+		// We want to have all FFT implementations
+		// return the same scaling so that the energy conservatino property of the fourier transform
+		// is supported.  This scaling factor is applied to the power spectrum, so we square the 
+		// scaling we would have performed on the magnitude spectrum.
+		switch (FFTAlgorithm->ForwardScaling())
+		{
+		case Audio::EFFTScaling::MultipliedByFFTSize:
+			FFTScale = 1.f / (FFTSize * FFTSize);
+			break;
+
+		case Audio::EFFTScaling::MultipliedBySqrtFFTSize:
+			FFTScale = 1.f / FFTSize;
+			break;
+
+		case Audio::EFFTScaling::DividedByFFTSize:
+			FFTScale = FFTSize * FFTSize;
+			break;
+
+		case Audio::EFFTScaling::DividedBySqrtFFTSize:
+			FFTScale = FFTSize;
+			break;
+
+		default:
+			FFTScale = 1.f;
+			break;
+		}
+
+		float BandwidthStretch = GetBandwidthStretch(CollectContext.SampleRate, FFTSize, NumBandsPerOctave, CollectContext.MinimumFrequency);
+
+		Audio::FPseudoConstantQKernelSettings CQTSettings = GetConstantQSettings(CollectContext.MinimumFrequency, CollectContext.MaximumFrequency, CollectContext.SpectrumResolution, NumBandsPerOctave, BandwidthStretch);
+		CQTKernel = Audio::NewPseudoConstantQKernelTransform(CQTSettings, FFTAlgorithm->Size(), CollectContext.SampleRate);
+
+		if (!CQTKernel.IsValid())
+		{
+			UE_LOG(LogNiagara, Error, TEXT("Failed to create CQT kernel."));
+			return;
+		}
+	}
+
+	void ResetSpectrumBuffers(const FCollectAudioContext& CollectContext)
+	{
+		PlanarChannelSpectrumBuffers.Reset();
+		PlanarChannelSpectrumBuffers.AddZeroed(CollectContext.SpectrumResolution * CollectContext.ChannelCount);
+	}
+
+	void ResizeWindow(int32 InNumChannels, int32 InFFTSize)
+	{
+		int32 NumWindowFrames = FMath::Max(1, InFFTSize);
+
+		NumWindowFrames = FMath::Min(NumWindowFrames, 1024);
+
+		WindowBuffer.Reset();
+		WindowBuffer.AddUninitialized(NumWindowFrames);
+
+		Audio::GenerateBlackmanWindow(WindowBuffer.GetData(), NumWindowFrames, 1, true);
+
+		int32 NumWindowSamples = NumWindowFrames;
+
+		if (InNumChannels > 0)
+		{
+			NumWindowSamples *= InNumChannels;
+		}
+
+		// 50% overlap
+		int32 NumHopSamples = FMath::Max(1, NumWindowSamples / 2);
+
+		SlidingBuffer = MakeUnique<FSlidingBuffer>(NumWindowSamples, NumHopSamples);
+
+	}
+
+	TConstArrayView<float> GetPlanarSpectrumBuffer() const
+	{
+		return MakeConstArrayView(PlanarChannelSpectrumBuffers);
+	}
+
+	// Get settings 
+	static void Clamp(float& InMinimumFrequency, float& InMaximumFrequency, int32& InNumBands)
+	{
+		const int32 MinimumSupportedNumBands = 1;
+		const int32 MaximumSupportedNumBands = 16384;
+		const float MinimumSupportedFrequency = 20.f;
+		const float MaximumSupportedFrequency = 20000.f;
+
+		InMinimumFrequency = FMath::Max(InMinimumFrequency, MinimumSupportedFrequency);
+		InMaximumFrequency = FMath::Min(InMaximumFrequency, MaximumSupportedFrequency);
+		InMaximumFrequency = FMath::Max(InMinimumFrequency, InMaximumFrequency);
+
+		InNumBands = FMath::Clamp(InNumBands, MinimumSupportedNumBands, MaximumSupportedNumBands);
+	}
+
+	static float GetNumOctaves(float InMinimumFrequency, float InMaximumFrequency)
+	{
+		const float MinimumSupportedNumOctaves = 0.01f;
+
+		float NumOctaves = FMath::Log2(InMaximumFrequency) - FMath::Log2(InMinimumFrequency);
+		NumOctaves = FMath::Max(MinimumSupportedNumOctaves, NumOctaves);
+
+		return NumOctaves;
+	}
+
+	static float GetNumBandsPerOctave(int32 InNumBands, float InNumOctaves)
+	{
+		InNumOctaves = FMath::Max(0.01f, InNumOctaves);
+		const float MinimumSupportedNumBandsPerOctave = 0.01f;
+
+		float NumBandsPerOctave = FMath::Max(MinimumSupportedNumBandsPerOctave, static_cast<float>(InNumBands) / InNumOctaves);
+
+		return NumBandsPerOctave;
+	}
+
+	static float GetBandwidthStretch(float InSampleRate, float InFFTSize, float InNumBandsPerOctave, float InMinimumFrequency)
+	{
+		InFFTSize = FMath::Max(1.f, InFFTSize);
+
+		float MinimumFrequencySpacing = (FMath::Pow(2.f, 1.f / InNumBandsPerOctave) - 1.f) * InMinimumFrequency;
+		MinimumFrequencySpacing = FMath::Max(0.01f, MinimumFrequencySpacing);
+
+		float FFTBinSpacing = InSampleRate / InFFTSize;
+		float BandwidthStretch = FFTBinSpacing / MinimumFrequencySpacing;
+
+		return FMath::Clamp(BandwidthStretch, 0.5f, 2.f);
+	}
+
+private:
+	TUniquePtr<FSlidingBuffer> SlidingBuffer;
+	Audio::FAlignedFloatBuffer PlanarChannelSpectrumBuffers;
+	TUniquePtr<Audio::FContiguousSparse2DKernelTransform> CQTKernel;
+	TUniquePtr<Audio::IFFTAlgorithm> FFTAlgorithm;
+	Audio::FAlignedFloatBuffer InterleavedBuffer;
+	Audio::FAlignedFloatBuffer DeinterleavedBuffer;
+	Audio::FAlignedFloatBuffer FFTInputBuffer;
+	Audio::FAlignedFloatBuffer FFTOutputBuffer;
+	Audio::FAlignedFloatBuffer PowerSpectrumBuffer;
+	Audio::FAlignedFloatBuffer SpectrumBuffer;
+	Audio::FAlignedFloatBuffer WindowBuffer;
+
+	int32 ChannelCount = 0;
+	float SampleRate = 0.0f;
+	float FFTScale = 1.0f;
+};
+
+
+class FNDIAudio_SharedResourceImpl : public FNDIAudio_SharedResource
+{
+public:
+	FNDIAudio_SharedResourceImpl() = delete;
+	FNDIAudio_SharedResourceImpl(const FNDIAudio_SharedResource&) = delete;
+	FNDIAudio_SharedResourceImpl(const FResourceKey& InKey)
+		: FNDIAudio_SharedResource(InKey)
+	{
+		AddSubmixListener(InKey.DeviceId);
+	}
+
+	virtual ~FNDIAudio_SharedResourceImpl()
+	{
+		RemoveSubmixListener(ResourceKey.DeviceId);
+	}
+
+	virtual int32 GetChannelCount() const override
+	{
+		int32 ChannelCount = 0;
+
+		if (SubmixListener.IsValid())
+		{
+			ChannelCount = FMath::Max(ChannelCount, SubmixListener->GetChannelCount());
+		}
+
+		return ChannelCount;
+	}
+
+	virtual float GetSampleRate() const override
+	{
+		float SampleRate = 0.0f;
+
+		if (SubmixListener.IsValid())
+		{
+			SampleRate = FMath::Max(SampleRate, SubmixListener->GetSampleRate());
+		}
+
+		return SampleRate;
+	}
+
+	virtual TConstArrayView<float> ReadAudioBuffer(int32& ChannelCount, int32& FrameCount) const override
+	{
+		ChannelCount = GetChannelCount();
+		FrameCount = ResourceKey.AudioResamplingResolution;
+
+		const int32 SampleCount = ResampledAudioBuffer.Num();
+
+		if (SampleCount < (ChannelCount * FrameCount))
+		{
+			ChannelCount = 0;
+			FrameCount = 0;
+			return TConstArrayView<float>();
+		}
+
+		return MakeConstArrayView(ResampledAudioBuffer);
+	}
+
+	virtual TConstArrayView<float> ReadSpectrumBuffer(int32& ChannelCount, int32& FrameCount) const override
+	{
+		ChannelCount = 0;
+		FrameCount = 0;
+
+		if (SpectrumBuilder)
+		{
+			TConstArrayView<float> SpectrumBuffer = SpectrumBuilder->GetPlanarSpectrumBuffer();
+			ChannelCount = GetChannelCount();
+			FrameCount = ResourceKey.SpectrumSamplingResolution;
+
+			if (SpectrumBuffer.Num() < (ChannelCount * FrameCount))
+			{
+				ChannelCount = 0;
+				FrameCount = 0;
+				return TConstArrayView<float>();
+			}
+
+			return SpectrumBuffer;
+		}
+
+		return TConstArrayView<float>();
+	}
+
+	bool ResampleAudioBuffer(int32 TargetResolution, TArray<float>& OutResampledBuffer) const
+	{
+		TRACE_CPUPROFILER_EVENT_SCOPE(NiagaraAudioSampling::ResampleAudioBuffer);
+
+		const int32 ChannelCount = GetChannelCount();
+		const int32 SourceSampleCount = PopBuffer.Num();
+
+		if (!ChannelCount
+			|| !TargetResolution
+			|| !SourceSampleCount)
+		{
+			OutResampledBuffer.Reset();
+			return false;
+		}
+
+		const float SourceSampleRate = GetSampleRate();
+		const int32 TargetSampleCount = TargetResolution * ChannelCount;
+
+		const float SampleRateRatio = (float)(TargetSampleCount) / ((float)SourceSampleCount);
+		const float DestinationSampleRate = SourceSampleRate * SampleRateRatio;
+
+		Audio::FResamplingParameters ResampleParameters = {
+			Audio::EResamplingMethod::Linear,
+			ChannelCount,
+			SourceSampleRate,
+			DestinationSampleRate,
+			const_cast<Audio::FAlignedFloatBuffer&>(PopBuffer)
+		};
+
+		int32 ResampledBufferSize = Audio::GetOutputBufferSize(ResampleParameters);
+
+		Audio::FAlignedFloatBuffer ResampledBuffer;
+		ResampledBuffer.Reset();
+		ResampledBuffer.AddZeroed(ResampledBufferSize);
+
+		Audio::FResamplerResults ResampleResults;
+		ResampleResults.OutBuffer = &ResampledBuffer;
+
+		bool bResampleResult = Audio::Resample(ResampleParameters, ResampleResults);
+		check(bResampleResult);
+
+		OutResampledBuffer = MoveTemp(ResampledBuffer);
+		OutResampledBuffer.SetNumZeroed(TargetSampleCount);
+
+		return true;
+	}
+
+	virtual void CollectAudio() override
+	{
+		constexpr bool bAsyncCollectAudio = true;
+		constexpr bool bAsyncGenerateSpectrum = true;
+
+		if (!SubmixListener)
+		{
+			return;
+		}
+
+		{
+			FSubmixListener::FOptionalEvent Event = SubmixListener->GetInitializedEvent();
+			if (Event.IsSet())
+			{
+				TRACE_CPUPROFILER_EVENT_SCOPE(NiagaraAudioSampling::EventWait);
+				Event->Wait();
+				SubmixListener->ReleaseInitializedEvent();
+			}
+		}
+
+		const bool bIncludeSpectrum = SpectrumCpuRefCount > 0 || SpectrumGpuRefCount > 0;
+		const bool bIncludeAudioResample = AudioCpuRefCount > 0 || AudioGpuRefCount > 0;
+
+		FCollectAudioContext CollectContext;
+		CollectContext.SpectrumResolution = bIncludeSpectrum ? ResourceKey.SpectrumSamplingResolution : 0;
+		CollectContext.ResampledAudioResolution = bIncludeAudioResample ? ResourceKey.AudioResamplingResolution : 0;
+		CollectContext.ChannelCount = GetChannelCount();
+		CollectContext.NoiseFloorDb = ResourceKey.NoiseFloorDb;
+		CollectContext.SampleRate = GetSampleRate();
+		CollectContext.MinimumFrequency = ResourceKey.MinimumFrequency;
+		CollectContext.MaximumFrequency = ResourceKey.MaximumFrequency;
+		CollectContext.bUseLatestAudio = ResourceKey.bUseLatestAudio;
+		CollectContext.bContinuousSampling = ResourceKey.bContinuousSampling;
+		CollectContext.AudioSampleCount = GetRequiredSampleCount(CollectContext.ChannelCount, CollectContext.SampleRate);
+
+		TerminalAudioTask.Reset();
+		if (bAsyncCollectAudio)
+		{
+			CollectAudioTask = UE::Tasks::Launch(UE_SOURCE_LOCATION, [this, CollectContext]
+			{
+				CollectAudioInternal(CollectContext);
+			});
+
+			TerminalAudioTask.Emplace(CollectAudioTask);
+		}
+		else
+		{
+			CollectAudioTask = UE::Tasks::FTask();
+			CollectAudioInternal(CollectContext);
+		}
+
+		if (CollectContext.SpectrumResolution > 0)
+		{
+			if (!SpectrumBuilder)
+			{
+				SpectrumBuilder = MakeUnique<FSpectrumBuilder>();
+			}
+
+			if (bAsyncGenerateSpectrum)
+			{
+				GenerateSpectrumTask = UE::Tasks::Launch(UE_SOURCE_LOCATION, [this, CollectContext]
+				{
+					SpectrumBuilder->Build(CollectContext, PopBuffer);
+				}, UE::Tasks::Prerequisites(CollectAudioTask));
+
+				TerminalAudioTask.Emplace(GenerateSpectrumTask);
+			}
+			else
+			{
+				CollectAudioTask.Wait();
+
+				GenerateSpectrumTask = UE::Tasks::FTask();
+				SpectrumBuilder->Build(CollectContext, PopBuffer);
+			}
+		}
+		else
+		{
+			SpectrumBuilder.Reset();
+		}
+	}
+
+	virtual void WaitForAudio() const override
+	{
+		if (TerminalAudioTask.IsSet())
+		{
+			TRACE_CPUPROFILER_EVENT_SCOPE(NiagaraAudioSampling::WaitForAudio);
+			TerminalAudioTask->Wait();
+		}
+	}
+
+protected:
+	void AddSubmixListener(Audio::FDeviceId InDeviceId)
+	{
+		if (SubmixListener)
+		{
+			check(SubmixListener->GetAudioDeviceId() == InDeviceId);
+		}
+		else
+		{
+			if (FAudioDevice* DeviceHandle = FAudioDeviceManager::Get()->GetAudioDeviceRaw(InDeviceId))
+			{
+				const float DeviceSampleRate = DeviceHandle->GetSampleRate();
+				const int32 MaxChannelCount = ResourceKey.MaxChannelCount ? ResourceKey.MaxChannelCount : AUDIO_MIXER_MAX_OUTPUT_CHANNELS;
+				const int32 MaxBufferSampleCount = GetRequiredSampleCount(MaxChannelCount, DeviceSampleRate);
+				const int32 MinSamplesRequired = 1;
+
+				SubmixListener = MakeShared<FSubmixListener>(PatchMixer, MaxChannelCount, MaxBufferSampleCount, MinSamplesRequired, InDeviceId, ResourceKey.Submix.Get());
+				SubmixListener->RegisterToSubmix();
+			}
+		}
+	}
+
+	void RemoveSubmixListener(Audio::FDeviceId InDeviceId)
+	{
+		if (SubmixListener)
+		{
+			check(SubmixListener->GetAudioDeviceId() == InDeviceId);
+			SubmixListener->UnregisterFromSubmix();
+			SubmixListener.Reset();
+		}
+	}
+
+	void CollectAudioInternal(const FCollectAudioContext& CollectContext)
+	{
+		TRACE_CPUPROFILER_EVENT_SCOPE(NiagaraAudioSampling::CollectAudio);
+
+		if (CollectContext.ChannelCount == 0 || CollectContext.AudioSampleCount == 0)
+		{
+			return;
+		}
+
+		int32 AvailableSampleCount = PatchMixer.MaxNumberOfSamplesThatCanBePopped();
+		AvailableSampleCount -= (AvailableSampleCount % CollectContext.ChannelCount);
+
+		if (AvailableSampleCount > 0)
+		{
+			if (CollectContext.bContinuousSampling)
+			{
+				PopBuffer.Reset();
+				PopBuffer.AddUninitialized(AvailableSampleCount);
+				PatchMixer.PopAudio(PopBuffer.GetData(), AvailableSampleCount, CollectContext.bUseLatestAudio);
+			}
+			// fill an entire buffer
+			else if (AvailableSampleCount >= CollectContext.AudioSampleCount)
+			{
+				PopBuffer.SetNumZeroed(CollectContext.AudioSampleCount);
+				PatchMixer.PopAudio(PopBuffer.GetData(), CollectContext.AudioSampleCount, CollectContext.bUseLatestAudio);
+			}
+			// see if we should append to the existing buffer
+			else if (AvailableSampleCount + PopBuffer.Num() <= CollectContext.AudioSampleCount)
+			{
+				const int32 ExistingSampleCount = PopBuffer.Num();
+				PopBuffer.SetNumZeroed(ExistingSampleCount + AvailableSampleCount);
+				PatchMixer.PopAudio(PopBuffer.GetData() + ExistingSampleCount, AvailableSampleCount, CollectContext.bUseLatestAudio);
+			}
+			// pop out existing data to make room for the new data coming in
+			else
+			{
+				const int32 SaveCount = FMath::Max(0, CollectContext.AudioSampleCount - AvailableSampleCount);
+				TArray<float> SavedSamples(MakeArrayView<float>(PopBuffer.GetData() + PopBuffer.Num() - SaveCount, SaveCount));
+				PopBuffer.Reset();
+				PopBuffer.Append(SavedSamples);
+				PopBuffer.SetNumZeroed(CollectContext.AudioSampleCount);
+				PatchMixer.PopAudio(PopBuffer.GetData() + SaveCount, CollectContext.AudioSampleCount - SaveCount, CollectContext.bUseLatestAudio);
+			}
+
+			if (CollectContext.ResampledAudioResolution)
+			{
+				ResampleAudioBuffer(CollectContext.ResampledAudioResolution, ResampledAudioBuffer);
+			}
+		}
+
+		if (!CollectContext.ResampledAudioResolution)
+		{
+			ResampledAudioBuffer.Reset();
+		}
+	}
+
+	int32 GetRequiredSampleCount(int32 ChannelCount, float SampleRate)
+	{
+		if (ResourceKey.SamplingMethod == ESamplingWindowMethod::ByCount)
+		{
+			return Align(ResourceKey.SamplingWindowCount, AUDIO_NUM_FLOATS_PER_VECTOR_REGISTER);
+		}
+		else if (ResourceKey.SamplingMethod == ESamplingWindowMethod::ByTime)
+		{
+			const int32 SamplesPerChannel = FMath::CeilToInt(SampleRate * (ResourceKey.SamplingWindowInMilliseconds / 1000.0f));
+			return ChannelCount * Align(SamplesPerChannel, AUDIO_NUM_FLOATS_PER_VECTOR_REGISTER);
+		}
+
+		return 0;
+	}
+
+	// Mixer for sending audio
+	Audio::FPatchMixer PatchMixer;
+
+	Audio::FAlignedFloatBuffer PopBuffer;
+	TArray<float> ResampledAudioBuffer;
+
+	TSharedPtr<FSubmixListener> SubmixListener;
+
+	TUniquePtr<FSpectrumBuilder> SpectrumBuilder;
+
+	UE::Tasks::FTask CollectAudioTask;
+	UE::Tasks::FTask GenerateSpectrumTask;
+
+	TOptional<UE::Tasks::FTask> TerminalAudioTask;
+};
+
+}
+
+const ETickingGroup FNDIAudio_GeneratedData::GeneratedDataTickGroup = NiagaraFirstTickGroup;
+
+FNDI_GeneratedData::TypeHash FNDIAudio_GeneratedData::GetTypeHash()
+{
+	static const TypeHash Hash = FCrc::Strihash_DEPRECATED(TEXT("FNDIAudio_GeneratedData"));
+	return Hash;
+}
+
+void FNDIAudio_GeneratedData::Tick(ETickingGroup TickGroup, float DeltaSeconds)
+{
+	if (TickGroup == GeneratedDataTickGroup)
+	{
+		FRWScopeLock WriteLock(SharedResourcesLock, SLT_Write);
+
+		for (FSharedResourceArray::TIterator ResourceIt(SharedResources); ResourceIt; ++ResourceIt)
+		{
+			TSharedPtr<FNDIAudio_SharedResource>& SharedResource = *ResourceIt;
+			if (SharedResource.IsValid() && SharedResource->IsUsed())
+			{
+				SharedResource->CollectAudio();
+			}
+			else
+			{
+				SharedResource.Reset();
+			}
+
+			// clear out any unused entries in SharedResources
+			if (!SharedResource.IsValid())
+			{
+				ResourceIt.RemoveCurrent();
+			}
+		}
+	}
+}
+
+FNDIAudio_SharedResourceHandle FNDIAudio_GeneratedData::GetSharedResource(const FNDIAudio_SharedResourceUsage& Usage, const FResourceDesc& ResourceDesc)
+{
+	using namespace UE::Niagara::GeneratedDataAudioSampling::Private;
+
+	auto FindExistingResource = [this, &ResourceDesc]() -> int32
+	{
+		const int32 ResourceCount = SharedResources.Num();
+		for (int32 ResourceIt = 0; ResourceIt < ResourceCount; ++ResourceIt)
+		{
+			const TSharedPtr<FNDIAudio_SharedResource>& SharedResource = SharedResources[ResourceIt];
+			if (SharedResource.IsValid())
+			{
+				if (SharedResource->GetResourceKey().CanSupport(ResourceDesc))
+				{
+					return ResourceIt;
+				}
+			}
+		}
+
+		return INDEX_NONE;
+	};
+
+	// see if any of our existing resources will suffice
+	{
+		FRWScopeLock ReadLock(SharedResourcesLock, SLT_ReadOnly);
+
+		const int32 ExistingIndex = FindExistingResource();
+		if (ExistingIndex != INDEX_NONE)
+		{
+			return FNDIAudio_SharedResourceHandle(Usage, SharedResources[ExistingIndex], false /*bNeedsDataImmediately*/);
+		}
+	}
+
+	FRWScopeLock ResourceWriteLock(SharedResourcesLock, SLT_Write);
+
+	{
+		const int32 ExistingIndex = FindExistingResource();
+		if (ExistingIndex != INDEX_NONE)
+		{
+			return FNDIAudio_SharedResourceHandle(Usage, SharedResources[ExistingIndex], false /*bNeedsDataImmediately*/);
+		}
+	}
+
+	TSharedPtr<FNDIAudio_SharedResourceImpl> SharedResource = MakeShared<FNDIAudio_SharedResourceImpl>(ResourceDesc);
+	SharedResources.Add(SharedResource);
+	FNDIAudio_SharedResourceHandle Handle(Usage, SharedResource, false /*bNeedsDataImmediately*/);
+
+	SharedResource->CollectAudio();
+
+	return Handle;
+}
+
+FNDIAudio_SharedResource::FNDIAudio_SharedResource(const FResourceKey& InKey)
+	: ResourceKey(InKey)
+{
+}
+
+bool FNDIAudio_SharedResource::IsUsed() const
+{
+	return AudioCpuRefCount > 0
+		|| AudioGpuRefCount > 0
+		|| SpectrumCpuRefCount > 0
+		|| SpectrumGpuRefCount > 0;
+}
+
+void FNDIAudio_SharedResource::RegisterUser(const FNDIAudio_SharedResourceUsage& Usage, bool bNeedsDataImmediately)
+{
+	AudioCpuRefCount += FNDIAudio_SharedResourceUsage::CpuUsageFlags(Usage.AudioUsage) ? 1 : 0;
+	AudioGpuRefCount += FNDIAudio_SharedResourceUsage::GpuUsageFlags(Usage.AudioUsage) ? 1 : 0;
+	SpectrumCpuRefCount += FNDIAudio_SharedResourceUsage::CpuUsageFlags(Usage.SpectrumUsage) ? 1 : 0;
+	SpectrumGpuRefCount += FNDIAudio_SharedResourceUsage::GpuUsageFlags(Usage.SpectrumUsage) ? 1 : 0;
+}
+
+void FNDIAudio_SharedResource::UnregisterUser(const FNDIAudio_SharedResourceUsage& Usage)
+{
+	AudioCpuRefCount -= FNDIAudio_SharedResourceUsage::CpuUsageFlags(Usage.AudioUsage) ? 1 : 0;
+	AudioGpuRefCount -= FNDIAudio_SharedResourceUsage::GpuUsageFlags(Usage.AudioUsage) ? 1 : 0;
+	SpectrumCpuRefCount -= FNDIAudio_SharedResourceUsage::CpuUsageFlags(Usage.SpectrumUsage) ? 1 : 0;
+	SpectrumGpuRefCount -= FNDIAudio_SharedResourceUsage::GpuUsageFlags(Usage.SpectrumUsage) ? 1 : 0;
+}
+
+bool FNDIAudio_SharedResource::FResourceKey::CanSupport(const FResourceKey& Other) const
+{
+	// check the settings that have to match
+	if ((Submix != Other.Submix)
+		|| (bUseLatestAudio != Other.bUseLatestAudio)
+		|| (DeviceId != Other.DeviceId)
+		|| (MaxChannelCount != Other.MaxChannelCount))
+	{
+		return false;
+	}
+
+	if (SamplingMethod != Other.SamplingMethod)
+	{
+		return false;
+	}
+	else if (SamplingMethod == ESamplingWindowMethod::ByCount)
+	{
+		if (SamplingWindowCount != Other.SamplingWindowCount)
+		{
+			return false;
+		}
+	}
+	else if (SamplingMethod == ESamplingWindowMethod::ByTime)
+	{
+		if (SamplingWindowInMilliseconds != Other.SamplingWindowInMilliseconds)
+		{
+			return false;
+		}
+	}
+
+	if (bResampleAudio && Other.bResampleAudio)
+	{
+		if (AudioResamplingResolution != Other.AudioResamplingResolution)
+		{
+			return false;
+		}
+	}
+
+	if (bGenerateSpectrum && Other.bGenerateSpectrum)
+	{
+		if ((MinimumFrequency != Other.MinimumFrequency)
+			|| (MaximumFrequency != Other.MaximumFrequency)
+			|| (NoiseFloorDb != Other.NoiseFloorDb)
+			|| (SpectrumSamplingResolution != Other.SpectrumSamplingResolution))
+		{
+			return false;
+		}
+	}
+
+	return true;
+}
+
diff --git a/Engine/Plugins/FX/Niagara/Source/Niagara/Private/NiagaraGeneratedDataAudioSampling.h b/Engine/Plugins/FX/Niagara/Source/Niagara/Private/NiagaraGeneratedDataAudioSampling.h
new file mode 100644
--- /dev/null
+++ b/Engine/Plugins/FX/Niagara/Source/Niagara/Private/NiagaraGeneratedDataAudioSampling.h
@@ -0,0 +1,152 @@
+// Copyright Epic Games, Inc. All Rights Reserved.
+#pragma once
+
+#include "CoreMinimal.h"
+
+#include "AudioDefines.h"
+#include "NiagaraDataInterface.h"
+#include "Sound/SoundSubmix.h"
+
+class FNDIAudio_SharedResourceUsage : public FNDI_SharedResourceUsage
+{
+public:
+	enum class EAudioBufferUsage { Unused, CpuOnly, GpuOnly, Both };
+	enum class ESpectrumBufferUsage { Unused, CpuOnly, GpuOnly, Both };
+
+	FNDIAudio_SharedResourceUsage(EAudioBufferUsage InAudioUsage = EAudioBufferUsage::Unused, ESpectrumBufferUsage InSpectrumUsage = ESpectrumBufferUsage::Unused)
+		: FNDI_SharedResourceUsage(CpuUsageFlags(InAudioUsage) || CpuUsageFlags(InSpectrumUsage), GpuUsageFlags(InAudioUsage) || GpuUsageFlags(InSpectrumUsage))
+		, AudioUsage(InAudioUsage)
+		, SpectrumUsage(InSpectrumUsage)
+	{
+	}
+
+	FNDIAudio_SharedResourceUsage(const FNDIAudio_SharedResourceUsage& Other)
+		: FNDI_SharedResourceUsage(Other)
+		, AudioUsage(Other.AudioUsage)
+		, SpectrumUsage(Other.SpectrumUsage)
+	{
+	}
+
+	EAudioBufferUsage AudioUsage = EAudioBufferUsage::Unused;
+	ESpectrumBufferUsage SpectrumUsage = ESpectrumBufferUsage::Unused;
+
+	template<typename T>
+	static bool CpuUsageFlags(const T& UsageType)
+	{
+		return UsageType == T::Both || UsageType == T::CpuOnly;
+	}
+
+	template<typename T>
+	static bool GpuUsageFlags(const T& UsageType)
+	{
+		return UsageType == T::Both || UsageType == T::GpuOnly;
+	}
+
+	template<typename T>
+	static void UpdateUsageFlags(bool bCpuUsage, bool bGpuUsage, T& UsageType)
+	{
+		if (bCpuUsage && bGpuUsage)
+		{
+			UsageType = T::Both;
+		}
+		else if (bCpuUsage)
+		{
+			UsageType = T::CpuOnly;
+		}
+		else if (bGpuUsage)
+		{
+			UsageType = T::GpuOnly;
+		}
+		else
+		{
+			UsageType = T::Unused;
+		}
+	}
+};
+
+// SharedResource that can be held by multiple system instances which manages the resources created by this DI
+class FNDIAudio_SharedResource
+{
+public:
+	enum class ESamplingWindowMethod
+	{
+		ByTime, ByCount
+	};
+
+	struct FResourceKey
+	{
+		TWeakObjectPtr<USoundSubmix> Submix = nullptr;
+
+		int32 AudioResamplingResolution = 0;
+		int32 SpectrumSamplingResolution = 0;
+		int32 MaxChannelCount = 0; // If set memory usage for the collection buffer can be reduced
+		float MinimumFrequency = 0.0f;
+		float MaximumFrequency = 0.0f;
+		float NoiseFloorDb = 0.0f;
+		Audio::FDeviceId DeviceId = INDEX_NONE;
+		bool bUseLatestAudio = false;
+		bool bResampleAudio = false;
+		bool bGenerateSpectrum = false;
+		bool bContinuousSampling = false;
+		ESamplingWindowMethod SamplingMethod = ESamplingWindowMethod::ByCount;
+		int32 SamplingWindowCount = 0;
+		float SamplingWindowInMilliseconds = 0.0f;
+
+		bool CanSupport(const FResourceKey& Other) const;
+	};
+
+	FNDIAudio_SharedResource() = delete;
+	FNDIAudio_SharedResource(const FNDIAudio_SharedResource&) = delete;
+	FNDIAudio_SharedResource(const FResourceKey& InKey);
+	virtual ~FNDIAudio_SharedResource() = default;
+
+	virtual int32 GetChannelCount() const = 0;
+	virtual float GetSampleRate() const = 0;
+
+	virtual TConstArrayView<float> ReadAudioBuffer(int32& ChannelCount, int32& FrameCount) const = 0;
+	virtual TConstArrayView<float> ReadSpectrumBuffer(int32& ChannelCount, int32& FrameCount) const = 0;
+
+	const FResourceKey& GetResourceKey() const { return ResourceKey; }
+	bool IsUsed() const;
+
+	void RegisterUser(const FNDIAudio_SharedResourceUsage& Usage, bool bNeedsDataImmediately);
+	void UnregisterUser(const FNDIAudio_SharedResourceUsage& Usage);
+
+	void Release();
+
+	virtual void CollectAudio() = 0;
+	virtual void WaitForAudio() const = 0;
+
+protected:
+
+	FResourceKey ResourceKey;
+
+	std::atomic<int32> AudioCpuRefCount = {0};
+	std::atomic<int32> AudioGpuRefCount = {0};
+	std::atomic<int32> SpectrumCpuRefCount = {0};
+	std::atomic<int32> SpectrumGpuRefCount = {0};
+};
+
+using FNDIAudio_SharedResourceHandle = FNDI_SharedResourceHandle<FNDIAudio_SharedResource, FNDIAudio_SharedResourceUsage>;
+
+// keyed by the requested Submix (if null, that means the default for each given device)
+class FNDIAudio_GeneratedData : public FNDI_GeneratedData
+{
+public:
+	static const ETickingGroup GeneratedDataTickGroup;
+
+	FNDIAudio_GeneratedData() = default;
+	virtual ~FNDIAudio_GeneratedData() = default;
+
+	static TypeHash GetTypeHash();
+
+	NIAGARA_API virtual void Tick(ETickingGroup TickGroup, float DeltaSeconds) override;
+
+	using FResourceDesc = FNDIAudio_SharedResource::FResourceKey;
+	NIAGARA_API FNDIAudio_SharedResourceHandle GetSharedResource(const FNDIAudio_SharedResourceUsage& Usage, const FResourceDesc& ResourceDesc);
+
+private:
+	using FSharedResourceArray = TArray<TSharedPtr<FNDIAudio_SharedResource>>;
+	FSharedResourceArray SharedResources;
+	FRWLock SharedResourcesLock;
+};
