- Add metadata to allow showing top level categories for LW emitters

#tests No tests for this CL as it provides infrastructure

#Synced-CL 48563820

diff --git a/Engine/Plugins/FX/Niagara/Source/NiagaraEditor/Private/ViewModels/Stack/NiagaraStackStatelessCommon.cpp b/Engine/Plugins/FX/Niagara/Source/NiagaraEditor/Private/ViewModels/Stack/NiagaraStackStatelessCommon.cpp
new file mode 100644
--- /dev/null
+++ b/Engine/Plugins/FX/Niagara/Source/NiagaraEditor/Private/ViewModels/Stack/NiagaraStackStatelessCommon.cpp
@@ -0,0 +1,8 @@
+// Copyright Epic Games, Inc. All Rights Reserved.
+
+#include "NiagaraStackStatelessCommon.h"
+
+bool FNiagaraStackStatelessCommon::ShouldHideTopLevelCategories(UStruct* Struct)
+{
+	return Struct ? Struct->HasMetaData(TEXT("ShowTopLevelCategories")) == false : true;
+}
diff --git a/Engine/Plugins/FX/Niagara/Source/NiagaraEditor/Private/ViewModels/Stack/NiagaraStackStatelessCommon.h b/Engine/Plugins/FX/Niagara/Source/NiagaraEditor/Private/ViewModels/Stack/NiagaraStackStatelessCommon.h
new file mode 100644
--- /dev/null
+++ b/Engine/Plugins/FX/Niagara/Source/NiagaraEditor/Private/ViewModels/Stack/NiagaraStackStatelessCommon.h
@@ -0,0 +1,10 @@
+// Copyright Epic Games, Inc. All Rights Reserved.
+
+#pragma once
+
+#include "UObject/Class.h"
+
+namespace FNiagaraStackStatelessCommon
+{
+	bool ShouldHideTopLevelCategories(UStruct* Struct);
+}
diff --git a/Engine/Plugins/FX/Niagara/Source/NiagaraEditor/Private/ViewModels/Stack/NiagaraStackStatelessEmitterGroup.cpp b/Engine/Plugins/FX/Niagara/Source/NiagaraEditor/Private/ViewModels/Stack/NiagaraStackStatelessEmitterGroup.cpp
--- a/Engine/Plugins/FX/Niagara/Source/NiagaraEditor/Private/ViewModels/Stack/NiagaraStackStatelessEmitterGroup.cpp
+++ b/Engine/Plugins/FX/Niagara/Source/NiagaraEditor/Private/ViewModels/Stack/NiagaraStackStatelessEmitterGroup.cpp
@@ -8,6 +8,7 @@
 #include "Stateless/NiagaraDistributionIntPropertyCustomization.h"
 #include "Stateless/NiagaraStatelessEmitter.h"
 #include "ViewModels/Stack/NiagaraStackObject.h"
+#include "ViewModels/Stack/NiagaraStackStatelessCommon.h"
 #include "NiagaraEditorStyle.h"
 
 #include "IDetailTreeNode.h"
@@ -140,9 +141,9 @@ void UNiagaraStackStatelessEmitterObjectItem::RefreshChildrenInternal(const TArr
 		UNiagaraStackObject* StatelessEmitterStackObject = StatelessEmitterStackObjectWeak.Get();
 		if (StatelessEmitterStackObject == nullptr)
 		{
+			const bool bIsTopLevelObject = true;
+			const bool bHideTopLevelCategories = FNiagaraStackStatelessCommon::ShouldHideTopLevelCategories(StatelessEmitter->GetClass());
 			StatelessEmitterStackObject = NewObject<UNiagaraStackObject>(this);
-			bool bIsTopLevelObject = true;
-			bool bHideTopLevelCategories = false;
 			StatelessEmitterStackObject->Initialize(CreateDefaultChildRequiredData(), StatelessEmitter, bIsTopLevelObject, bHideTopLevelCategories, GetStackEditorDataKey());
 			StatelessEmitterStackObject->SetOnFilterDetailNodes(OnFilterDetailNodes, UNiagaraStackObject::EDetailNodeFilterMode::FilterAllNodes);
 			StatelessEmitterStackObject->RegisterInstancedCustomPropertyLayout(UNiagaraStatelessEmitter::StaticClass(), FOnGetDetailCustomizationInstance::CreateStatic(&FNiagaraStatelessEmitterDetails::MakeInstance));
diff --git a/Engine/Plugins/FX/Niagara/Source/NiagaraEditor/Private/ViewModels/Stack/NiagaraStackStatelessEmitterSimulateGroup.cpp b/Engine/Plugins/FX/Niagara/Source/NiagaraEditor/Private/ViewModels/Stack/NiagaraStackStatelessEmitterSimulateGroup.cpp
--- a/Engine/Plugins/FX/Niagara/Source/NiagaraEditor/Private/ViewModels/Stack/NiagaraStackStatelessEmitterSimulateGroup.cpp
+++ b/Engine/Plugins/FX/Niagara/Source/NiagaraEditor/Private/ViewModels/Stack/NiagaraStackStatelessEmitterSimulateGroup.cpp
@@ -19,6 +19,7 @@
 #include "ViewModels/NiagaraSystemViewModel.h"
 #include "ViewModels/Stack/NiagaraStackItemPropertyHeaderValueShared.h"
 #include "ViewModels/Stack/NiagaraStackObject.h"
+#include "ViewModels/Stack/NiagaraStackStatelessCommon.h"
 
 #include UE_INLINE_GENERATED_CPP_BY_NAME(NiagaraStackStatelessEmitterSimulateGroup)
 
@@ -397,8 +398,8 @@ void UNiagaraStackStatelessModuleItem::RefreshChildrenInternal(const TArray<UNia
 		UNiagaraStackObject* ModuleObject = ModuleObjectWeak.Get();
 		if (ModuleObject == nullptr || ModuleObject->GetObject() != StatelessModule)
 		{
-			bool bIsInTopLevelObject = true;
-			bool bHideTopLevelCategories = true;
+			const bool bIsInTopLevelObject = true;
+			const bool bHideTopLevelCategories = FNiagaraStackStatelessCommon::ShouldHideTopLevelCategories(StatelessModule->GetClass());
 			ModuleObject = NewObject<UNiagaraStackObject>(this);
 			ModuleObject->Initialize(CreateDefaultChildRequiredData(), StatelessModule, bIsInTopLevelObject, bHideTopLevelCategories, GetStackEditorDataKey());
 			ModuleObject->SetOnFilterDetailNodes(FNiagaraStackObjectShared::FOnFilterDetailNodes::CreateStatic(&UNiagaraStackStatelessModuleItem::FilterDetailNodes), UNiagaraStackObject::EDetailNodeFilterMode::FilterAllNodes);
diff --git a/Engine/Plugins/FX/Niagara/Source/NiagaraEditor/Private/ViewModels/Stack/NiagaraStackStatelessEmitterSpawnGroup.cpp b/Engine/Plugins/FX/Niagara/Source/NiagaraEditor/Private/ViewModels/Stack/NiagaraStackStatelessEmitterSpawnGroup.cpp
--- a/Engine/Plugins/FX/Niagara/Source/NiagaraEditor/Private/ViewModels/Stack/NiagaraStackStatelessEmitterSpawnGroup.cpp
+++ b/Engine/Plugins/FX/Niagara/Source/NiagaraEditor/Private/ViewModels/Stack/NiagaraStackStatelessEmitterSpawnGroup.cpp
@@ -17,6 +17,7 @@
 #include "ViewModels/NiagaraSystemViewModel.h"
 #include "ViewModels/Stack/NiagaraStackItemPropertyHeaderValueShared.h"
 #include "ViewModels/Stack/NiagaraStackObject.h"
+#include "ViewModels/Stack/NiagaraStackStatelessCommon.h"
 #include "NiagaraClipboard.h"
 
 #include UE_INLINE_GENERATED_CPP_BY_NAME(NiagaraStackStatelessEmitterSpawnGroup)
@@ -427,8 +428,9 @@ void UNiagaraStackStatelessEmitterSpawnItem::RefreshChildrenInternal(const TArra
 		if (SpawnInfoObject == nullptr || SpawnInfoObject->GetObject() != StatelessEmitterObject ||
 			SpawnInfoObject->GetDisplayedStruct().IsValid() == false || SpawnInfoObject->GetDisplayedStruct()->GetStructMemory() != SpawnInfoPtr)
 		{
-			bool bIsInTopLevelStruct = true;
-			bool bHideTopLevelCategories = true;
+			const bool bIsInTopLevelStruct = true;
+			const bool bHideTopLevelCategories = FNiagaraStackStatelessCommon::ShouldHideTopLevelCategories(FNiagaraStatelessSpawnInfo::StaticStruct());
+
 			SpawnInfoObject = NewObject<UNiagaraStackObject>(this);
 			SpawnInfoObject->Initialize(CreateDefaultChildRequiredData(), StatelessEmitterObject, SpawnInfoStructOnScope.ToSharedRef(), TEXT("SpawnInfo"), bIsInTopLevelStruct, bHideTopLevelCategories, GetStackEditorDataKey());
 			SpawnInfoObject->SetOnFilterDetailNodes(
