- Add categories to a few LWE modules

#tests Add a LWE and look at the initialize particle module for visible categories

#Synced-CL 48563820

diff --git a/Engine/Plugins/FX/Niagara/Source/Niagara/Internal/NiagaraSystemEmitterState.h b/Engine/Plugins/FX/Niagara/Source/Niagara/Internal/NiagaraSystemEmitterState.h
--- a/Engine/Plugins/FX/Niagara/Source/Niagara/Internal/NiagaraSystemEmitterState.h
+++ b/Engine/Plugins/FX/Niagara/Source/Niagara/Internal/NiagaraSystemEmitterState.h
@@ -104,67 +104,67 @@ struct FNiagaraEmitterStateData
 	//UPROPERTY(EditAnywhere, Category="Emitter State")
 	//ENiagaraStatelessEmitterState_SelfSystem LifeCycleMode = ENiagaraStatelessEmitterState_SelfSystem::Self;
 
-	UPROPERTY(EditAnywhere, Category="Emitter State", meta=(SegmentedDisplay))
+	UPROPERTY(EditAnywhere, Category="Life Cycle", meta=(SegmentedDisplay))
 	ENiagaraEmitterInactiveResponse InactiveResponse = ENiagaraEmitterInactiveResponse::Complete;
 
-	UPROPERTY(EditAnywhere, Category="Emitter State", meta=(SegmentedDisplay))
+	UPROPERTY(EditAnywhere, Category="Life Cycle", meta=(SegmentedDisplay))
 	ENiagaraLoopBehavior LoopBehavior = ENiagaraLoopBehavior::Infinite;
 
-	UPROPERTY(EditAnywhere, Category = "Emitter State", meta = (ClampMin = "1", EditCondition = "LoopBehavior == ENiagaraLoopBehavior::Multiple", EditConditionHides))
+	UPROPERTY(EditAnywhere, Category = "Life Cycle", meta = (ClampMin = "1", EditCondition = "LoopBehavior == ENiagaraLoopBehavior::Multiple", EditConditionHides))
 	int32 LoopCount = 1;
 
-	UPROPERTY(EditAnywhere, Category = "Emitter State", meta = (EditCondition = "LoopBehavior == ENiagaraLoopBehavior::Once", EditConditionHides, SegmentedDisplay))
+	UPROPERTY(EditAnywhere, Category = "Life Cycle", meta = (EditCondition = "LoopBehavior == ENiagaraLoopBehavior::Once", EditConditionHides, SegmentedDisplay))
 	ENiagaraLoopDurationMode LoopDurationMode = ENiagaraLoopDurationMode::Fixed;
 
-	UPROPERTY(EditAnywhere, Category = "Emitter State", meta = (EditCondition = "LoopBehavior != ENiagaraLoopBehavior::Once || (LoopBehavior == ENiagaraLoopBehavior::Once && LoopDurationMode == ENiagaraLoopDurationMode::Fixed)", EditConditionHides, ClampMin = "0.0", Units="s"))
+	UPROPERTY(EditAnywhere, Category = "Life Cycle", meta = (EditCondition = "LoopBehavior != ENiagaraLoopBehavior::Once || (LoopBehavior == ENiagaraLoopBehavior::Once && LoopDurationMode == ENiagaraLoopDurationMode::Fixed)", EditConditionHides, ClampMin = "0.0", Units="s"))
 	FNiagaraDistributionRangeFloat LoopDuration = FNiagaraDistributionRangeFloat(1.0f);
 
-	UPROPERTY(EditAnywhere, Category = "Emitter State", meta = (ClampMin = "0.0", Units="s", EditCondition="bLoopDelayEnabled"))
+	UPROPERTY(EditAnywhere, Category = "Life Cycle", meta = (ClampMin = "0.0", Units="s", EditCondition="bLoopDelayEnabled"))
 	FNiagaraDistributionRangeFloat LoopDelay = FNiagaraDistributionRangeFloat(0.0f);
 
-	UPROPERTY(EditAnywhere, Category = "Emitter State", meta = (InlineEditConditionToggle))
+	UPROPERTY(EditAnywhere, Category = "Life Cycle", meta = (InlineEditConditionToggle))
 	uint32 bLoopDelayEnabled : 1 = false;
 
-	UPROPERTY(EditAnywhere, Category = "Emitter State", meta = (EditCondition = "LoopBehavior != ENiagaraLoopBehavior::Once", EditConditionHides, DisplayAfter = "LoopDuration"))
+	UPROPERTY(EditAnywhere, Category = "Life Cycle", meta = (EditCondition = "LoopBehavior != ENiagaraLoopBehavior::Once", EditConditionHides, DisplayAfter = "LoopDuration"))
 	uint32 bRecalculateDurationEachLoop : 1 = false;
 
-	UPROPERTY(EditAnywhere, Category = "Emitter State", meta = (EditCondition = "LoopBehavior != ENiagaraLoopBehavior::Once && bLoopDelayEnabled", EditConditionHides))
+	UPROPERTY(EditAnywhere, Category = "Life Cycle", meta = (EditCondition = "LoopBehavior != ENiagaraLoopBehavior::Once && bLoopDelayEnabled", EditConditionHides))
 	uint32 bDelayFirstLoopOnly : 1 = false;
 
-	UPROPERTY(EditAnywhere, Category = "Emitter State", meta = (EditCondition = "LoopBehavior != ENiagaraLoopBehavior::Once && bLoopDelayEnabled && !bDelayFirstLoopOnly", EditConditionHides))
+	UPROPERTY(EditAnywhere, Category = "Life Cycle", meta = (EditCondition = "LoopBehavior != ENiagaraLoopBehavior::Once && bLoopDelayEnabled && !bDelayFirstLoopOnly", EditConditionHides))
 	uint32 bRecalculateDelayEachLoop : 1 = false;
 
-	UPROPERTY(EditAnywhere, Category = "Emitter Scalability")
+	UPROPERTY(EditAnywhere, Category = "Scalability")
 	uint32 bEnableDistanceCulling : 1 = false;
 
-	UPROPERTY(EditAnywhere, Category = "Emitter Scalability", meta = (DisplayAfter = "MaxDistanceReaction"))
+	UPROPERTY(EditAnywhere, Category = "Scalability", meta = (DisplayAfter = "MaxDistanceReaction"))
 	uint32 bEnableVisibilityCulling : 1 = false;
 
-	UPROPERTY(EditAnywhere, Category = "Emitter Scalability", meta = (InlineEditConditionToggle))
+	UPROPERTY(EditAnywhere, Category = "Scalability", meta = (InlineEditConditionToggle))
 	uint32 bMinDistanceEnabled : 1 = false;
 
-	UPROPERTY(EditAnywhere, Category = "Emitter Scalability", meta = (InlineEditConditionToggle))
+	UPROPERTY(EditAnywhere, Category = "Scalability", meta = (InlineEditConditionToggle))
 	uint32 bMaxDistanceEnabled : 1 = false;
 
-	UPROPERTY(EditAnywhere, Category = "Emitter Scalability", meta = (DisplayAfter = "VisibilityCullDelay"))
+	UPROPERTY(EditAnywhere, Category = "Scalability", meta = (DisplayAfter = "VisibilityCullDelay"))
 	uint32 bResetAgeOnAwaken : 1 = false;
 
-	UPROPERTY(EditAnywhere, Category = "Emitter Scalability", meta = (DisplayAfter = "bEnableDistanceCulling", EditCondition = "bMinDistanceEnabled"))
+	UPROPERTY(EditAnywhere, Category = "Scalability", meta = (DisplayAfter = "bEnableDistanceCulling", EditCondition = "bMinDistanceEnabled"))
 	float MinDistance = 0.0f;
 
-	UPROPERTY(EditAnywhere, Category = "Emitter Scalability", meta = (DisplayAfter = "MinDistance"))
+	UPROPERTY(EditAnywhere, Category = "Scalability", meta = (DisplayAfter = "MinDistance"))
 	ENiagaraExecutionStateManagement MinDistanceReaction = ENiagaraExecutionStateManagement::Awaken;
 
-	UPROPERTY(EditAnywhere, Category = "Emitter Scalability", meta = (DisplayAfter = "bEnableDistanceCulling", EditCondition = "bMaxDistanceEnabled"))
+	UPROPERTY(EditAnywhere, Category = "Scalability", meta = (DisplayAfter = "bEnableDistanceCulling", EditCondition = "bMaxDistanceEnabled"))
 	float MaxDistance = 5000.0f;
 
-	UPROPERTY(EditAnywhere, Category = "Emitter Scalability", meta = (DisplayAfter = "MaxDistance"))
+	UPROPERTY(EditAnywhere, Category = "Scalability", meta = (DisplayAfter = "MaxDistance"))
 	ENiagaraExecutionStateManagement MaxDistanceReaction = ENiagaraExecutionStateManagement::SleepAndLetParticlesFinish;
 
-	UPROPERTY(EditAnywhere, Category = "Emitter Scalability", meta = (EditCondition = "bEnableVisibilityCulling", EditConditionHides))
+	UPROPERTY(EditAnywhere, Category = "Scalability", meta = (EditCondition = "bEnableVisibilityCulling", EditConditionHides))
 	ENiagaraExecutionStateManagement VisibilityCullReaction = ENiagaraExecutionStateManagement::SleepAndLetParticlesFinish;
 
-	UPROPERTY(EditAnywhere, Category = "Emitter Scalability", meta = (EditCondition = "bEnableVisibilityCulling", EditConditionHides))
+	UPROPERTY(EditAnywhere, Category = "Scalability", meta = (EditCondition = "bEnableVisibilityCulling", EditConditionHides))
 	float VisibilityCullDelay = 1.0f;
 
 #if WITH_EDITORONLY_DATA
diff --git a/Engine/Plugins/FX/Niagara/Source/Niagara/Internal/Stateless/Modules/NiagaraStatelessModuleCommon.h b/Engine/Plugins/FX/Niagara/Source/Niagara/Internal/Stateless/Modules/NiagaraStatelessModuleCommon.h
--- a/Engine/Plugins/FX/Niagara/Source/Niagara/Internal/Stateless/Modules/NiagaraStatelessModuleCommon.h
+++ b/Engine/Plugins/FX/Niagara/Source/Niagara/Internal/Stateless/Modules/NiagaraStatelessModuleCommon.h
@@ -15,71 +15,71 @@ struct FNiagaraStatelessSystemScaleBuildData
 
 	////////////////////////////////////////////////////////////////////////////////////////////////////////
 	// Enable bits
-	UPROPERTY(EditAnywhere, Category = "Parameters|Camera", meta = (DisplayPriority = 600))
+	UPROPERTY(EditAnywhere, Category = "Camera", meta = (DisplayPriority = 600))
 	uint32 bScaleCameraOffset : 1 = false;
 
 	// Not including drag scale right now as we have separate solver passes which is harder to handle
-	//UPROPERTY(EditAnywhere, Category = "Parameters|Drag", meta = (DisplayPriority = 500))
+	//UPROPERTY(EditAnywhere, Category = "Drag", meta = (DisplayPriority = 500))
 	//uint32 bScaleDrag : 1 = false;
 
-	UPROPERTY(EditAnywhere, Category = "Parameters|Forces", meta = (DisplayPriority = 200))
+	UPROPERTY(EditAnywhere, Category = "Forces", meta = (DisplayPriority = 200))
 	uint32 bScaleForces : 1 = false;
 
 	// Mesh scale in a no-op for local space with standard emitters
 	//UPROPERTY(EditAnywhere, Category = "Parameters", meta = (DisplayPriority = 400))
 	//uint32 bScaleInitialMeshSize : 1 = false;
 
-	UPROPERTY(EditAnywhere, Category = "Parameters|Ribbon", meta = (DisplayPriority = 700))
+	UPROPERTY(EditAnywhere, Category = "Ribbon", meta = (DisplayPriority = 700))
 	uint32 bScaleInitialRibbonWidth : 1 = false;
 
-	UPROPERTY(EditAnywhere, Category = "Parameters|Sprite", meta = (DisplayPriority = 300))
+	UPROPERTY(EditAnywhere, Category = "Sprite", meta = (DisplayPriority = 300))
 	uint32 bScaleInitialSpriteSize : 1 = false;
 
-	UPROPERTY(EditAnywhere, Category = "Parameters|Velocity", meta = (DisplayPriority = 100))
+	UPROPERTY(EditAnywhere, Category = "Velocity", meta = (DisplayPriority = 100))
 	uint32 bScaleInitialVelocity : 1 = false;
 
-	UPROPERTY(EditAnywhere, Category = "Parameters|Decal", meta = (DisplayPriority = 800))
+	UPROPERTY(EditAnywhere, Category = "Decal", meta = (DisplayPriority = 800))
 	uint32 bScaleInitialDecalSize : 1 = false;
 
-	UPROPERTY(EditAnywhere, Category = "Parameters|Lights", meta = (DisplayPriority = 900))
+	UPROPERTY(EditAnywhere, Category = "Lights", meta = (DisplayPriority = 900))
 	uint32 bScaleInitialLightRadius : 1 = false;
 
 	////////////////////////////////////////////////////////////////////////////////////////////////////////
 	// Bitfields
-	UPROPERTY(EditAnywhere, Category = "Parameters|Sprite", meta = (DisplayPriority = 302, EditCondition = "bScaleInitialSpriteSize", EditConditionHides))
+	UPROPERTY(EditAnywhere, Category = "Sprite", meta = (DisplayPriority = 302, EditCondition = "bScaleInitialSpriteSize", EditConditionHides))
 	uint32 bScaleInitialSpriteSizeBySingleAxis : 1 = false;
 
 	//UPROPERTY(EditAnywhere, Category = "Parameters", meta = (DisplayPriority = 403, EditCondition = "bScaleInitialMeshSize", EditConditionHides))
 	//uint32 bScaleInitialMeshSizeAbsoluteScale : 1 = false;
 
-	UPROPERTY(EditAnywhere, Category = "Parameters|Ribbon", meta = (DisplayPriority = 702, EditCondition = "bScaleInitialRibbonWidth", EditConditionHides))
+	UPROPERTY(EditAnywhere, Category = "Ribbon", meta = (DisplayPriority = 702, EditCondition = "bScaleInitialRibbonWidth", EditConditionHides))
 	uint32 bScaleInitialRibbonWidthBySingleAxis : 1 = false;
 
-	//UPROPERTY(EditAnywhere, Category = "Parameters|Drag", meta = (DisplayPriority = 502))
+	//UPROPERTY(EditAnywhere, Category = "Drag", meta = (DisplayPriority = 502))
 	//uint32 bInverseDragScaling : 1 = false;
 
-	UPROPERTY(EditAnywhere, Category = "Parameters|Lights", meta = (DisplayPriority = 902, EditCondition = "bScaleInitialLightRadius", EditConditionHides))
+	UPROPERTY(EditAnywhere, Category = "Lights", meta = (DisplayPriority = 902, EditCondition = "bScaleInitialLightRadius", EditConditionHides))
 	uint32 bScaleInitialLightRadiusBySingleAxis : 1 = false;
 
 	////////////////////////////////////////////////////////////////////////////////////////////////////////
 	// Velocity
-	UPROPERTY(EditAnywhere, Category = "Parameters|Velocity", meta = (DisplayPriority = 101, EditCondition = "bScaleInitialVelocity", EditConditionHides))
+	UPROPERTY(EditAnywhere, Category = "Velocity", meta = (DisplayPriority = 101, EditCondition = "bScaleInitialVelocity", EditConditionHides))
 	FVector3f ScaleInitialVelocityScaleAmount = FVector3f::OneVector;
 
 	////////////////////////////////////////////////////////////////////////////////////////////////////////
 	// Forces
-	UPROPERTY(EditAnywhere, Category = "Parameters|Forces", meta = (DisplayPriority = 201, EditCondition = "bScaleForces", EditConditionHides))
+	UPROPERTY(EditAnywhere, Category = "Forces", meta = (DisplayPriority = 201, EditCondition = "bScaleForces", EditConditionHides))
 	FVector3f ScaleForcesScaleAmount = FVector3f::OneVector;
 
-	UPROPERTY(EditAnywhere, Category = "Parameters|Forces", meta = (DisplayPriority = 202, EditCondition = "bScaleForces", EditConditionHides))
+	UPROPERTY(EditAnywhere, Category = "Forces", meta = (DisplayPriority = 202, EditCondition = "bScaleForces", EditConditionHides))
 	uint32 bScaleForcesInverseScaling : 1 = false;
 
 	////////////////////////////////////////////////////////////////////////////////////////////////////////
 	// Sprites
-	UPROPERTY(EditAnywhere, Category = "Parameters|Sprite", meta = (DisplayPriority = 301, EditCondition = "bScaleInitialSpriteSize", EditConditionHides))
+	UPROPERTY(EditAnywhere, Category = "Sprite", meta = (DisplayPriority = 301, EditCondition = "bScaleInitialSpriteSize", EditConditionHides))
 	float ScaleInitialSpriteSizeScaleAmount = 1.0f;
 
-	UPROPERTY(EditAnywhere, Category = "Parameters|Sprite", meta = (DisplayPriority = 303, EditCondition = "bScaleInitialSpriteSize && bScaleInitialSpriteSizeBySingleAxis", EditConditionHides))
+	UPROPERTY(EditAnywhere, Category = "Sprite", meta = (DisplayPriority = 303, EditCondition = "bScaleInitialSpriteSize && bScaleInitialSpriteSizeBySingleAxis", EditConditionHides))
 	ENiagaraOrientationAxis ScaleInitialSpriteSizeAxis = ENiagaraOrientationAxis::XAxis;
 
 	////////////////////////////////////////////////////////////////////////////////////////////////////////
@@ -92,31 +92,31 @@ struct FNiagaraStatelessSystemScaleBuildData
 
 	////////////////////////////////////////////////////////////////////////////////////////////////////////
 	// Drag
-	//UPROPERTY(EditAnywhere, Category = "Parameters|Drag", meta = (DisplayPriority = 501, EditCondition = "bScaleDrag", EditConditionHides))
+	//UPROPERTY(EditAnywhere, Category = "Drag", meta = (DisplayPriority = 501, EditCondition = "bScaleDrag", EditConditionHides))
 	//float ScaleDragScaleAmount = 1.0f;
 
 	////////////////////////////////////////////////////////////////////////////////////////////////////////
 	// Camera
-	UPROPERTY(EditAnywhere, Category = "Parameters|Camera", meta = (DisplayPriority = 601, EditCondition = "bScaleCameraOffset", EditConditionHides))
+	UPROPERTY(EditAnywhere, Category = "Camera", meta = (DisplayPriority = 601, EditCondition = "bScaleCameraOffset", EditConditionHides))
 	float ScaleCameraOffsetAmount = 1.0f;
 
 	////////////////////////////////////////////////////////////////////////////////////////////////////////
 	// Ribbons
-	UPROPERTY(EditAnywhere, Category = "Parameters|Ribbon", meta = (DisplayPriority = 701, EditCondition = "bScaleInitialRibbonWidth", EditConditionHides))
+	UPROPERTY(EditAnywhere, Category = "Ribbon", meta = (DisplayPriority = 701, EditCondition = "bScaleInitialRibbonWidth", EditConditionHides))
 	float ScaleInitialRibbonWidthScaleAmount = 1.0f;
-	UPROPERTY(EditAnywhere, Category = "Parameters|Ribbon", meta = (DisplayPriority = 703, EditCondition = "bScaleInitialRibbonWidth && bScaleInitialRibbonWidthBySingleAxis", EditConditionHides))
+	UPROPERTY(EditAnywhere, Category = "Ribbon", meta = (DisplayPriority = 703, EditCondition = "bScaleInitialRibbonWidth && bScaleInitialRibbonWidthBySingleAxis", EditConditionHides))
 	ENiagaraOrientationAxis ScaleInitialRibbonWidthScalingAxis = ENiagaraOrientationAxis::XAxis;
 
 	////////////////////////////////////////////////////////////////////////////////////////////////////////
 	// Decals
-	UPROPERTY(EditAnywhere, Category = "Parameters|Decal", meta = (DisplayPriority = 801, EditCondition = "bScaleInitialDecalSize", EditConditionHides))
+	UPROPERTY(EditAnywhere, Category = "Decal", meta = (DisplayPriority = 801, EditCondition = "bScaleInitialDecalSize", EditConditionHides))
 	FVector3f ScaleInitialDecalSizeScaleAmount = FVector3f::OneVector;
 
 	////////////////////////////////////////////////////////////////////////////////////////////////////////
 	// Lights
-	UPROPERTY(EditAnywhere, Category = "Parameters|Lights", meta = (DisplayPriority = 901, EditCondition = "bScaleInitialLightRadius", EditConditionHides))
+	UPROPERTY(EditAnywhere, Category = "Lights", meta = (DisplayPriority = 901, EditCondition = "bScaleInitialLightRadius", EditConditionHides))
 	float ScaleInitialLightRadiusScaleAmount = 1.0f;
-	UPROPERTY(EditAnywhere, Category = "Parameters|Lights", meta = (DisplayPriority = 903, EditCondition = "bScaleInitialLightRadius && bScaleInitialLightRadiusBySingleAxis", EditConditionHides))
+	UPROPERTY(EditAnywhere, Category = "Lights", meta = (DisplayPriority = 903, EditCondition = "bScaleInitialLightRadius && bScaleInitialLightRadiusBySingleAxis", EditConditionHides))
 	ENiagaraOrientationAxis ScaleInitialLightRadiusScalingAxis = ENiagaraOrientationAxis::XAxis;
 };
 
diff --git a/Engine/Plugins/FX/Niagara/Source/Niagara/Internal/Stateless/Modules/NiagaraStatelessModule_ApplyOwnerScaleToAttributes.h b/Engine/Plugins/FX/Niagara/Source/Niagara/Internal/Stateless/Modules/NiagaraStatelessModule_ApplyOwnerScaleToAttributes.h
--- a/Engine/Plugins/FX/Niagara/Source/Niagara/Internal/Stateless/Modules/NiagaraStatelessModule_ApplyOwnerScaleToAttributes.h
+++ b/Engine/Plugins/FX/Niagara/Source/Niagara/Internal/Stateless/Modules/NiagaraStatelessModule_ApplyOwnerScaleToAttributes.h
@@ -8,7 +8,7 @@
 #include "NiagaraStatelessModule_ApplyOwnerScaleToAttributes.generated.h"
 
 // Allows you to control how Engine.Owner.Scale is applied to attributes
-UCLASS(MinimalAPI, EditInlineNew, meta = (DisplayName = "Apply Owner Scale To Attributes"))
+UCLASS(MinimalAPI, EditInlineNew, meta = (DisplayName = "Apply Owner Scale To Attributes", ShowTopLevelCategories))
 class UNiagaraStatelessModule_ApplyOwnerScaleToAttributes : public UNiagaraStatelessModule
 {
 	GENERATED_BODY()
diff --git a/Engine/Plugins/FX/Niagara/Source/Niagara/Internal/Stateless/Modules/NiagaraStatelessModule_InitializeParticle.h b/Engine/Plugins/FX/Niagara/Source/Niagara/Internal/Stateless/Modules/NiagaraStatelessModule_InitializeParticle.h
--- a/Engine/Plugins/FX/Niagara/Source/Niagara/Internal/Stateless/Modules/NiagaraStatelessModule_InitializeParticle.h
+++ b/Engine/Plugins/FX/Niagara/Source/Niagara/Internal/Stateless/Modules/NiagaraStatelessModule_InitializeParticle.h
@@ -11,37 +11,37 @@
 #include "NiagaraStatelessModule_InitializeParticle.generated.h"
 
 // Initialize common particle attributes using common settings and options.
-UCLASS(MinimalAPI, EditInlineNew, meta = (DisplayName = "Initialize Particle"))
+UCLASS(MinimalAPI, EditInlineNew, meta = (DisplayName = "Initialize Particle", ShowTopLevelCategories))
 class UNiagaraStatelessModule_InitializeParticle : public UNiagaraStatelessModule
 {
 	GENERATED_BODY()
 
 public:
-	UPROPERTY(EditAnywhere, Category = "Parameters", meta = (DisplayName = "Lifetime", Units="s"))
+	UPROPERTY(EditAnywhere, Category = "Point Attributes", meta = (DisplayName = "Lifetime", Units="s"))
 	FNiagaraDistributionRangeFloat LifetimeDistribution = FNiagaraDistributionRangeFloat(FNiagaraStatelessGlobals::GetDefaultLifetimeValue());
 
-	UPROPERTY(EditAnywhere, Category = "Parameters", meta = (DisplayName = "Color", DistributionLookupValueEnumPath = "/Script/Niagara.ENiagaraDistributionInitialLookupValueMode"))
+	UPROPERTY(EditAnywhere, Category = "Point Attributes", meta = (DisplayName = "Color", DistributionLookupValueEnumPath = "/Script/Niagara.ENiagaraDistributionInitialLookupValueMode"))
 	FNiagaraDistributionColor ColorDistribution = FNiagaraDistributionColor(FNiagaraStatelessGlobals::GetDefaultColorValue());
 
-	UPROPERTY(EditAnywhere, Category = "Parameters", meta = (DisplayName = "Mass"))
+	UPROPERTY(EditAnywhere, Category = "Point Attributes", meta = (DisplayName = "Mass"))
 	FNiagaraDistributionRangeFloat MassDistribution = FNiagaraDistributionRangeFloat(FNiagaraStatelessGlobals::GetDefaultMassValue());
 
-	UPROPERTY(EditAnywhere, Category = "Parameters", meta = (DisplayName = "Sprite Size"))
+	UPROPERTY(EditAnywhere, Category = "Sprite Attributes", meta = (DisplayName = "Sprite Size"))
 	FNiagaraDistributionRangeVector2 SpriteSizeDistribution = FNiagaraDistributionRangeVector2(FNiagaraStatelessGlobals::GetDefaultSpriteSizeValue());
 
-	UPROPERTY(EditAnywhere, Category = "Parameters", meta = (DisplayName = "Sprite Rotation", Units="deg"))
+	UPROPERTY(EditAnywhere, Category = "Sprite Attributes", meta = (DisplayName = "Sprite Rotation", Units="deg"))
 	FNiagaraDistributionRangeFloat SpriteRotationDistribution = FNiagaraDistributionRangeFloat(FNiagaraStatelessGlobals::GetDefaultSpriteRotationValue());
 
-	UPROPERTY(EditAnywhere, Category = "Parameters", meta = (DisplayName = "Mesh Scale"))
+	UPROPERTY(EditAnywhere, Category = "Mesh Attributes", meta = (DisplayName = "Mesh Scale"))
 	FNiagaraDistributionRangeVector3 MeshScaleDistribution = FNiagaraDistributionRangeVector3(FNiagaraStatelessGlobals::GetDefaultScaleValue());
 
-	UPROPERTY(EditAnywhere, Category = "Parameters", meta = (InlineEditConditionToggle))
+	UPROPERTY(EditAnywhere, Category = "Ribbon Attributes", meta = (InlineEditConditionToggle))
 	bool bWriteRibbonWidth = false;
 
-	UPROPERTY(EditAnywhere, Category = "Parameters", meta = (DisplayName = "Ribbon Width", EditCondition="bWriteRibbonWidth"))
+	UPROPERTY(EditAnywhere, Category = "Ribbon Attributes", meta = (DisplayName = "Ribbon Width", EditCondition="bWriteRibbonWidth"))
 	FNiagaraDistributionRangeFloat RibbonWidthDistribution = FNiagaraDistributionRangeFloat(FNiagaraStatelessGlobals::GetDefaultRibbonWidthValue());
 
-	UPROPERTY(EditAnywhere, Category = "Parameters", meta = (DistributionLookupValueEnumPath = "/Script/Niagara.ENiagaraDistributionInitialLookupValueMode"))
+	UPROPERTY(EditAnywhere, Category = "Point Attributes", meta = (DistributionLookupValueEnumPath = "/Script/Niagara.ENiagaraDistributionInitialLookupValueMode"))
 	FNiagaraDistributionPosition InitialPositionDistribution = FNiagaraDistributionPosition(FVector3f::ZeroVector);
 
 	virtual void BuildEmitterData(const FNiagaraStatelessEmitterDataBuildContext& BuildContext) const override;
diff --git a/Engine/Plugins/FX/Niagara/Source/Niagara/Internal/Stateless/NiagaraStatelessEmitter.h b/Engine/Plugins/FX/Niagara/Source/Niagara/Internal/Stateless/NiagaraStatelessEmitter.h
--- a/Engine/Plugins/FX/Niagara/Source/Niagara/Internal/Stateless/NiagaraStatelessEmitter.h
+++ b/Engine/Plugins/FX/Niagara/Source/Niagara/Internal/Stateless/NiagaraStatelessEmitter.h
@@ -28,7 +28,7 @@ using FNiagaraStatelessEmitterDataPtr = TSharedPtr<const FNiagaraStatelessEmitte
 * Editor data for stateless emitters
 * Generates runtime data to be consumed by the game
 */
-UCLASS(MinimalAPI, EditInlineNew)
+UCLASS(MinimalAPI, EditInlineNew, meta = (ShowTopLevelCategories))
 class UNiagaraStatelessEmitter : public UNiagaraEmitterBase
 {
 	GENERATED_BODY()
@@ -174,10 +174,10 @@ protected:
 	TObjectPtr<UClass> EmitterTemplateClass_DEPRECATED;
 #endif
 
-	UPROPERTY(EditAnywhere, Category = "Emitter Properties")
+	UPROPERTY(EditAnywhere, Category = "Emitter")
 	TObjectPtr<UNiagaraStatelessEmitterTemplate> EmitterTemplate;
 
-	UPROPERTY(EditAnywhere, Category = "Emitter Properties")
+	UPROPERTY(EditAnywhere, Category = "Emitter")
 	uint32 bDeterministic : 1 = false;
 
 #if WITH_EDITORONLY_DATA
@@ -185,24 +185,24 @@ protected:
 	When enabled the emitter will output all available attributes.
 	You should not need to modify this with the exception of debugging / testing and as it will impact cooked performance and memory
 	*/
-	UPROPERTY(EditAnywhere, Category = "Emitter Properties", AdvancedDisplay)
+	UPROPERTY(EditAnywhere, Category = "Emitter", AdvancedDisplay)
 	uint32 bForceOutputAllAttributes : 1 = false;
 
 	/**
 	When enabled the emitter will always include UniqueID in the output attributes.
 	You should not need to modify this with the exception of debugging / testing and as it will impact cooked performance and memory
 	*/
-	UPROPERTY(EditAnywhere, Category = "Emitter Properties", AdvancedDisplay, meta = (EditCondition = "!bForceOutputAllAttributes"))
+	UPROPERTY(EditAnywhere, Category = "Emitter", AdvancedDisplay, meta = (EditCondition = "!bForceOutputAllAttributes"))
 	uint32 bForceOutputUniqueID : 1 = false;
 #endif
 
-	UPROPERTY(EditAnywhere, Category = "Emitter Properties", AdvancedDisplay, meta = (Bitmask, BitMaskEnum = "/Script/Niagara.ENiagaraStatelessFeatureMask"))
+	UPROPERTY(EditAnywhere, Category = "Emitter", AdvancedDisplay, meta = (Bitmask, BitMaskEnum = "/Script/Niagara.ENiagaraStatelessFeatureMask"))
 	uint32 AllowedFeatureMask = uint32(ENiagaraStatelessFeatureMask::All);
 
-	UPROPERTY(EditAnywhere, Category = "Emitter Properties")
+	UPROPERTY(EditAnywhere, Category = "Emitter")
 	int32 RandomSeed = 0;
 
-	UPROPERTY(EditAnywhere, Category = "Emitter Properties")
+	UPROPERTY(EditAnywhere, Category = "Emitter")
 	FBox FixedBounds = FBox(FVector(-100), FVector(100));
 
 	UPROPERTY(EditAnywhere, Category = "Emitter State")
