- Add transform offset to LW shape location module

#tests See shape location ET for an example

#virtualized
#Synced-CL 48563820

diff --git a/Engine/Plugins/FX/Niagara/Shaders/Private/Stateless/Modules/NiagaraStatelessModule_ShapeLocation.ush b/Engine/Plugins/FX/Niagara/Shaders/Private/Stateless/Modules/NiagaraStatelessModule_ShapeLocation.ush
--- a/Engine/Plugins/FX/Niagara/Shaders/Private/Stateless/Modules/NiagaraStatelessModule_ShapeLocation.ush
+++ b/Engine/Plugins/FX/Niagara/Shaders/Private/Stateless/Modules/NiagaraStatelessModule_ShapeLocation.ush
@@ -4,6 +4,8 @@ uint4	ShapeLocation_Mode;
 float4	ShapeLocation_Parameters0;
 float4	ShapeLocation_Parameters1;
 float4	ShapeLocation_Parameters2;
+FNiagaraStatelessBuiltDistributionType	ShapeLocation_TransformOffset;
+float3	ShapeLocation_TransformOffsetScale;
 
 float3 ShapeLocation_GetLocation()
 {
@@ -108,7 +110,8 @@ float3 ShapeLocation_GetLocation()
 
 void ShapeLocation_Simulate(inout FStatelessParticle Particle)
 {
-	float3 Location = ShapeLocation_GetLocation();
-	Particle.Position += Location;
-	Particle.PreviousPosition += Location;
+	const float3 Location = ShapeLocation_GetLocation();
+	const float2x3 PositionOffsets = SampleDistributionValuesFloat3(ShapeLocation_TransformOffset, 1, Particle.NormalizedAge, Particle.PreviousNormalizedAge);
+	Particle.Position += (PositionOffsets[0] * ShapeLocation_TransformOffsetScale) + Location;
+	Particle.PreviousPosition += (PositionOffsets[1] * ShapeLocation_TransformOffsetScale) + Location;
 }
diff --git a/Engine/Plugins/FX/Niagara/Source/Niagara/Internal/Stateless/Modules/NiagaraStatelessModule_ShapeLocation.h b/Engine/Plugins/FX/Niagara/Source/Niagara/Internal/Stateless/Modules/NiagaraStatelessModule_ShapeLocation.h
--- a/Engine/Plugins/FX/Niagara/Source/Niagara/Internal/Stateless/Modules/NiagaraStatelessModule_ShapeLocation.h
+++ b/Engine/Plugins/FX/Niagara/Source/Niagara/Internal/Stateless/Modules/NiagaraStatelessModule_ShapeLocation.h
@@ -27,7 +27,7 @@ enum class ENSM_SurfaceExpansionMode : uint8
 	Outside,
 };
 
-UCLASS(MinimalAPI, EditInlineNew, meta = (DisplayName = "Shape Location"))
+UCLASS(MinimalAPI, EditInlineNew, meta = (DisplayName = "Shape Location", ShowTopLevelCategories))
 class UNiagaraStatelessModule_ShapeLocation : public UNiagaraStatelessModule
 {
 	GENERATED_BODY()
@@ -35,61 +35,67 @@ class UNiagaraStatelessModule_ShapeLocation : public UNiagaraStatelessModule
 public:
 	using FParameters = NiagaraStateless::FShapeLocationModule_ShaderParameters;
 
-	UPROPERTY(EditAnywhere, Category = "Parameters", meta = (ShowInStackItemHeader, StackItemHeaderAlignment = "Left"))
+	UPROPERTY(EditAnywhere, Category = "Shape", meta = (ShowInStackItemHeader, StackItemHeaderAlignment = "Left"))
 	ENSM_ShapePrimitive ShapePrimitive = ENSM_ShapePrimitive::Sphere;
 
-	UPROPERTY(EditAnywhere, Category = "Parameters", meta = (DisableRangeDistribution, EditConditionHides, EditCondition = "ShapePrimitive == ENSM_ShapePrimitive::Box"))
+	UPROPERTY(EditAnywhere, Category = "Shape", meta = (DisableRangeDistribution, EditConditionHides, EditCondition = "ShapePrimitive == ENSM_ShapePrimitive::Box"))
 	FNiagaraDistributionRangeVector3 BoxSize = FNiagaraDistributionRangeVector3(FVector3f(100.0f));
-	UPROPERTY(EditAnywhere, Category = "Parameters", meta = (EditConditionHides, EditCondition = "ShapePrimitive == ENSM_ShapePrimitive::Box"))
+	UPROPERTY(EditAnywhere, Category = "Shape", meta = (EditConditionHides, EditCondition = "ShapePrimitive == ENSM_ShapePrimitive::Box"))
 	bool bBoxSurfaceOnly = false;
-	UPROPERTY(EditAnywhere, Category = "Parameters", meta = (EditConditionHides, EditCondition = "ShapePrimitive == ENSM_ShapePrimitive::Box && bBoxSurfaceOnly"))
+	UPROPERTY(EditAnywhere, Category = "Shape", meta = (EditConditionHides, EditCondition = "ShapePrimitive == ENSM_ShapePrimitive::Box && bBoxSurfaceOnly"))
 	ENSM_SurfaceExpansionMode BoxSurfaceExpansion = ENSM_SurfaceExpansionMode::Centered;
-	UPROPERTY(EditAnywhere, Category = "Parameters", meta = (DisableRangeDistribution, EditConditionHides, EditCondition = "ShapePrimitive == ENSM_ShapePrimitive::Box && bBoxSurfaceOnly"))
+	UPROPERTY(EditAnywhere, Category = "Shape", meta = (DisableRangeDistribution, EditConditionHides, EditCondition = "ShapePrimitive == ENSM_ShapePrimitive::Box && bBoxSurfaceOnly"))
 	FNiagaraDistributionRangeFloat BoxSurfaceThickness = FNiagaraDistributionRangeFloat(0.0f);
 
-	UPROPERTY(EditAnywhere, Category = "Parameters", meta = (DisableRangeDistribution, EditConditionHides, EditCondition = "ShapePrimitive == ENSM_ShapePrimitive::Plane"))
+	UPROPERTY(EditAnywhere, Category = "Shape", meta = (DisableRangeDistribution, EditConditionHides, EditCondition = "ShapePrimitive == ENSM_ShapePrimitive::Plane"))
 	FNiagaraDistributionRangeVector2 PlaneSize = FNiagaraDistributionRangeVector2(FVector2f(100.0f));
-	UPROPERTY(EditAnywhere, Category = "Parameters", meta = (EditConditionHides, EditCondition = "ShapePrimitive == ENSM_ShapePrimitive::Plane"))
+	UPROPERTY(EditAnywhere, Category = "Shape", meta = (EditConditionHides, EditCondition = "ShapePrimitive == ENSM_ShapePrimitive::Plane"))
 	bool bPlaneEdgesOnly = false;
-	UPROPERTY(EditAnywhere, Category = "Parameters", meta = (EditConditionHides, EditCondition = "ShapePrimitive == ENSM_ShapePrimitive::Plane && bPlaneEdgesOnly"))
+	UPROPERTY(EditAnywhere, Category = "Shape", meta = (EditConditionHides, EditCondition = "ShapePrimitive == ENSM_ShapePrimitive::Plane && bPlaneEdgesOnly"))
 	ENSM_SurfaceExpansionMode PlaneEdgeExpansion = ENSM_SurfaceExpansionMode::Centered;
-	UPROPERTY(EditAnywhere, Category = "Parameters", meta = (DisableRangeDistribution, EditConditionHides, EditCondition = "ShapePrimitive == ENSM_ShapePrimitive::Plane && bPlaneEdgesOnly"))
+	UPROPERTY(EditAnywhere, Category = "Shape", meta = (DisableRangeDistribution, EditConditionHides, EditCondition = "ShapePrimitive == ENSM_ShapePrimitive::Plane && bPlaneEdgesOnly"))
 	FNiagaraDistributionRangeFloat PlaneEdgeThickness = FNiagaraDistributionRangeFloat(0.0f);
 
-	UPROPERTY(EditAnywhere, Category = "Parameters", meta = (DisableRangeDistribution, EditConditionHides, EditCondition = "ShapePrimitive == ENSM_ShapePrimitive::Cylinder"))
+	UPROPERTY(EditAnywhere, Category = "Shape", meta = (DisableRangeDistribution, EditConditionHides, EditCondition = "ShapePrimitive == ENSM_ShapePrimitive::Cylinder"))
 	FNiagaraDistributionRangeFloat CylinderHeight = FNiagaraDistributionRangeFloat(100.0f);
-	UPROPERTY(EditAnywhere, Category = "Parameters", meta = (DisableRangeDistribution, EditConditionHides, EditCondition = "ShapePrimitive == ENSM_ShapePrimitive::Cylinder"))
+	UPROPERTY(EditAnywhere, Category = "Shape", meta = (DisableRangeDistribution, EditConditionHides, EditCondition = "ShapePrimitive == ENSM_ShapePrimitive::Cylinder"))
 	FNiagaraDistributionRangeFloat CylinderRadius = FNiagaraDistributionRangeFloat(100.0f);
-	UPROPERTY(EditAnywhere, Category = "Parameters", meta = (DisableRangeDistribution, EditConditionHides, EditCondition = "ShapePrimitive == ENSM_ShapePrimitive::Cylinder"))
+	UPROPERTY(EditAnywhere, Category = "Shape", meta = (DisableRangeDistribution, EditConditionHides, EditCondition = "ShapePrimitive == ENSM_ShapePrimitive::Cylinder"))
 	FNiagaraDistributionRangeFloat CylinderHeightMidpoint = FNiagaraDistributionRangeFloat(0.5f);
-	UPROPERTY(EditAnywhere, Category = "Parameters", meta = (EditConditionHides, EditCondition = "ShapePrimitive == ENSM_ShapePrimitive::Cylinder"))
+	UPROPERTY(EditAnywhere, Category = "Shape", meta = (EditConditionHides, EditCondition = "ShapePrimitive == ENSM_ShapePrimitive::Cylinder"))
 	bool bCylinderSurfaceOnly = false;
-	UPROPERTY(EditAnywhere, Category = "Parameters", meta = (EditConditionHides, EditCondition = "ShapePrimitive == ENSM_ShapePrimitive::Cylinder && bCylinderSurfaceOnly"))
+	UPROPERTY(EditAnywhere, Category = "Shape", meta = (EditConditionHides, EditCondition = "ShapePrimitive == ENSM_ShapePrimitive::Cylinder && bCylinderSurfaceOnly"))
 	bool bCylinderSurfaceOnlyIncludeEndCaps = true;
-	UPROPERTY(EditAnywhere, Category = "Parameters", meta = (EditConditionHides, EditCondition = "ShapePrimitive == ENSM_ShapePrimitive::Cylinder && bCylinderSurfaceOnly"))
+	UPROPERTY(EditAnywhere, Category = "Shape", meta = (EditConditionHides, EditCondition = "ShapePrimitive == ENSM_ShapePrimitive::Cylinder && bCylinderSurfaceOnly"))
 	ENSM_SurfaceExpansionMode CylinderSurfaceExpansion = ENSM_SurfaceExpansionMode::Centered;
-	UPROPERTY(EditAnywhere, Category = "Parameters", meta = (DisableRangeDistribution, EditConditionHides, EditCondition = "ShapePrimitive == ENSM_ShapePrimitive::Cylinder && bCylinderSurfaceOnly"))
+	UPROPERTY(EditAnywhere, Category = "Shape", meta = (DisableRangeDistribution, EditConditionHides, EditCondition = "ShapePrimitive == ENSM_ShapePrimitive::Cylinder && bCylinderSurfaceOnly"))
 	FNiagaraDistributionRangeFloat CylinderSurfaceThickness = FNiagaraDistributionRangeFloat(0.0f);
 
-	UPROPERTY(EditAnywhere, Category = "Parameters", meta = (DisableRangeDistribution, EditConditionHides, EditCondition = "ShapePrimitive == ENSM_ShapePrimitive::Ring"))
+	UPROPERTY(EditAnywhere, Category = "Shape", meta = (DisableRangeDistribution, EditConditionHides, EditCondition = "ShapePrimitive == ENSM_ShapePrimitive::Ring"))
 	FNiagaraDistributionRangeFloat RingRadius = FNiagaraDistributionRangeFloat(100.0f);
-	UPROPERTY(EditAnywhere, Category = "Parameters", meta = (DisableRangeDistribution, EditConditionHides, EditCondition = "ShapePrimitive == ENSM_ShapePrimitive::Ring", ClampMin = "0.0", ClampMax = "1.0", UIMin = "0.0", UIMax = "1.0"))
+	UPROPERTY(EditAnywhere, Category = "Shape", meta = (DisableRangeDistribution, EditConditionHides, EditCondition = "ShapePrimitive == ENSM_ShapePrimitive::Ring", ClampMin = "0.0", ClampMax = "1.0", UIMin = "0.0", UIMax = "1.0"))
 	FNiagaraDistributionRangeFloat DiscCoverage = FNiagaraDistributionRangeFloat(0.0f);
-	UPROPERTY(EditAnywhere, Category = "Parameters", meta = (DisableRangeDistribution, EditConditionHides, EditCondition = "ShapePrimitive == ENSM_ShapePrimitive::Ring", ClampMin = "0.0", ClampMax = "1.0", UIMin = "0.0", UIMax = "1.0"))
+	UPROPERTY(EditAnywhere, Category = "Shape", meta = (DisableRangeDistribution, EditConditionHides, EditCondition = "ShapePrimitive == ENSM_ShapePrimitive::Ring", ClampMin = "0.0", ClampMax = "1.0", UIMin = "0.0", UIMax = "1.0"))
 	FNiagaraDistributionRangeFloat RingUDistribution = FNiagaraDistributionRangeFloat(0.0f);
 
-	UPROPERTY(EditAnywhere, Category = "Parameters", meta=(EditConditionHides, EditCondition="ShapePrimitive == ENSM_ShapePrimitive::Sphere"))
+	UPROPERTY(EditAnywhere, Category = "Shape", meta=(EditConditionHides, EditCondition="ShapePrimitive == ENSM_ShapePrimitive::Sphere"))
 	FNiagaraDistributionRangeFloat SphereRadius = FNiagaraDistributionRangeFloat(100.0f);
 
-	UPROPERTY(EditAnywhere, Category = "Parameters", meta = (SegmentedDisplay))
+	UPROPERTY(EditAnywhere, Category = "Shape", meta = (SegmentedDisplay))
 	ENiagaraCoordinateSpace CoordinateSpace = ENiagaraCoordinateSpace::Simulation;
 
-	UPROPERTY(EditAnywhere, Category = "Parameters", meta = (DisableRangeDistribution, DisableUniformDistribution))
+	UPROPERTY(EditAnywhere, Category = "Shape", meta = (DisableRangeDistribution, DisableUniformDistribution))
 	FNiagaraDistributionRangeRotator ShapeRotation;
 
-	UPROPERTY(EditAnywhere, Category = "Parameters", meta = (DisableRangeDistribution, DisableUniformDistribution))
+	UPROPERTY(EditAnywhere, Category = "Shape", meta = (DisableRangeDistribution, DisableUniformDistribution))
 	FNiagaraDistributionRangeVector3 ShapeScale = FNiagaraDistributionRangeVector3(FVector3f::OneVector);
 
+	UPROPERTY(EditAnywhere, Category = "Transform")
+	FNiagaraDistributionPosition TransformOffset = FNiagaraDistributionPosition(FVector3f::ZeroVector);
+
+	UPROPERTY(EditAnywhere, Category = "Transform", meta = (DisableRangeDistribution))
+	FNiagaraDistributionRangeVector3 TransformOffsetScale = FNiagaraDistributionRangeVector3(FVector3f::OneVector);
+
 private:
 #if WITH_EDITORONLY_DATA
 	UPROPERTY()
diff --git a/Engine/Plugins/FX/Niagara/Source/Niagara/Private/Stateless/Modules/NiagaraStatelessModule_ShapeLocation.cpp b/Engine/Plugins/FX/Niagara/Source/Niagara/Private/Stateless/Modules/NiagaraStatelessModule_ShapeLocation.cpp
--- a/Engine/Plugins/FX/Niagara/Source/Niagara/Private/Stateless/Modules/NiagaraStatelessModule_ShapeLocation.cpp
+++ b/Engine/Plugins/FX/Niagara/Source/Niagara/Private/Stateless/Modules/NiagaraStatelessModule_ShapeLocation.cpp
@@ -16,6 +16,8 @@ namespace NSMShapeLocationPrivate
 		ENiagaraCoordinateSpace			CoordinateSpace = ENiagaraCoordinateSpace::Simulation;
 		FNiagaraStatelessRangeRotator	ShapeRotation;
 		FNiagaraStatelessRangeVector3	ShapeScale;
+		FUintVector3					TransformOffset = FUintVector3::ZeroValue;
+		FNiagaraStatelessRangeVector3	TransformOffsetScale;
 		int32							PositionVariableOffset = INDEX_NONE;
 		int32							PreviousPositionVariableOffset = INDEX_NONE;
 	};
@@ -194,14 +196,20 @@ namespace NSMShapeLocationPrivate
 		const FModuleBuiltData* ModuleBuiltData = ParticleSimulationContext.ReadBuiltData<FModuleBuiltData>();
 		const UNiagaraStatelessModule_ShapeLocation::FParameters* ShaderParameters = ParticleSimulationContext.ReadParameterNestedStruct<UNiagaraStatelessModule_ShapeLocation::FParameters>();
 
+		const float* ParticlesAge = ParticleSimulationContext.GetParticleNormalizedAge();
+		const float* ParticlesPreviousAge = ParticleSimulationContext.GetParticlePreviousNormalizedAge();
+
 		for (uint32 i = 0; i < ParticleSimulationContext.GetNumInstances(); ++i)
 		{
 			const FVector3f ShapeLocation = ShapeLocation_GetLocation(ParticleSimulationContext, ShaderParameters, i);
 			const FVector3f Position = ParticleSimulationContext.ReadParticleVariable(ModuleBuiltData->PositionVariableOffset, i, FVector3f::ZeroVector);
 			const FVector3f PreviousPosition = ParticleSimulationContext.ReadParticleVariable(ModuleBuiltData->PreviousPositionVariableOffset, i, FVector3f::ZeroVector);
 
-			ParticleSimulationContext.WriteParticleVariable(ModuleBuiltData->PositionVariableOffset, i, Position + ShapeLocation);
-			ParticleSimulationContext.WriteParticleVariable(ModuleBuiltData->PreviousPositionVariableOffset, i, PreviousPosition + ShapeLocation);
+			FVector3f PositionOffsets[2];
+			ParticleSimulationContext.SampleDistributionValues(ShaderParameters->ShapeLocation_TransformOffset, i, 1, ParticlesAge[i], ParticlesPreviousAge[i], PositionOffsets[0], PositionOffsets[1]);
+
+			ParticleSimulationContext.WriteParticleVariable(ModuleBuiltData->PositionVariableOffset, i, Position + (PositionOffsets[0] * ShaderParameters->ShapeLocation_TransformOffsetScale) + ShapeLocation);
+			ParticleSimulationContext.WriteParticleVariable(ModuleBuiltData->PreviousPositionVariableOffset, i, PreviousPosition + (PositionOffsets[1] * ShaderParameters->ShapeLocation_TransformOffsetScale) + ShapeLocation);
 		}
 	}
 }
@@ -224,10 +232,12 @@ void UNiagaraStatelessModule_ShapeLocation::BuildEmitterData(const FNiagaraState
 		return;
 	}
 
-	BuiltData->ShapePrimitive	= ShapePrimitive;
-	BuiltData->CoordinateSpace	= CoordinateSpace;
-	BuiltData->ShapeRotation	= BuildContext.ConvertDistributionToRange(ShapeRotation, FRotator3f::ZeroRotator);
-	BuiltData->ShapeScale		= BuildContext.ConvertDistributionToRange(ShapeScale, FVector3f::OneVector);
+	BuiltData->ShapePrimitive		= ShapePrimitive;
+	BuiltData->CoordinateSpace		= CoordinateSpace;
+	BuiltData->ShapeRotation		= BuildContext.ConvertDistributionToRange(ShapeRotation, FRotator3f::ZeroRotator);
+	BuiltData->ShapeScale			= BuildContext.ConvertDistributionToRange(ShapeScale, FVector3f::OneVector);
+	BuiltData->TransformOffset		= BuildContext.AddDistribution(TransformOffset);
+	BuiltData->TransformOffsetScale	= BuildContext.ConvertDistributionToRange(TransformOffsetScale, FVector3f::OneVector);
 
 	switch (ShapePrimitive)
 	{
@@ -301,10 +311,12 @@ void UNiagaraStatelessModule_ShapeLocation::SetShaderParameters(const FNiagaraSt
 
 	const FVector3f Scale = SetShaderParameterContext.ConvertRangeToValue(ModuleBuiltData->ShapeScale);
 
-	Parameters->ShapeLocation_Mode			= FUintVector4::ZeroValue;
-	Parameters->ShapeLocation_Parameters0	= FVector4f::Zero();
-	Parameters->ShapeLocation_Parameters1	= FVector4f::Zero();
-	Parameters->ShapeLocation_Parameters2	= FVector4f::Zero();
+	Parameters->ShapeLocation_Mode					= FUintVector4::ZeroValue;
+	Parameters->ShapeLocation_Parameters0			= FVector4f::Zero();
+	Parameters->ShapeLocation_Parameters1			= FVector4f::Zero();
+	Parameters->ShapeLocation_Parameters2			= FVector4f::Zero();
+	Parameters->ShapeLocation_TransformOffset		= ModuleBuiltData->TransformOffset;
+	Parameters->ShapeLocation_TransformOffsetScale	= SetShaderParameterContext.ConvertRangeToValue(ModuleBuiltData->TransformOffsetScale);
 
 	switch (ModuleBuiltData->ShapePrimitive)
 	{
diff --git a/Engine/Plugins/FX/Niagara/Source/NiagaraShader/Internal/Stateless/NiagaraStatelessModuleShaderParameters.h b/Engine/Plugins/FX/Niagara/Source/NiagaraShader/Internal/Stateless/NiagaraStatelessModuleShaderParameters.h
--- a/Engine/Plugins/FX/Niagara/Source/NiagaraShader/Internal/Stateless/NiagaraStatelessModuleShaderParameters.h
+++ b/Engine/Plugins/FX/Niagara/Source/NiagaraShader/Internal/Stateless/NiagaraStatelessModuleShaderParameters.h
@@ -147,6 +147,8 @@ namespace NiagaraStateless
 		SHADER_PARAMETER(FVector4f,		ShapeLocation_Parameters0)
 		SHADER_PARAMETER(FVector4f,		ShapeLocation_Parameters1)
 		SHADER_PARAMETER(FVector4f,		ShapeLocation_Parameters2)
+		SHADER_PARAMETER(FNiagaraStatelessBuiltDistributionType, ShapeLocation_TransformOffset)
+		SHADER_PARAMETER(FVector3f,		ShapeLocation_TransformOffsetScale)
 	END_SHADER_PARAMETER_STRUCT()
 
 	BEGIN_SHADER_PARAMETER_STRUCT(FSolveVelocitiesAndForcesModule_ShaderParameters, NIAGARASHADER_API)
