Niagara - Add a new feature to allow the niagara clipboard to handle pasting objects copied from outside of the niagara editor and fix some minor copy/paste issues.
+ Implement a handler to allow pasting float and vector curves copied from cascade.
+ Fix an issue where curves pasted onto stateless emitter distrobutions would not update the displayed curve.
+ Fix an issue where pasting a curve onto a module input could not be undone.

#Synced-CL 48514417

diff --git a/Engine/Plugins/FX/Niagara/Source/NiagaraEditor/Private/NiagaraClipboard.cpp b/Engine/Plugins/FX/Niagara/Source/NiagaraEditor/Private/NiagaraClipboard.cpp
--- a/Engine/Plugins/FX/Niagara/Source/NiagaraEditor/Private/NiagaraClipboard.cpp
+++ b/Engine/Plugins/FX/Niagara/Source/NiagaraEditor/Private/NiagaraClipboard.cpp
@@ -381,8 +381,30 @@ UNiagaraClipboardContent* FNiagaraClipboard::ClipboardContentFromString(const FS
 		// If the clipboard text wasn't a niagara clipboard object, it's likely been copied from elsewhere in the editor and may be 
 		// a valid portable value so we construct an empty clipboard object and set the portable data.
 		UNiagaraClipboardContent* ClipboardContent = UNiagaraClipboardContent::Create();
-		FNiagaraClipboardPortableValue& PortableValue = ClipboardContent->PortableValues.AddDefaulted_GetRef();
-		PortableValue.ValueString = ClipboardText;
+
+		// Check to see if the clipboard string can represent an external object reference and if so, check to see if we have a registered
+		// handler.
+		FSoftObjectPath ClipboardObjectPath = FSoftObjectPath(ClipboardText);
+		UObject* ClipboardObject = ClipboardObjectPath.ResolveObject();
+		if (ClipboardObject != nullptr)
+		{
+			TSharedPtr<const INiagaraEditorExternalObjectUtilities, ESPMode::ThreadSafe> ExternalObjectUtilities = FNiagaraEditorModule::Get().GetExternalObjectUtilities(*ClipboardObject->GetClass());
+			FNiagaraClipboardPortableValue ObjectValue;
+			if (ExternalObjectUtilities.IsValid() &&
+				ExternalObjectUtilities->SupportsClipboardPortableValues() &&
+				ExternalObjectUtilities->TryUpdateClipboardPortableValueFromObject(*ClipboardObject, ObjectValue))
+			{
+				ClipboardContent->PortableValues.Add(ObjectValue);
+			}
+		}
+
+		// If it was not handled by an external object handler, just put the raw clipboard text in the portable values.
+		if (ClipboardContent->PortableValues.Num() == 0)
+		{
+			FNiagaraClipboardPortableValue& PortableValue = ClipboardContent->PortableValues.AddDefaulted_GetRef();
+			PortableValue.ValueString = ClipboardText;
+		}
+
 		return ClipboardContent;
 	}
 
diff --git a/Engine/Plugins/FX/Niagara/Source/NiagaraEditor/Private/NiagaraEditorModule.cpp b/Engine/Plugins/FX/Niagara/Source/NiagaraEditor/Private/NiagaraEditorModule.cpp
--- a/Engine/Plugins/FX/Niagara/Source/NiagaraEditor/Private/NiagaraEditorModule.cpp
+++ b/Engine/Plugins/FX/Niagara/Source/NiagaraEditor/Private/NiagaraEditorModule.cpp
@@ -41,6 +41,7 @@
 #include "TypeEditorUtilities/NiagaraMatrixTypeEditorUtilities.h"
 #include "TypeEditorUtilities/NiagaraDataInterfaceCurveTypeEditorUtilities.h"
 #include "TypeEditorUtilities/NiagaraDistributionPropertyEditorUtilities.h"
+#include "TypeEditorUtilities/NiagaraCascadeDistributionExternalObjectUtilities.h"
 
 #include "NiagaraActorPlacement.h"
 #include "NiagaraSystemCompilingManager.h"
@@ -166,6 +167,7 @@
 #include "DeviceProfiles/DeviceProfileManager.h"
 #include "Containers/Ticker.h"
 #include "UObject/AssetRegistryTagsContext.h"
+#include "Distributions/Distribution.h"
 
 #include "ViewModels/Stack/NiagaraStackObjectIssueGenerator.h"
 #include "NiagaraPlatformSet.h"
@@ -1518,6 +1520,8 @@ void FNiagaraEditorModule::StartupModule()
 	RegisterPropertyUtilities(FNiagaraDistributionCurveFloat::StaticStruct(), MakeShared<FNiagaraDistributionPropertyEditorUtilities, ESPMode::ThreadSafe>());
 	RegisterPropertyUtilities(FNiagaraDistributionCurveVector3::StaticStruct(), MakeShared<FNiagaraDistributionPropertyEditorUtilities, ESPMode::ThreadSafe>());
 
+	RegisterExternalObjectUtilities(*UDistribution::StaticClass(), true, MakeShared<FNiagaraCascadeDistributionExternalObjectUtilities, ESPMode::ThreadSafe>());
+
 	FEdGraphUtilities::RegisterVisualPinFactory(GraphPanelPinFactory);
 
 	FNiagaraOpInfo::Init();
@@ -1991,16 +1995,27 @@ void FNiagaraEditorModule::OnNiagaraSettingsChangedEvent(const FName& PropertyNa
 
 void FNiagaraEditorModule::RegisterTypeUtilities(FNiagaraTypeDefinition Type, TSharedRef<INiagaraEditorTypeUtilities, ESPMode::ThreadSafe> EditorUtilities)
 {
-	TypeEditorsCS.Lock();
+	UE::TScopeLock ScopeLock(EditorUtilitiesGuard);
 	TypeToEditorUtilitiesMap.Add(Type, EditorUtilities);
-	TypeEditorsCS.Unlock();
 }
 
 void FNiagaraEditorModule::RegisterPropertyUtilities(const UScriptStruct* InStruct, TSharedRef<INiagaraEditorPropertyUtilities, ESPMode::ThreadSafe> InPropertyUtilities)
 {
-	TypeEditorsCS.Lock();
+	UE::TScopeLock ScopeLock(EditorUtilitiesGuard);
 	StructToPropertyUtilitiesMap.Add(InStruct, InPropertyUtilities);
-	TypeEditorsCS.Unlock();
+}
+
+void FNiagaraEditorModule::RegisterExternalObjectUtilities(const UClass& InClass, bool bHandleSubclasses, TSharedRef<INiagaraEditorExternalObjectUtilities, ESPMode::ThreadSafe> InExternalObjectUtilities)
+{
+	UE::TScopeLock ScopeLock(EditorUtilitiesGuard);
+	if (RegisteredExternalObjectUtilities.ContainsByPredicate([&InClass]
+		(const FRegisteredExternalObjectUtilities& RegisteredUtilities) { return RegisteredUtilities.Class == &InClass; }) == false)
+	{
+		FRegisteredExternalObjectUtilities& AddedUtilities = RegisteredExternalObjectUtilities.AddDefaulted_GetRef();
+		AddedUtilities.Class = &InClass;
+		AddedUtilities.bHandleSubclasses = bHandleSubclasses;
+		AddedUtilities.ExternalObjectUtilities = InExternalObjectUtilities;
+	}
 }
 
 void FNiagaraEditorModule::RegisterModuleWizards(TSharedRef<UE::Niagara::Wizard::FModuleWizardGenerator> WizardGenerator)
@@ -2008,15 +2023,17 @@ void FNiagaraEditorModule::RegisterModuleWizards(TSharedRef<UE::Niagara::Wizard:
 	ModuleWizards.Add(WizardGenerator);
 }
 
-TSharedPtr<INiagaraEditorTypeUtilities, ESPMode::ThreadSafe> FNiagaraEditorModule::GetTypeUtilities(const FNiagaraTypeDefinition& Type)
+TSharedPtr<INiagaraEditorTypeUtilities, ESPMode::ThreadSafe> FNiagaraEditorModule::GetTypeUtilities(const FNiagaraTypeDefinition& Type) const
 {
-	TypeEditorsCS.Lock();
-	TSharedRef<INiagaraEditorTypeUtilities, ESPMode::ThreadSafe>* EditorUtilities = TypeToEditorUtilitiesMap.Find(Type);
-	TypeEditorsCS.Unlock();
+	const TSharedRef<INiagaraEditorTypeUtilities, ESPMode::ThreadSafe>* TypeUtilities = nullptr;
+	{
+		UE::TScopeLock ScopeLock(EditorUtilitiesGuard);
+		TypeUtilities = TypeToEditorUtilitiesMap.Find(Type);
+	}
 
-	if(EditorUtilities != nullptr)
+	if(TypeUtilities != nullptr)
 	{
-		return *EditorUtilities;
+		return *TypeUtilities;
 	}
 
 	if (Type.IsEnum())
@@ -2027,15 +2044,29 @@ TSharedPtr<INiagaraEditorTypeUtilities, ESPMode::ThreadSafe> FNiagaraEditorModul
 	return TSharedPtr<INiagaraEditorTypeUtilities, ESPMode::ThreadSafe>();
 }
 
-TSharedPtr<INiagaraEditorPropertyUtilities, ESPMode::ThreadSafe> FNiagaraEditorModule::GetPropertyUtilities(const UScriptStruct& Struct)
+TSharedPtr<INiagaraEditorPropertyUtilities, ESPMode::ThreadSafe> FNiagaraEditorModule::GetPropertyUtilities(const UScriptStruct& Struct) const
 {
-	TypeEditorsCS.Lock();
-	TSharedRef<INiagaraEditorPropertyUtilities, ESPMode::ThreadSafe>* PropertyUtilities = StructToPropertyUtilitiesMap.Find(&Struct);
-	TypeEditorsCS.Unlock();
-
+	const TSharedRef<INiagaraEditorPropertyUtilities, ESPMode::ThreadSafe>* PropertyUtilities = nullptr;
+	{
+		UE::TScopeLock ScopeLock(EditorUtilitiesGuard);
+		PropertyUtilities = StructToPropertyUtilitiesMap.Find(&Struct);
+	}
 	return PropertyUtilities != nullptr ? *PropertyUtilities : TSharedPtr<INiagaraEditorPropertyUtilities>();
 }
 
+TSharedPtr<const INiagaraEditorExternalObjectUtilities, ESPMode::ThreadSafe> FNiagaraEditorModule::GetExternalObjectUtilities(const UClass& Class) const
+{
+	const FRegisteredExternalObjectUtilities* MatchingRegisteredUtilities = nullptr;
+	{
+		UE::TScopeLock ScopeLock(EditorUtilitiesGuard);
+		MatchingRegisteredUtilities = RegisteredExternalObjectUtilities.FindByPredicate([&Class](const FRegisteredExternalObjectUtilities& RegisteredUtilities)
+		{
+			return RegisteredUtilities.Class == &Class || (RegisteredUtilities.bHandleSubclasses && Class.IsChildOf(RegisteredUtilities.Class));
+		});
+	}
+	return MatchingRegisteredUtilities != nullptr ? MatchingRegisteredUtilities->ExternalObjectUtilities : TSharedPtr<const INiagaraEditorExternalObjectUtilities>();
+}
+
 void FNiagaraEditorModule::RegisterWidgetProvider(TSharedRef<INiagaraEditorWidgetProvider> InWidgetProvider)
 {
 	checkf(WidgetProvider.IsValid() == false, TEXT("Widget provider has already been set."));
diff --git a/Engine/Plugins/FX/Niagara/Source/NiagaraEditor/Private/Stateless/NiagaraDistributionPropertyCustomization.cpp b/Engine/Plugins/FX/Niagara/Source/NiagaraEditor/Private/Stateless/NiagaraDistributionPropertyCustomization.cpp
--- a/Engine/Plugins/FX/Niagara/Source/NiagaraEditor/Private/Stateless/NiagaraDistributionPropertyCustomization.cpp
+++ b/Engine/Plugins/FX/Niagara/Source/NiagaraEditor/Private/Stateless/NiagaraDistributionPropertyCustomization.cpp
@@ -70,8 +70,6 @@ public:
 
 	virtual bool IsValid() const override { return PropertyHandle.IsValid() && SourceDistribution != nullptr; }
 
-	virtual TSharedPtr<IPropertyHandle> GetPropertyHandle() const override { return PropertyHandle; }
-
 	virtual int32 GetNumChannels() const override { return EditorNumChannels; }
 
 	virtual FText GetChannelDisplayName(int32 ChannelIndex) const override 
@@ -451,7 +449,9 @@ public:
 			ContinuousTransactionIndex = GEditor->BeginTransaction(TransactionText);
 		}
 
+		TGuardValue<bool> UpdateValueGuard(bUpdatingDistributionValueInternally, true);
 		const FScopedTransaction Transaction(TransactionText, ContinuousTransactionIndex.IsSet() == false);
+
 		OwnerObject->Modify();
 		PropertyHandle->NotifyPreChange();
 
@@ -512,7 +512,9 @@ public:
 			ContinuousTransactionIndex = GEditor->BeginTransaction(TransactionText);
 		}
 
+		TGuardValue<bool> UpdateValueGuard(bUpdatingDistributionValueInternally, true);
 		const FScopedTransaction Transaction(TransactionText, ContinuousTransactionIndex.IsSet() == false);
+
 		OwnerObject->Modify();
 		PropertyHandle->NotifyPreChange();
 
@@ -571,7 +573,9 @@ public:
 				ContinuousTransactionIndex = GEditor->BeginTransaction(TransactionText);
 			}
 
+			TGuardValue<bool> UpdateValueGuard(bUpdatingDistributionValueInternally, true);
 			const FScopedTransaction Transaction(TransactionText, ContinuousTransactionIndex.IsSet() == false);
+
 			OwnerObject->Modify();
 			PropertyHandle->NotifyPreChange();
 
@@ -674,7 +678,9 @@ public:
 				ContinuousTransactionIndex = GEditor->BeginTransaction(TransactionText);
 			}
 
+			TGuardValue<bool> UpdateValueGuard(bUpdatingDistributionValueInternally, true);
 			const FScopedTransaction Transaction(TransactionText, ContinuousTransactionIndex.IsSet() == false && bShouldTransact);
+
 			if (bShouldTransact)
 			{
 				OwnerObject->Modify();
@@ -746,6 +752,14 @@ public:
 		return SourceDistribution->AllowLookupValueMode(PropertyHandle ? PropertyHandle->GetProperty() : nullptr);
 	}
 
+	virtual const UEnum* GetLookupValueModeEnum() const override
+	{
+		FProperty* DistributionProperty = PropertyHandle->GetProperty();
+		return DistributionProperty != nullptr
+			? FNiagaraDistributionBase::GetLookupValueModeEnum(DistributionProperty)
+			: nullptr;
+	}
+
 	virtual uint8 GetLookupValueMode() const override
 	{
 		return SourceDistribution->GetLookupValueMode();
@@ -792,6 +806,20 @@ public:
 		PropertyHandle->NotifyFinishedChangingProperties();
 	}
 
+	virtual FSimpleMulticastDelegate& OnDistributionValueChanged()
+	{
+		return OnDistributionValueChangedDelegate;
+	}
+
+	virtual void GetOwners(TArray<UObject*>& OutOwners) const
+	{
+		UObject* OwnerObject = OwnerObjectWeak.Get();
+		if (OwnerObject != nullptr)
+		{
+			OutOwners.AddUnique(OwnerObject);
+		}
+	}
+
 	virtual void ModifyOwners() override
 	{
 		UObject* OwnerObject = OwnerObjectWeak.Get();
@@ -812,6 +840,11 @@ private:
 			EditorNumChannels = FNiagaraDistributionEditorUtilities::IsUniform(EditorMode) ? 1 : SourceNumChannels;
 			OnDistributionEditorModeChangedDelegate.Broadcast();
 		}
+
+		if (bUpdatingDistributionValueInternally == false)
+		{
+			OnDistributionValueChangedDelegate.Broadcast();
+		}
 	}
 
 	void MigrateDataFromModeChange(ENiagaraDistributionMode NewMode, int32 NewEditorNumChannels)
@@ -1055,8 +1088,12 @@ private:
 	bool bContinuousChangeActive = false;
 
 	bool bDisplayAsColor = false;
+
 	FSimpleMulticastDelegate OnDistributionEditorModeChangedDelegate;
 	bool bUpdatingDistributionModeInternally = false;
+
+	FSimpleMulticastDelegate OnDistributionValueChangedDelegate;
+	bool bUpdatingDistributionValueInternally = false;
 };
 
 class FNiagaraDistributionAdapterRotator : public FNiagaraDistributionAdapter
diff --git a/Engine/Plugins/FX/Niagara/Source/NiagaraEditor/Private/TypeEditorUtilities/NiagaraCascadeDistributionExternalObjectUtilities.cpp b/Engine/Plugins/FX/Niagara/Source/NiagaraEditor/Private/TypeEditorUtilities/NiagaraCascadeDistributionExternalObjectUtilities.cpp
new file mode 100644
--- /dev/null
+++ b/Engine/Plugins/FX/Niagara/Source/NiagaraEditor/Private/TypeEditorUtilities/NiagaraCascadeDistributionExternalObjectUtilities.cpp
@@ -0,0 +1,100 @@
+// Copyright Epic Games, Inc. All Rights Reserved.
+
+#include "TypeEditorUtilities/NiagaraCascadeDistributionExternalObjectUtilities.h"
+
+#include "Distributions/DistributionFloat.h"
+#include "Distributions/DistributionVector.h"
+#include "Distributions/DistributionFloatConstantCurve.h"
+#include "Distributions/DistributionVectorConstantCurve.h"
+#include "NiagaraClipboard.h"
+
+namespace NiagaraCascadeDistributionUtilities
+{
+	ERichCurveInterpMode RichCurveInterpModeFromLegacyInterpMode(EInterpCurveMode LegacyInterpMode)
+	{
+		switch (LegacyInterpMode)
+		{
+		case CIM_Constant:
+			return ERichCurveInterpMode::RCIM_Constant;
+		case CIM_Linear:
+			return ERichCurveInterpMode::RCIM_Linear;
+		case CIM_CurveAuto:
+		case CIM_CurveUser:
+		case CIM_CurveBreak:
+		case CIM_CurveAutoClamped:
+			return ERichCurveInterpMode::RCIM_Cubic;
+		default:
+			return ERichCurveInterpMode::RCIM_None;
+		}
+	}
+
+	ERichCurveTangentMode RichCurveTangentModeFromLegacyInterpMode(EInterpCurveMode LegacyInterpMode)
+	{
+		switch (LegacyInterpMode)
+		{
+		case CIM_Linear:
+		case CIM_Constant:
+			return ERichCurveTangentMode::RCTM_None;
+		case CIM_CurveAuto:
+			return ERichCurveTangentMode::RCTM_Auto;
+		case CIM_CurveUser:
+			return ERichCurveTangentMode::RCTM_User;
+		case CIM_CurveBreak:
+			return ERichCurveTangentMode::RCTM_Break;
+		case CIM_CurveAutoClamped:
+			return ERichCurveTangentMode::RCTM_SmartAuto;
+		default:
+			return ERichCurveTangentMode::RCTM_None;
+		}
+	}
+}
+
+bool FNiagaraCascadeDistributionExternalObjectUtilities::TryUpdateClipboardPortableValueFromObject(const UObject& InExternalObject, FNiagaraClipboardPortableValue& InTargetClipboardPortableValue) const
+{
+	auto AddKeyToCurve = [](FRichCurve& TargetCurve, ERichCurveInterpMode InterpMode, ERichCurveTangentMode TangentMode, float Time, float Value, float ArriveTangent, float LeaveTangent)
+	{
+		FKeyHandle KeyHandle = TargetCurve.AddKey(Time, Value);
+		FRichCurveKey& AddedKey = TargetCurve.GetKey(KeyHandle);
+		AddedKey.InterpMode = InterpMode;
+		AddedKey.TangentMode = TangentMode;
+		AddedKey.ArriveTangent = ArriveTangent;
+		AddedKey.LeaveTangent = LeaveTangent;
+	};
+
+	const UDistributionFloatConstantCurve* FloatCurve = Cast<UDistributionFloatConstantCurve>(&InExternalObject);
+	if (FloatCurve != nullptr)
+	{
+		FNiagaraClipboardCurveCollection CurveCollection;
+		FRichCurve& Curve = CurveCollection.Curves.AddDefaulted_GetRef();
+		for (const FInterpCurvePointFloat& CurvePoint : FloatCurve->ConstantCurve.Points)
+		{
+			ERichCurveInterpMode InterpMode = NiagaraCascadeDistributionUtilities::RichCurveInterpModeFromLegacyInterpMode(CurvePoint.InterpMode);
+			ERichCurveTangentMode TangentMode = NiagaraCascadeDistributionUtilities::RichCurveTangentModeFromLegacyInterpMode(CurvePoint.InterpMode);
+			AddKeyToCurve(Curve, InterpMode, TangentMode, CurvePoint.InVal, CurvePoint.OutVal, CurvePoint.ArriveTangent, CurvePoint.LeaveTangent);
+		}
+		InTargetClipboardPortableValue = FNiagaraClipboardPortableValue::CreateFromStructValue(CurveCollection);
+		return true;
+	}
+	else
+	{
+		const UDistributionVectorConstantCurve* VectorCurve = Cast<UDistributionVectorConstantCurve>(&InExternalObject);
+		if (VectorCurve != nullptr)
+		{
+			FNiagaraClipboardCurveCollection CurveCollection;
+			FRichCurve& XCurve = CurveCollection.Curves.AddDefaulted_GetRef();
+			FRichCurve& YCurve = CurveCollection.Curves.AddDefaulted_GetRef();
+			FRichCurve& ZCurve = CurveCollection.Curves.AddDefaulted_GetRef();
+			for (const FInterpCurvePointVector& CurvePoint : VectorCurve->ConstantCurve.Points)
+			{
+				ERichCurveInterpMode InterpMode = NiagaraCascadeDistributionUtilities::RichCurveInterpModeFromLegacyInterpMode(CurvePoint.InterpMode);
+				ERichCurveTangentMode TangentMode = NiagaraCascadeDistributionUtilities::RichCurveTangentModeFromLegacyInterpMode(CurvePoint.InterpMode);
+				AddKeyToCurve(XCurve, InterpMode, TangentMode, CurvePoint.InVal, CurvePoint.OutVal.X, CurvePoint.ArriveTangent.X, CurvePoint.LeaveTangent.X);
+				AddKeyToCurve(YCurve, InterpMode, TangentMode, CurvePoint.InVal, CurvePoint.OutVal.Y, CurvePoint.ArriveTangent.Y, CurvePoint.LeaveTangent.Y);
+				AddKeyToCurve(ZCurve, InterpMode, TangentMode, CurvePoint.InVal, CurvePoint.OutVal.Z, CurvePoint.ArriveTangent.Z, CurvePoint.LeaveTangent.Z);
+			}
+			InTargetClipboardPortableValue = FNiagaraClipboardPortableValue::CreateFromStructValue(CurveCollection);
+			return true;
+		}
+	}
+	return false;
+}
\ No newline at end of file
diff --git a/Engine/Plugins/FX/Niagara/Source/NiagaraEditor/Private/TypeEditorUtilities/NiagaraCascadeDistributionExternalObjectUtilities.h b/Engine/Plugins/FX/Niagara/Source/NiagaraEditor/Private/TypeEditorUtilities/NiagaraCascadeDistributionExternalObjectUtilities.h
new file mode 100644
--- /dev/null
+++ b/Engine/Plugins/FX/Niagara/Source/NiagaraEditor/Private/TypeEditorUtilities/NiagaraCascadeDistributionExternalObjectUtilities.h
@@ -0,0 +1,14 @@
+// Copyright Epic Games, Inc. All Rights Reserved.
+
+#pragma once
+
+#include "INiagaraEditorTypeUtilities.h"
+
+class UObject;
+
+class FNiagaraCascadeDistributionExternalObjectUtilities : public FNiagaraEditorExternalObjectUtilities
+{
+public:
+	virtual bool SupportsClipboardPortableValues() const override { return true; }
+	virtual bool TryUpdateClipboardPortableValueFromObject(const UObject& InExternalObject, FNiagaraClipboardPortableValue& InTargetClipboardPortableValue) const override;
+};
\ No newline at end of file
diff --git a/Engine/Plugins/FX/Niagara/Source/NiagaraEditor/Private/TypeEditorUtilities/NiagaraDataInterfaceCurveTypeEditorUtilities.cpp b/Engine/Plugins/FX/Niagara/Source/NiagaraEditor/Private/TypeEditorUtilities/NiagaraDataInterfaceCurveTypeEditorUtilities.cpp
--- a/Engine/Plugins/FX/Niagara/Source/NiagaraEditor/Private/TypeEditorUtilities/NiagaraDataInterfaceCurveTypeEditorUtilities.cpp
+++ b/Engine/Plugins/FX/Niagara/Source/NiagaraEditor/Private/TypeEditorUtilities/NiagaraDataInterfaceCurveTypeEditorUtilities.cpp
@@ -98,7 +98,7 @@ bool FNiagaraDataInterfaceCurveTypeEditorUtilitiesBase::TryUpdateClipboardPortab
 		{
 			CurveCollection.Curves.Add(*CurveDataItem.Curve);
 		}
-		InTargetClipboardPortableValue = FNiagaraClipboardPortableValue::CreateFromStructValue(*FNiagaraClipboardCurveCollection::StaticStruct(), (uint8*)&CurveCollection);
+		InTargetClipboardPortableValue = FNiagaraClipboardPortableValue::CreateFromStructValue(CurveCollection);
 		return true;
 	}
 	return false;
@@ -106,18 +106,16 @@ bool FNiagaraDataInterfaceCurveTypeEditorUtilitiesBase::TryUpdateClipboardPortab
 
 bool FNiagaraDataInterfaceCurveTypeEditorUtilitiesBase::CanUpdateTypedValueFromClipboardPortableValue(const FNiagaraClipboardPortableValue& InSourceClipboardPortableValue, const FNiagaraTypeDefinition& InTargetType) const
 {
-	UScriptStruct* CurveCollectionStruct = FNiagaraClipboardCurveCollection::StaticStruct();
 	FNiagaraClipboardCurveCollection CurveCollection;
-	return InSourceClipboardPortableValue.TryUpdateStructValue(*CurveCollectionStruct, (uint8*)&CurveCollection);
+	return InSourceClipboardPortableValue.TryUpdateStructValue(CurveCollection);
 }
 
 bool FNiagaraDataInterfaceCurveTypeEditorUtilitiesBase::TryUpdateTypedValueFromClipboardPortableValue(const FNiagaraClipboardPortableValue& InSourceClipboardPortableValue, const FNiagaraTypeDefinition& InTargetType, FNiagaraVariant& InTargetValue) const
 {
 	if (InTargetType.GetClass()->IsChildOf(UNiagaraDataInterfaceCurveBase::StaticClass()))
 	{
-		UScriptStruct* CurveCollectionStruct = FNiagaraClipboardCurveCollection::StaticStruct();
 		FNiagaraClipboardCurveCollection CurveCollection;
-		if (InSourceClipboardPortableValue.TryUpdateStructValue(*CurveCollectionStruct, (uint8*)&CurveCollection))
+		if (InSourceClipboardPortableValue.TryUpdateStructValue(CurveCollection))
 		{
 			UNiagaraDataInterfaceCurveBase* CurveDataInterface = NewObject<UNiagaraDataInterfaceCurveBase>(GetTransientPackage(), InTargetType.GetClass());
 			TArray<UNiagaraDataInterfaceCurveBase::FCurveData> CurveData;
diff --git a/Engine/Plugins/FX/Niagara/Source/NiagaraEditor/Private/ViewModels/Stack/NiagaraStackFunctionInput.cpp b/Engine/Plugins/FX/Niagara/Source/NiagaraEditor/Private/ViewModels/Stack/NiagaraStackFunctionInput.cpp
--- a/Engine/Plugins/FX/Niagara/Source/NiagaraEditor/Private/ViewModels/Stack/NiagaraStackFunctionInput.cpp
+++ b/Engine/Plugins/FX/Niagara/Source/NiagaraEditor/Private/ViewModels/Stack/NiagaraStackFunctionInput.cpp
@@ -3180,6 +3180,7 @@ void UNiagaraStackFunctionInput::SetValueFromClipboardFunctionInput(const UNiaga
 			UNiagaraDataInterface* InputDataInterface = GetDataValueObject();
 			if (ensureMsgf(InputDataInterface != nullptr, TEXT("Data interface paste failed.  Current data value object null even after reset.")))
 			{
+				InputDataInterface->Modify();
 				ClipboardFunctionInput.Data->CopyTo(InputDataInterface);
 			}
 			break;
@@ -3193,6 +3194,7 @@ void UNiagaraStackFunctionInput::SetValueFromClipboardFunctionInput(const UNiaga
 
 			if (UNiagaraNodeInput* InputNode = InputValues.ObjectAssetInputNode.Get())
 			{
+				InputNode->Modify();
 				InputNode->SetObjectAsset(ClipboardFunctionInput.ObjectAsset);
 				RefreshChildren();
 				RefreshValues();
diff --git a/Engine/Plugins/FX/Niagara/Source/NiagaraEditor/Private/Widgets/INiagaraDistributionAdapter.h b/Engine/Plugins/FX/Niagara/Source/NiagaraEditor/Private/Widgets/INiagaraDistributionAdapter.h
--- a/Engine/Plugins/FX/Niagara/Source/NiagaraEditor/Private/Widgets/INiagaraDistributionAdapter.h
+++ b/Engine/Plugins/FX/Niagara/Source/NiagaraEditor/Private/Widgets/INiagaraDistributionAdapter.h
@@ -57,9 +57,6 @@ public:
 	/** Returns whether or not this adapter valid for use */
 	virtual bool IsValid() const = 0;
 
-	/** Get the underlying property we are modifying. */
-	virtual TSharedPtr<class IPropertyHandle> GetPropertyHandle() const = 0;
-
 	/** Gets the number of channels currently supported by this adapter. This can change with different distribution modes. */
 	virtual int32 GetNumChannels() const = 0;
 
@@ -157,6 +154,8 @@ public:
 	virtual uint8 GetLookupValueMode() const = 0;
 	/** Set the lookup value mode on the distribution */
 	virtual void SetLookupValueMode(uint8 Mode) = 0;
+	/** Gets the UEnum to be used in lookup value mode. */
+	virtual const UEnum* GetLookupValueModeEnum() const = 0;
 
 	/** Called to notify the distribution that a continuous change has begun. */
 	virtual void BeginContinuousChange() = 0;
@@ -167,6 +166,12 @@ public:
 	/** Called to notify the distribution that a continuous change has been cancelled. */
 	virtual void CancelContinuousChange() = 0;
 
+	/** A delegate that is called whenever the value of the distribution is changed. */
+	virtual FSimpleMulticastDelegate& OnDistributionValueChanged() = 0;
+
+	/** Gets the owning objects for this distribution. */
+	virtual void GetOwners(TArray<UObject*>& OutOwners) const = 0;
+
 	/** Called to modify the distribution that an external editor has started a 
 		transaction and that owning objects should be modified to take part in the transaction. */
 	virtual void ModifyOwners() = 0;
diff --git a/Engine/Plugins/FX/Niagara/Source/NiagaraEditor/Private/Widgets/SNiagaraDistributionArrayEditor.cpp b/Engine/Plugins/FX/Niagara/Source/NiagaraEditor/Private/Widgets/SNiagaraDistributionArrayEditor.cpp
--- a/Engine/Plugins/FX/Niagara/Source/NiagaraEditor/Private/Widgets/SNiagaraDistributionArrayEditor.cpp
+++ b/Engine/Plugins/FX/Niagara/Source/NiagaraEditor/Private/Widgets/SNiagaraDistributionArrayEditor.cpp
@@ -304,8 +304,7 @@ namespace SNiagaraDistributionArrayEditorPrivate
 				TSharedPtr<INiagaraEditorTypeUtilities, ESPMode::ThreadSafe> TypeEditorUtilities = NiagaraEditorModule.GetTypeUtilities(TypeDef);
 				if (TypeEditorUtilities.IsValid())
 				{
-					TSharedPtr<IPropertyHandle> PropertyHandle = DistributionAdapter->GetPropertyHandle();
-					PropertyHandle->SetOnPropertyValueChanged(FSimpleDelegate::CreateSP(this, &SArrayEntryWidget::OnExternalChange));
+					DistributionAdapter->OnDistributionValueChanged().AddSP(this, &SArrayEntryWidget::OnExternalChange);
 
 					ElementData = GetArrayElement(DistributionAdapter, *RowIndexPtr);
 
diff --git a/Engine/Plugins/FX/Niagara/Source/NiagaraEditor/Private/Widgets/SNiagaraDistributionCurveEditor.cpp b/Engine/Plugins/FX/Niagara/Source/NiagaraEditor/Private/Widgets/SNiagaraDistributionCurveEditor.cpp
--- a/Engine/Plugins/FX/Niagara/Source/NiagaraEditor/Private/Widgets/SNiagaraDistributionCurveEditor.cpp
+++ b/Engine/Plugins/FX/Niagara/Source/NiagaraEditor/Private/Widgets/SNiagaraDistributionCurveEditor.cpp
@@ -58,21 +58,17 @@ namespace
 
 	struct FNiagaraDistributionCurveEditorDetailsTreeItem : public ICurveEditorTreeItem, TSharedFromThis<FNiagaraDistributionCurveEditorDetailsTreeItem>
 	{
-		FNiagaraDistributionCurveEditorDetailsTreeItem(TSharedRef<INiagaraDistributionAdapter> InDistributionAdapter, int InChannelIndex)
-			: DistributionAdapter(InDistributionAdapter)
-			, ChannelIndex(InChannelIndex)
+		void Initialize(TSharedRef<INiagaraDistributionAdapter> InDistributionAdapter, int InChannelIndex)
 		{
+			DistributionAdapter = InDistributionAdapter;
+			ChannelIndex = InChannelIndex;
+
 			Curve			= *DistributionAdapter->GetCurveValue(ChannelIndex);
 			CurveName		= DistributionAdapter->GetChannelDisplayName(ChannelIndex);
 			CurveToolTip	= DistributionAdapter->GetChannelToolTip(ChannelIndex);
 			CurveColor		= DistributionAdapter->GetChannelColor(ChannelIndex).GetSpecifiedColor();
 
-			if ( TSharedPtr<IPropertyHandle> PropertyHandle = DistributionAdapter->GetPropertyHandle() )
-			{
-				TArray<UObject*> OwnerObjects;
-				PropertyHandle->GetOuterObjects(OwnerObjects);
-				WeakOwnerObject = OwnerObjects.Num() > 0 ? OwnerObjects[0] : nullptr;
-			}
+			DistributionAdapter->OnDistributionValueChanged().AddSP(this, &FNiagaraDistributionCurveEditorDetailsTreeItem::OnDistributionValueChanged);
 		}
 
 		virtual TSharedPtr<SWidget> GenerateCurveEditorTreeWidget(const FName& InColumnName, TWeakPtr<FCurveEditor> InCurveEditor, FCurveEditorTreeItemID InTreeItemID, const TSharedRef<ITableRow>& TableRow) override
@@ -121,13 +117,14 @@ namespace
 
 		virtual void CreateCurveModels(TArray<TUniquePtr<FCurveModel>>& OutCurveModels) override
 		{
-			UObject* OwnerObject = WeakOwnerObject.Get();
-			if (OwnerObject == nullptr)
+			TArray<UObject*> OwnerObjects;
+			DistributionAdapter->GetOwners(OwnerObjects);
+			if (OwnerObjects.Num() != 1)
 			{
 				return;
 			}
 
-			TUniquePtr<FRichCurveEditorModelRaw> NewCurve = MakeUnique<FRichCurveEditorModelRaw>(&Curve, OwnerObject);
+			TUniquePtr<FRichCurveEditorModelRaw> NewCurve = MakeUnique<FRichCurveEditorModelRaw>(&Curve, OwnerObjects[0]);
 			NewCurve->SetShortDisplayName(CurveName);
 			NewCurve->SetColor(CurveColor);
 			NewCurve->OnCurveModified().AddSP(this, &FNiagaraDistributionCurveEditorDetailsTreeItem::OnCurveChanged);
@@ -140,16 +137,15 @@ namespace
 		void OnCurveChanged()
 		{
 			DistributionAdapter->SetCurveValue(ChannelIndex, Curve, false);
+		}
 
-			if (TSharedPtr<IPropertyHandle> PropertyHandle = DistributionAdapter->GetPropertyHandle())
-			{
-				PropertyHandle->NotifyPostChange(EPropertyChangeType::ValueSet);
-			}
+		void OnDistributionValueChanged()
+		{
+			Curve = *DistributionAdapter->GetCurveValue(ChannelIndex);
 		}
 
 	private:
-		TSharedRef<INiagaraDistributionAdapter>	DistributionAdapter;
-		TWeakObjectPtr<UObject>					WeakOwnerObject;
+		TSharedPtr<INiagaraDistributionAdapter>	DistributionAdapter;
 		int										ChannelIndex = 0;
 		FRichCurve								Curve;
 		FText									CurveName;
@@ -779,7 +775,8 @@ void SNiagaraDistributionCurveEditor::Construct(const FArguments& InArgs, TShare
 	TArray<FCurveEditorTreeItemID> TreeItemIds;
 	for (int iChannel=0; iChannel < NumChannels; ++iChannel)
 	{
-		TSharedRef<FNiagaraDistributionCurveEditorDetailsTreeItem> TreeItem = MakeShared<FNiagaraDistributionCurveEditorDetailsTreeItem>(DistributionAdapter.ToSharedRef(), iChannel);
+		TSharedRef<FNiagaraDistributionCurveEditorDetailsTreeItem> TreeItem = MakeShared<FNiagaraDistributionCurveEditorDetailsTreeItem>();
+		TreeItem->Initialize(DistributionAdapter.ToSharedRef(), iChannel);
 		FCurveEditorTreeItem* NewItem = CurveEditor->AddTreeItem(FCurveEditorTreeItemID::Invalid());
 		NewItem->SetStrongItem(TreeItem);
 		TreeItemIds.Add(NewItem->GetID());
diff --git a/Engine/Plugins/FX/Niagara/Source/NiagaraEditor/Private/Widgets/SNiagaraDistributionLookupValueWidget.cpp b/Engine/Plugins/FX/Niagara/Source/NiagaraEditor/Private/Widgets/SNiagaraDistributionLookupValueWidget.cpp
--- a/Engine/Plugins/FX/Niagara/Source/NiagaraEditor/Private/Widgets/SNiagaraDistributionLookupValueWidget.cpp
+++ b/Engine/Plugins/FX/Niagara/Source/NiagaraEditor/Private/Widgets/SNiagaraDistributionLookupValueWidget.cpp
@@ -19,14 +19,7 @@ void SNiagaraDistributionLookupValueWidget::Construct(const FArguments& InArgs,
 		return;
 	}
 
-	TSharedPtr<IPropertyHandle> DistributionPropertyHandle = DistributionAdapter->GetPropertyHandle();
-	FProperty* DistributionProperty = DistributionPropertyHandle ? DistributionPropertyHandle->GetProperty() : nullptr;
-	if (DistributionProperty == nullptr)
-	{
-		return;
-	}
-
-	LookupValueModeEnum = FNiagaraDistributionBase::GetLookupValueModeEnum(DistributionProperty);
+	LookupValueModeEnum = DistributionAdapter->GetLookupValueModeEnum();
 	if (LookupValueModeEnum == nullptr)
 	{
 		return;
diff --git a/Engine/Plugins/FX/Niagara/Source/NiagaraEditor/Public/INiagaraEditorTypeUtilities.h b/Engine/Plugins/FX/Niagara/Source/NiagaraEditor/Public/INiagaraEditorTypeUtilities.h
--- a/Engine/Plugins/FX/Niagara/Source/NiagaraEditor/Public/INiagaraEditorTypeUtilities.h
+++ b/Engine/Plugins/FX/Niagara/Source/NiagaraEditor/Public/INiagaraEditorTypeUtilities.h
@@ -128,3 +128,17 @@ public:
 	virtual bool TryUpdatePropertyFromClipboardPortableValue(const FNiagaraClipboardPortableValue& InSourceClipboardPortableValue, IPropertyHandle& InPropertyHandle) const override { return false; }
 };
 
+class INiagaraEditorExternalObjectUtilities
+{
+public:
+	virtual ~INiagaraEditorExternalObjectUtilities() = default;
+	virtual bool SupportsClipboardPortableValues() const = 0;
+	virtual bool TryUpdateClipboardPortableValueFromObject(const UObject& InExternalObject, FNiagaraClipboardPortableValue& InTargetClipboardPortableValue) const = 0;
+};
+
+class FNiagaraEditorExternalObjectUtilities : public INiagaraEditorExternalObjectUtilities
+{
+public:
+	virtual bool SupportsClipboardPortableValues() const override { return false; }
+	virtual bool TryUpdateClipboardPortableValueFromObject(const UObject& InExternalObject, FNiagaraClipboardPortableValue& InTargetClipboardPortableValue) const override { return false; }
+};
\ No newline at end of file
diff --git a/Engine/Plugins/FX/Niagara/Source/NiagaraEditor/Public/NiagaraClipboard.h b/Engine/Plugins/FX/Niagara/Source/NiagaraEditor/Public/NiagaraClipboard.h
--- a/Engine/Plugins/FX/Niagara/Source/NiagaraEditor/Public/NiagaraClipboard.h
+++ b/Engine/Plugins/FX/Niagara/Source/NiagaraEditor/Public/NiagaraClipboard.h
@@ -234,11 +234,23 @@ struct FNiagaraClipboardPortableValue
 	static FNiagaraClipboardPortableValue NIAGARAEDITOR_API CreateFromTypedValue(const FNiagaraTypeDefinition& InType, const FNiagaraVariant& InValue);
 	static FNiagaraClipboardPortableValue NIAGARAEDITOR_API CreateFromPropertyHandle(const IPropertyHandle& InPropertyHandle);
 
+	template<typename StructType>
+	static FNiagaraClipboardPortableValue CreateFromStructValue(const StructType& InStructValue)
+	{
+		return CreateFromStructValue(*StructType::StaticStruct(), (uint8*)&InStructValue);
+	}
+
 	NIAGARAEDITOR_API bool TryUpdateStructValue(const UScriptStruct& TargetStruct, uint8* StructMemory) const;
 	NIAGARAEDITOR_API bool CanUpdateTypedValue(const FNiagaraTypeDefinition& InTargetInputType) const;
 	NIAGARAEDITOR_API bool TryUpdateTypedValue(const FNiagaraTypeDefinition& InTargetInputType, FNiagaraVariant& InValue) const;
 	NIAGARAEDITOR_API bool TryUpdatePropertyHandle(IPropertyHandle& InTargetPropertyHandle) const;
 
+	template<typename StructType>
+	bool TryUpdateStructValue(StructType& OutStructValue) const 
+	{
+		return TryUpdateStructValue(*StructType::StaticStruct(), (uint8*)&OutStructValue);
+	}
+
 	UPROPERTY()
 	FString ValueString;
 };
diff --git a/Engine/Plugins/FX/Niagara/Source/NiagaraEditor/Public/NiagaraEditorModule.h b/Engine/Plugins/FX/Niagara/Source/NiagaraEditor/Public/NiagaraEditorModule.h
--- a/Engine/Plugins/FX/Niagara/Source/NiagaraEditor/Public/NiagaraEditorModule.h
+++ b/Engine/Plugins/FX/Niagara/Source/NiagaraEditor/Public/NiagaraEditorModule.h
@@ -36,6 +36,7 @@ class IAssetTools;
 class IAssetTypeActions;
 class INiagaraEditorTypeUtilities;
 class INiagaraEditorPropertyUtilities;
+class INiagaraEditorExternalObjectUtilities;
 class UNiagaraDataInterface;
 class UNiagaraSettings;
 class USequencerSettings;
@@ -216,6 +217,9 @@ public:
 	/** Registers niagara editor property utilities for a specific struct type. */
 	void RegisterPropertyUtilities(const UScriptStruct* InStruct, TSharedRef<INiagaraEditorPropertyUtilities, ESPMode::ThreadSafe> PropertyUtilities);
 
+	/** Registers niagara editor external object utilities for a specific class type, optionally allowing handling subtypes. */
+	void RegisterExternalObjectUtilities(const UClass& InClass, bool bHandleSubclasses, TSharedRef<INiagaraEditorExternalObjectUtilities, ESPMode::ThreadSafe> ExternalObjectUtilities);
+
 	/** Registers niagara wizards that can be used to add generated modules to the stack. */
 	void RegisterModuleWizards(TSharedRef<UE::Niagara::Wizard::FModuleWizardGenerator> WizardGenerator);
 	TConstArrayView<TSharedRef<UE::Niagara::Wizard::FModuleWizardGenerator>> GetModuleWizards() const { return ModuleWizards; }
@@ -227,10 +231,13 @@ public:
 	void InitializeTaggedAssetBrowserDetails();
 	
 	/** Gets Niagara editor type utilities for a specific type if there are any registered. */
-	TSharedPtr<INiagaraEditorTypeUtilities, ESPMode::ThreadSafe> NIAGARAEDITOR_API GetTypeUtilities(const FNiagaraTypeDefinition& Type);
+	TSharedPtr<INiagaraEditorTypeUtilities, ESPMode::ThreadSafe> NIAGARAEDITOR_API GetTypeUtilities(const FNiagaraTypeDefinition& Type) const;
 
 	/** Gets Niagara editor property utilities for a specific struct if there are any registered. */
-	TSharedPtr<INiagaraEditorPropertyUtilities, ESPMode::ThreadSafe> NIAGARAEDITOR_API GetPropertyUtilities(const UScriptStruct& Struct);
+	TSharedPtr<INiagaraEditorPropertyUtilities, ESPMode::ThreadSafe> NIAGARAEDITOR_API GetPropertyUtilities(const UScriptStruct& Struct) const;
+
+	/** Gets Niagara editor external object utilities for a specific class if there are any registered. */
+	TSharedPtr<const INiagaraEditorExternalObjectUtilities, ESPMode::ThreadSafe> NIAGARAEDITOR_API GetExternalObjectUtilities(const UClass& Class) const;
 
 	NIAGARAEDITOR_API void RegisterWidgetProvider(TSharedRef<INiagaraEditorWidgetProvider> InWidgetProvider);
 	NIAGARAEDITOR_API void UnregisterWidgetProvider(TSharedRef<INiagaraEditorWidgetProvider> InWidgetProvider);
@@ -473,11 +480,21 @@ private:
 	TArray< TSharedPtr<IAssetTypeActions> > CreatedAssetTypeActions;
 	TArray<TSharedRef<UE::Niagara::Wizard::FModuleWizardGenerator>> ModuleWizards;
 
-	FCriticalSection TypeEditorsCS;
+	mutable UE::FMutex EditorUtilitiesGuard;
+
 	TMap<FNiagaraTypeDefinition, TSharedRef<INiagaraEditorTypeUtilities, ESPMode::ThreadSafe>> TypeToEditorUtilitiesMap;
 	TSharedPtr<INiagaraEditorTypeUtilities, ESPMode::ThreadSafe> EnumTypeUtilities;
 	TMap<const UScriptStruct*, TSharedRef<INiagaraEditorPropertyUtilities, ESPMode::ThreadSafe>> StructToPropertyUtilitiesMap;
 
+	struct FRegisteredExternalObjectUtilities
+	{
+		const UClass* Class = nullptr;
+		bool bHandleSubclasses = false;
+		TSharedPtr<const INiagaraEditorExternalObjectUtilities, ESPMode::ThreadSafe> ExternalObjectUtilities;
+	};
+
+	TArray<FRegisteredExternalObjectUtilities> RegisteredExternalObjectUtilities;
+
 	FOnScriptApplied OnScriptAppliedDelegate;
 
 	FDelegateHandle CreateEmitterTrackEditorHandle;
