- Handle interpolated / not interpolated for emitter merging
- The script is not the owner of the interpoation state, the emitter is, so we need to make sure that doesn't get out of sync

#Synced-CL 48951098

diff --git a/Engine/Plugins/FX/Niagara/Source/NiagaraEditor/Private/NiagaraScriptMergeManager.cpp b/Engine/Plugins/FX/Niagara/Source/NiagaraEditor/Private/NiagaraScriptMergeManager.cpp
--- a/Engine/Plugins/FX/Niagara/Source/NiagaraEditor/Private/NiagaraScriptMergeManager.cpp
+++ b/Engine/Plugins/FX/Niagara/Source/NiagaraEditor/Private/NiagaraScriptMergeManager.cpp
@@ -2584,7 +2584,15 @@ void FNiagaraScriptMergeManager::DiffScriptStacks(TSharedRef<FNiagaraScriptStack
 		}
 	}
 
-	if (BaseScriptStackAdapter->GetScript()->GetUsage() != OtherScriptStackAdapter->GetScript()->GetUsage())
+	// We need to ignore the interpolation setting when checking script usage as the emitter is the controller for this
+	// and we don't merge the emitter properties right now
+	auto GetUsageIgnoreInterpolation =
+		[](ENiagaraScriptUsage InUsage) -> ENiagaraScriptUsage
+		{
+			return InUsage == ENiagaraScriptUsage::ParticleSpawnScriptInterpolated ? ENiagaraScriptUsage::ParticleSpawnScript : InUsage;
+		};
+
+	if (GetUsageIgnoreInterpolation(BaseScriptStackAdapter->GetScript()->GetUsage()) != GetUsageIgnoreInterpolation(OtherScriptStackAdapter->GetScript()->GetUsage()))
 	{
 		DiffResults.ChangedBaseUsage = BaseScriptStackAdapter->GetScript()->GetUsage();
 		DiffResults.ChangedOtherUsage = OtherScriptStackAdapter->GetScript()->GetUsage();
