- Fix for TSan issue with accessing CurrentNumInstances_RT

#Synced-CL 48951098

diff --git a/Engine/Plugins/FX/Niagara/Source/Niagara/Classes/NiagaraComputeExecutionContext.h b/Engine/Plugins/FX/Niagara/Source/Niagara/Classes/NiagaraComputeExecutionContext.h
--- a/Engine/Plugins/FX/Niagara/Source/Niagara/Classes/NiagaraComputeExecutionContext.h
+++ b/Engine/Plugins/FX/Niagara/Source/Niagara/Classes/NiagaraComputeExecutionContext.h
@@ -219,7 +219,8 @@ public:
 	// The current number of instances on the RT
 	// Before ticks are processed on the RT this will be CurrentData's NumInstances
 	// As ticks are processed (or we generated the tick batches) this will change and won't be accurate until dispatches are executed
-	uint32 CurrentNumInstances_RT = 0;
+	UE_DEPRECATED(5.8, "Direct access of CurrentNumInstances_RT is deprecated please use the get / set methods")
+	std::atomic<uint32> CurrentNumInstances_RT = 0;
 	// The current maximum of instances on the RT
 	uint32 CurrentMaxInstances_RT = 0;
 	// The current maximum instances we should allocate on the RT
@@ -227,6 +228,20 @@ public:
 
 	FNiagaraSimStageExecutionDataPtr SimStageExecData;
 
+	void SetCurrentNumInstances(uint32 NumInstances)
+	{
+		PRAGMA_DISABLE_DEPRECATION_WARNINGS
+		CurrentNumInstances_RT.store(NumInstances, std::memory_order_relaxed);
+		PRAGMA_ENABLE_DEPRECATION_WARNINGS
+	}
+
+	uint32 GetCurrentNumInstances() const
+	{
+		PRAGMA_DISABLE_DEPRECATION_WARNINGS
+		return CurrentNumInstances_RT.load(std::memory_order_relaxed);
+		PRAGMA_ENABLE_DEPRECATION_WARNINGS
+	}
+
 	bool IsOutputStage(FNiagaraDataInterfaceProxy* DIProxy, uint32 SimulationStageIndex) const;
 	bool IsInputStage(FNiagaraDataInterfaceProxy* DIProxy, uint32 SimulationStageIndex) const;
 	bool IsIterationStage(FNiagaraDataInterfaceProxy* DIProxy, uint32 SimulationStageIndex) const;
diff --git a/Engine/Plugins/FX/Niagara/Source/Niagara/Private/NiagaraComputeExecutionContext.cpp b/Engine/Plugins/FX/Niagara/Source/Niagara/Private/NiagaraComputeExecutionContext.cpp
--- a/Engine/Plugins/FX/Niagara/Source/Niagara/Private/NiagaraComputeExecutionContext.cpp
+++ b/Engine/Plugins/FX/Niagara/Source/Niagara/Private/NiagaraComputeExecutionContext.cpp
@@ -226,7 +226,7 @@ void FNiagaraComputeExecutionContext::ResetInternal(FNiagaraGpuComputeDispatchIn
 		ComputeDispatchInterface->GetGPUInstanceCounterManager().FreeEntry(CountOffset_RT);
 	}
 
-	CurrentNumInstances_RT = 0;
+	SetCurrentNumInstances(0);
 	CountOffset_RT = INDEX_NONE;
 	EmitterInstanceReadback.GPUCountOffset = INDEX_NONE;
 
diff --git a/Engine/Plugins/FX/Niagara/Source/Niagara/Private/NiagaraEmitterInstance.cpp b/Engine/Plugins/FX/Niagara/Source/Niagara/Private/NiagaraEmitterInstance.cpp
--- a/Engine/Plugins/FX/Niagara/Source/Niagara/Private/NiagaraEmitterInstance.cpp
+++ b/Engine/Plugins/FX/Niagara/Source/Niagara/Private/NiagaraEmitterInstance.cpp
@@ -41,7 +41,7 @@ int32 FNiagaraEmitterInstance::GetNumParticles() const
 		if (GPUExecContext->ParticleCountReadFence.load(std::memory_order_relaxed) <= GPUExecContext->ParticleCountWriteFence.load(std::memory_order_relaxed))
 		{
 			// Fence has passed we read directly from the GPU Exec Context which will have the most up-to-date information
-			return GPUExecContext->CurrentNumInstances_RT;
+			return GPUExecContext->GetCurrentNumInstances();
 		}
 		// Fence has not been passed return the TotalSpawnedParticles as a 'guess' to the amount
 		return TotalSpawnedParticles;
diff --git a/Engine/Plugins/FX/Niagara/Source/Niagara/Private/NiagaraGpuComputeDispatch.cpp b/Engine/Plugins/FX/Niagara/Source/Niagara/Private/NiagaraGpuComputeDispatch.cpp
--- a/Engine/Plugins/FX/Niagara/Source/Niagara/Private/NiagaraGpuComputeDispatch.cpp
+++ b/Engine/Plugins/FX/Niagara/Source/Niagara/Private/NiagaraGpuComputeDispatch.cpp
@@ -753,9 +753,10 @@ void FNiagaraGpuComputeDispatch::UpdateInstanceCountManager(FRHICommandListImmed
 						}
 
 						const uint32 DeadInstanceCount = ComputeContext->EmitterInstanceReadback.CPUCount - Counts[ComputeContext->EmitterInstanceReadback.GPUCountOffset];
-						if ( DeadInstanceCount <= ComputeContext->CurrentNumInstances_RT )
+						const uint32 CurrentInstanceCount = ComputeContext->GetCurrentNumInstances();
+						if ( DeadInstanceCount <= CurrentInstanceCount)
 						{
-							ComputeContext->CurrentNumInstances_RT -= DeadInstanceCount;
+							ComputeContext->SetCurrentNumInstances(CurrentInstanceCount - DeadInstanceCount);
 						}
 						if (FNiagaraGpuComputeDispatchLocal::GDebugLogging)
 						{
@@ -827,7 +828,7 @@ void FNiagaraGpuComputeDispatch::PrepareTicksForProxy(FRHICommandListImmediate&
 			// Instance requires a reset?
 			if (InstanceData.bResetData)
 			{
-				ComputeContext->CurrentNumInstances_RT = 0;
+				ComputeContext->SetCurrentNumInstances(0);
 				if (ComputeContext->CountOffset_RT != INDEX_NONE)
 				{
 					GpuDispatchList.CountsToRelease.Add(ComputeContext->CountOffset_RT);
@@ -873,23 +874,26 @@ void FNiagaraGpuComputeDispatch::PrepareTicksForProxy(FRHICommandListImmediate&
 
 			// Calculate instance counts
 			const uint32 MaxBufferInstances = ComputeContext->MainDataSet->GetMaxInstanceCount();
-			const uint32 PrevNumInstances = ComputeContext->CurrentNumInstances_RT;
+			const uint32 PrevNumInstances = ComputeContext->GetCurrentNumInstances();
 
-			ComputeContext->CurrentNumInstances_RT = PrevNumInstances + InstanceData.SpawnInfo.SpawnRateInstances + InstanceData.SpawnInfo.EventSpawnTotal;
-			ComputeContext->CurrentNumInstances_RT = FMath::Min(ComputeContext->CurrentNumInstances_RT, MaxBufferInstances);
+			uint32 NewInstanceCount = PrevNumInstances;
+			NewInstanceCount = NewInstanceCount + InstanceData.SpawnInfo.SpawnRateInstances + InstanceData.SpawnInfo.EventSpawnTotal;
+			NewInstanceCount = FMath::Min(NewInstanceCount, MaxBufferInstances);
 
 			// Calculate new maximum count
-			ComputeContext->CurrentMaxInstances_RT = FMath::Max(ComputeContext->CurrentMaxInstances_RT, ComputeContext->CurrentNumInstances_RT);
+			NewInstanceCount = FMath::Max(ComputeContext->CurrentMaxInstances_RT, NewInstanceCount);
 
-			if (!GNiagaraBatcherFreeBufferEarly || (ComputeContext->CurrentMaxInstances_RT > 0))
+			if (!GNiagaraBatcherFreeBufferEarly || (NewInstanceCount > 0))
 			{
-				ComputeContext->CurrentMaxAllocateInstances_RT = FMath::Max3(ComputeContext->CurrentMaxAllocateInstances_RT, ComputeContext->CurrentMaxInstances_RT, InstanceData.SpawnInfo.MaxParticleCount);
+				ComputeContext->CurrentMaxAllocateInstances_RT = FMath::Max3(ComputeContext->CurrentMaxAllocateInstances_RT, NewInstanceCount, InstanceData.SpawnInfo.MaxParticleCount);
 			}
 			else
 			{
-				ComputeContext->CurrentMaxAllocateInstances_RT = FMath::Max(ComputeContext->CurrentMaxAllocateInstances_RT, ComputeContext->CurrentMaxInstances_RT);
+				ComputeContext->CurrentMaxAllocateInstances_RT = FMath::Max(ComputeContext->CurrentMaxAllocateInstances_RT, NewInstanceCount);
 			}
 
+			ComputeContext->SetCurrentNumInstances(NewInstanceCount);
+
 			bHasFreeIDUpdates |= ComputeContext->MainDataSet->RequiresPersistentIDs();
 
 			//-OPT: Do we need this test?  Can remove in favor of MaxUpdateIterations
@@ -926,20 +930,20 @@ void FNiagaraGpuComputeDispatch::PrepareTicksForProxy(FRHICommandListImmediate&
 					{
 						SimStageData.Source = SourceData;
 						SimStageData.SourceCountOffset = ComputeContext->CountOffset_RT;
-						SimStageData.SourceNumInstances = ComputeContext->CurrentNumInstances_RT;
+						SimStageData.SourceNumInstances = NewInstanceCount;
 						SimStageData.Destination = nullptr;
 						SimStageData.DestinationCountOffset = ComputeContext->CountOffset_RT;
-						SimStageData.DestinationNumInstances = ComputeContext->CurrentNumInstances_RT;
+						SimStageData.DestinationNumInstances = NewInstanceCount;
 					}
 					// This stage writes particles but will not kill any, we can use the buffer as both source and destination
 					else if (SimStageData.StageMetaData->bPartialParticleUpdate)
 					{
 						SimStageData.Source = nullptr;
 						SimStageData.SourceCountOffset = ComputeContext->CountOffset_RT;
-						SimStageData.SourceNumInstances = ComputeContext->CurrentNumInstances_RT;
+						SimStageData.SourceNumInstances = NewInstanceCount;
 						SimStageData.Destination = SourceData;
 						SimStageData.DestinationCountOffset = ComputeContext->CountOffset_RT;
-						SimStageData.DestinationNumInstances = ComputeContext->CurrentNumInstances_RT;
+						SimStageData.DestinationNumInstances = NewInstanceCount;
 					}
 					// This stage may kill particles, we need to allocate a new destination buffer
 					else
@@ -947,10 +951,10 @@ void FNiagaraGpuComputeDispatch::PrepareTicksForProxy(FRHICommandListImmediate&
 						SimStageData.Source = SourceData;
 						SimStageData.SourceCountOffset = ComputeContext->CountOffset_RT;
 						//-TODO: This is a little odd, perhaps we need to change the preallocate
-						SimStageData.SourceNumInstances = SimStageIndex == 0 && IterationIndex == 0 ? PrevNumInstances : ComputeContext->CurrentNumInstances_RT;
+						SimStageData.SourceNumInstances = SimStageIndex == 0 && IterationIndex == 0 ? PrevNumInstances : NewInstanceCount;
 						SimStageData.Destination = ComputeContext->GetNextDataBuffer();
 						SimStageData.DestinationCountOffset = GPUInstanceCounterManager.AcquireEntry();
-						SimStageData.DestinationNumInstances = ComputeContext->CurrentNumInstances_RT;
+						SimStageData.DestinationNumInstances = NewInstanceCount;
 
 						ComputeContext->AdvanceDataBuffer();
 						ComputeContext->CountOffset_RT = SimStageData.DestinationCountOffset;
@@ -1053,7 +1057,7 @@ void FNiagaraGpuComputeDispatch::PrepareTicksForProxy(FRHICommandListImmediate&
 		{
 			FNiagaraDataBuffer* FinalBuffer = ComputeContext->GetPrevDataBuffer();
 			FinalBuffer->SetGPUInstanceCountBufferOffset(ComputeContext->CountOffset_RT);
-			FinalBuffer->SetNumInstances(ComputeContext->CurrentNumInstances_RT);
+			FinalBuffer->SetNumInstances(ComputeContext->GetCurrentNumInstances());
 			FinalBuffer->SetGPUDataReadyStage(ComputeProxy->GetComputeTickStage());
 			ComputeContext->SetDataToRender(FinalBuffer);
 		}
@@ -1064,7 +1068,7 @@ void FNiagaraGpuComputeDispatch::PrepareTicksForProxy(FRHICommandListImmediate&
 		{
 			FNiagaraDataBuffer* FinalBuffer = ComputeContext->GetPrevDataBuffer();
 			FinalBuffer->SetGPUInstanceCountBufferOffset(ComputeContext->CountOffset_RT);
-			FinalBuffer->SetNumInstances(ComputeContext->CurrentNumInstances_RT);
+			FinalBuffer->SetNumInstances(ComputeContext->GetCurrentNumInstances());
 			FinalBuffer->SetGPUDataReadyStage(ComputeProxy->GetComputeTickStage());
 			ComputeContext->SetTranslucentDataToRender(FinalBuffer);
 		}
diff --git a/Engine/Plugins/FX/Niagara/Source/Niagara/Private/NiagaraSimCacheHelper.h b/Engine/Plugins/FX/Niagara/Source/Niagara/Private/NiagaraSimCacheHelper.h
--- a/Engine/Plugins/FX/Niagara/Source/Niagara/Private/NiagaraSimCacheHelper.h
+++ b/Engine/Plugins/FX/Niagara/Source/Niagara/Private/NiagaraSimCacheHelper.h
@@ -889,7 +889,7 @@ struct FNiagaraSimCacheHelper
 
 					DataBuffer.SetGPUInstanceCountBufferOffset(GPUExecContext->CountOffset_RT);
 
-					GPUExecContext->CurrentNumInstances_RT = NumInstances;
+					GPUExecContext->SetCurrentNumInstances(NumInstances);
 					GPUExecContext->ParticleCountWriteFence.store(ParticleCountFence, std::memory_order_relaxed);
 
 					const FRWBuffer& CountBuffer = CountManager.GetInstanceCountBuffer();
diff --git a/Engine/Plugins/FX/Niagara/Source/Niagara/Private/NiagaraSystemGpuComputeProxy.cpp b/Engine/Plugins/FX/Niagara/Source/Niagara/Private/NiagaraSystemGpuComputeProxy.cpp
--- a/Engine/Plugins/FX/Niagara/Source/Niagara/Private/NiagaraSystemGpuComputeProxy.cpp
+++ b/Engine/Plugins/FX/Niagara/Source/Niagara/Private/NiagaraSystemGpuComputeProxy.cpp
@@ -96,7 +96,7 @@ void FNiagaraSystemGpuComputeProxy::AddToRenderThread(FNiagaraGpuComputeDispatch
 			for (FNiagaraComputeExecutionContext* ComputeContext : ComputeContexts)
 			{
 				ComputeContext->bHasTickedThisFrame_RT = false;
-				ComputeContext->CurrentNumInstances_RT = 0;
+				ComputeContext->SetCurrentNumInstances(0);
 				ComputeContext->CurrentMaxInstances_RT = 0;
 
 				ComputeContext->EmitterInstanceReadback.CPUCount = 0;
