Niagara - Add a cvar which can force the compilation request tasks to run inline rather then multithreaded.

#Synced-CL 48514417

diff --git a/Engine/Plugins/FX/Niagara/Source/NiagaraEditor/Private/NiagaraCompilationTasks.cpp b/Engine/Plugins/FX/Niagara/Source/NiagaraEditor/Private/NiagaraCompilationTasks.cpp
--- a/Engine/Plugins/FX/Niagara/Source/NiagaraEditor/Private/NiagaraCompilationTasks.cpp
+++ b/Engine/Plugins/FX/Niagara/Source/NiagaraEditor/Private/NiagaraCompilationTasks.cpp
@@ -34,6 +34,14 @@
 #include "NiagaraNodeOutput.h"
 #include "NiagaraNodeFunctionCall.h"
 
+static int GNiagaraCompilationForceOneThread = 0;
+static FAutoConsoleVariableRef CVarNiagaraCompilationForceOneThread(
+	TEXT("fx.Niagara.Compilation.ForceOneThread"),
+	GNiagaraCompilationForceOneThread,
+	TEXT("Forces the compile tasks to run inline on one thread."),
+	ECVF_Default
+);
+
 #if ENABLE_COOK_STATS
 namespace NiagaraSystemCookStats
 {
@@ -2649,6 +2657,11 @@ void FNiagaraSystemCompilationTask::IssueCompilationTasks()
 	TArray<FTask> IssuedCompilationTasks;
 	TArray<FTaskEvent> CompilationTasks;
 
+	ETaskPriority TaskPriority = ETaskPriority::Normal;
+	EExtendedTaskPriority TaskExtendedPriority = (GNiagaraCompilationForceOneThread > 0 || FApp::ShouldUseThreadingForPerformance() == false)
+		? EExtendedTaskPriority::Inline 
+		: EExtendedTaskPriority::None;
+
 	if (HasOutstandingCompileTasks())
 	{
 		FTask PrecompileTask = Launch(UE_SOURCE_LOCATION, [Self = AsShared()]
@@ -2657,7 +2670,7 @@ void FNiagaraSystemCompilationTask::IssueCompilationTasks()
 			UE_LOG_CONTEXT("NiagaraSystem", Self->SystemInfo.GetSystemPath());
 
 			Self->Precompile();
-		});
+		}, TaskPriority, TaskExtendedPriority);
 
 		// async compilation copy tasks
 		for (FCompileGroupInfo& GroupInfo : CompileGroups)
@@ -2670,7 +2683,7 @@ void FNiagaraSystemCompilationTask::IssueCompilationTasks()
 					UE_LOG_CONTEXT("NiagaraSystem", Self->SystemInfo.GetSystemPath());
 
 					GroupInfo.InstantiateCompileGraph(*Self);
-				}, PrecompileTask);
+				}, PrecompileTask, TaskPriority, TaskExtendedPriority);
 
 				for (int32 CompileTaskIndex : GroupInfo.CompileTaskIndices)
 				{
@@ -2715,7 +2728,7 @@ void FNiagaraSystemCompilationTask::IssueCompilationTasks()
 									CompileTask.Translate(TaskPtr, GroupInfo);
 								break;
 							}
-						}, InstantiateGraphTask));
+						}, InstantiateGraphTask, TaskPriority, TaskExtendedPriority));
 
 						Launch(UE_SOURCE_LOCATION, [Self = AsShared(), &GroupInfo, &CompileTask]
 						{
@@ -2723,7 +2736,7 @@ void FNiagaraSystemCompilationTask::IssueCompilationTasks()
 							UE_LOG_CONTEXT("NiagaraSystem", Self->SystemInfo.GetSystemPath());
 
 							CompileTask.ProcessCompilation(&Self.Get(), GroupInfo);
-						}, *CompileTask.CompileResultsReadyEvent);
+						}, *CompileTask.CompileResultsReadyEvent, TaskPriority, TaskExtendedPriority);
 
 						CompilationTasks.Add(*CompileTask.CompileResultsProcessedEvent);
 					}
@@ -2735,13 +2748,13 @@ void FNiagaraSystemCompilationTask::IssueCompilationTasks()
 	Launch(UE_SOURCE_LOCATION, [Self = AsShared()]
 	{
 		Self->CompilationState = EState::WaitingForProcessing;
-	}, IssuedCompilationTasks);
+	}, IssuedCompilationTasks, TaskPriority, TaskExtendedPriority);
 
 	Launch(UE_SOURCE_LOCATION, [Self = AsShared()]
 	{
 		UE_LOG_CONTEXT("NiagaraSystem", Self->SystemInfo.GetSystemPath());
 		Self->IssuePostResultsProcessedTasks();
-	}, CompilationTasks);
+	}, CompilationTasks, TaskPriority, TaskExtendedPriority);
 }
 
 bool FNiagaraSystemCompilationTask::HasOutstandingCompileTasks() const
