- Fix for light bounds calculation not taking into account the radius scale on the renderer

#tests Open the debugger, visualize bounds and have a single light renderer adjusting the radius scale on the renderer

#Synced-CL 48951098

diff --git a/Engine/Plugins/FX/Niagara/Source/Niagara/Private/NiagaraLightRendererProperties.cpp b/Engine/Plugins/FX/Niagara/Source/Niagara/Private/NiagaraLightRendererProperties.cpp
--- a/Engine/Plugins/FX/Niagara/Source/Niagara/Private/NiagaraLightRendererProperties.cpp
+++ b/Engine/Plugins/FX/Niagara/Source/Niagara/Private/NiagaraLightRendererProperties.cpp
@@ -160,9 +160,10 @@ FNiagaraRenderer* UNiagaraLightRendererProperties::CreateEmitterRenderer(ERHIFea
 class FNiagaraBoundsCalculatorLights : public FNiagaraBoundsCalculator
 {
 public:
-	FNiagaraBoundsCalculatorLights(const FNiagaraDataSetAccessor<FNiagaraPosition>& InPositionAccessor, const FNiagaraDataSetAccessor<float>& InRadiusAccessor)
+	FNiagaraBoundsCalculatorLights(const FNiagaraDataSetAccessor<FNiagaraPosition>& InPositionAccessor, const FNiagaraDataSetAccessor<float>& InRadiusAccessor, float InRadiusScale)
 		: PositionAccessor(InPositionAccessor)
 		, RadiusAccessor(InRadiusAccessor)
+		, RadiusScale(InRadiusScale)
 	{}
 
 	FNiagaraBoundsCalculatorLights() = delete;
@@ -181,19 +182,23 @@ public:
 		PositionAccessor.GetReader(DataSet).GetMinMax(BoundsMin, BoundsMax);
 		FBox Bounds(BoundsMin, BoundsMax);
 
+		float MaxRadius = 0.0f;
 		if (RadiusAccessor.IsValid())
 		{
-			return Bounds.ExpandBy(RadiusAccessor.GetReader(DataSet).GetMax());
+			MaxRadius = RadiusAccessor.GetReader(DataSet).GetMax();
+		}
+		else
+		{
+			MaxRadius = FNiagaraConstants::GetAttributeWithDefaultValue(SYS_PARAM_PARTICLES_LIGHT_RADIUS).GetValue<float>();
 		}
 
-		const FNiagaraVariable DefaultRadius = FNiagaraConstants::GetAttributeWithDefaultValue(SYS_PARAM_PARTICLES_LIGHT_RADIUS);
-
-		return Bounds.ExpandBy(DefaultRadius.GetValue<float>());
+		return Bounds.ExpandBy(MaxRadius * RadiusScale);
 	}
 
 protected:
 	const FNiagaraDataSetAccessor<FNiagaraPosition>& PositionAccessor;
 	const FNiagaraDataSetAccessor<float>& RadiusAccessor;
+	const float RadiusScale = 1.0f;
 };
 
 FNiagaraBoundsCalculator* UNiagaraLightRendererProperties::CreateBoundsCalculator()
@@ -203,7 +208,7 @@ FNiagaraBoundsCalculator* UNiagaraLightRendererProperties::CreateBoundsCalculato
 		return nullptr;
 	}
 
-	return new FNiagaraBoundsCalculatorLights(PositionDataSetAccessor, RadiusDataSetAccessor);
+	return new FNiagaraBoundsCalculatorLights(PositionDataSetAccessor, RadiusDataSetAccessor, RadiusScale);
 }
 
 void UNiagaraLightRendererProperties::GetUsedMaterials(const FNiagaraEmitterInstance* InEmitter, TArray<UMaterialInterface*>& OutMaterials) const
