Licensee fix for Niagara Async GPU Trace to properly support Inline Raytracing on PC. Previous to this change, the normals would always be incorrect.
#tests Ran a niagara system using async trace collision on the gpu and used verified that the collision normals on a sloped surface was corrected

#Synced-CL 0

diff --git a/Engine/Plugins/FX/Niagara/Source/NiagaraShader/Private/NiagaraAsyncGpuTraceProviderHwrt.cpp b/Engine/Plugins/FX/Niagara/Source/NiagaraShader/Private/NiagaraAsyncGpuTraceProviderHwrt.cpp
--- a/Engine/Plugins/FX/Niagara/Source/NiagaraShader/Private/NiagaraAsyncGpuTraceProviderHwrt.cpp
+++ b/Engine/Plugins/FX/Niagara/Source/NiagaraShader/Private/NiagaraAsyncGpuTraceProviderHwrt.cpp
@@ -10,6 +10,7 @@
 #include "Containers/StridedView.h"
 #include "FXRenderingUtils.h"
 #include "GlobalShader.h"
+#include "GlobalRenderResources.h"
 #include "MeshPassProcessor.h"
 #include "RayTracingMeshDrawCommands.h"
 #include "NiagaraSettings.h"
@@ -40,6 +41,13 @@ bool SupportsNiagaraAsyncGpuTraceHwrtInline()
 	return GRHISupportsInlineRayTracing && GNiagaraAsyncGpuTraceHwrtInline;
 }
 
+bool PlatformRequiresInlineSBTForAsyncGpuTraceHwrtInline(EShaderPlatform Platform)
+{
+	// PC can only fetch normals if InlineRayTracingBindingDataBuffer is valid. Otherwise we will fetch garbage
+	// Other platforms can query the normal by pulling verts directly from the RT BVH
+	return IsPCPlatform(Platform);
+}
+
 /// TODO
 ///  -get geometry masking working when an environmental mask is implemented
 
@@ -72,6 +80,8 @@ public:
 		SHADER_PARAMETER(uint32, HashTableSize)
 
 		SHADER_PARAMETER_SRV(RaytracingAccelerationStructure, TLAS)
+		SHADER_PARAMETER_SRV(StructuredBuffer, RayTracingSceneMetadata)
+
 		SHADER_PARAMETER_SRV(Buffer<FNiagaraAsyncGpuTrace>, Rays)
 		SHADER_PARAMETER(uint32, RaysOffset)
 		SHADER_PARAMETER_UAV(Buffer<FNiagaraAsyncGpuTraceResult>, CollisionOutput)
@@ -352,6 +362,15 @@ bool FNiagaraAsyncGpuTraceProviderHwrt::IsAvailable() const
 		return false;
 	}
 
+	if (PlatformRequiresInlineSBTForAsyncGpuTraceHwrtInline(GMaxRHIShaderPlatform))
+	{
+		FSceneInterface* Scene = Dispatcher->GetSceneInterface();
+		if (UE::FXRenderingUtils::RayTracing::GetInlineRayTracingBindingDataBuffer(*Scene) == nullptr)
+		{
+			return false;
+		}
+	}
+
 	return UE::FXRenderingUtils::RayTracing::HasRayTracingScene(*Dispatcher->GetSceneInterface());
 }
 
@@ -366,6 +385,9 @@ void FNiagaraAsyncGpuTraceProviderHwrt::PostRenderOpaque(FRHICommandList& RHICmd
 	if (UE::FXRenderingUtils::RayTracing::HasRayTracingScene(*Scene))
 	{
 		TLASSRV = UE::FXRenderingUtils::RayTracing::GetRayTracingSceneView(RHICmdList, *Scene, ReferenceView);
+
+		InlineRayTracingBindingDataBufferSRV = UE::FXRenderingUtils::RayTracing::GetInlineRayTracingBindingDataBuffer(*Scene);
+
 		ViewUniformBuffer = ReferenceView.ViewUniformBuffer;
 
 		if (!SupportsNiagaraAsyncGpuTraceHwrtInline() && GRHISupportsRayTracingShaders)
@@ -441,8 +463,14 @@ void FNiagaraAsyncGpuTraceProviderHwrt::IssueTraces(FRHICommandList& RHICmdList,
 	Params.MaxRetraces = Request.MaxRetraceCount;
 	Params.RayTraceCountsOffset = Request.TraceCountsOffset;
 
+	FRHIShaderResourceView* NullSRVBinding = GEmptyStructuredBufferWithUAV->ShaderResourceViewRHI.GetReference();
+	bool bValidInlineRaytracingBindingDataBufferFound = InlineRayTracingBindingDataBufferSRV != nullptr;
+	Params.RayTracingSceneMetadata = InlineRayTracingBindingDataBufferSRV ? InlineRayTracingBindingDataBufferSRV : NullSRVBinding;
+	
 	if (SupportsNiagaraAsyncGpuTraceHwrtInline())
 	{
+		check(bValidInlineRaytracingBindingDataBufferFound || !PlatformRequiresInlineSBTForAsyncGpuTraceHwrtInline(GMaxRHIShaderPlatform));
+
 		FNiagaraCollisionRayTraceCS::FPermutationDomain PermutationVector;
 		PermutationVector.Set<FNiagaraCollisionRayTraceCS::FSupportsCollisionGroups>(CollisionGroupHash != nullptr);
 		TShaderRef<FNiagaraCollisionRayTraceCS> Shader = ShaderMap->GetShader<FNiagaraCollisionRayTraceCS>(PermutationVector);
@@ -503,6 +531,7 @@ void FNiagaraAsyncGpuTraceProviderHwrt::Reset()
 	RayTracingSBT = nullptr;
 	TLASSRV = nullptr;
 	ViewUniformBuffer = TUniformBufferRef<FViewUniformShaderParameters>();
+	InlineRayTracingBindingDataBufferSRV = nullptr;
 }
 
 #endif // RHI_RAYTRACING
diff --git a/Engine/Plugins/FX/Niagara/Source/NiagaraShader/Private/NiagaraAsyncGpuTraceProviderHwrt.h b/Engine/Plugins/FX/Niagara/Source/NiagaraShader/Private/NiagaraAsyncGpuTraceProviderHwrt.h
--- a/Engine/Plugins/FX/Niagara/Source/NiagaraShader/Private/NiagaraAsyncGpuTraceProviderHwrt.h
+++ b/Engine/Plugins/FX/Niagara/Source/NiagaraShader/Private/NiagaraAsyncGpuTraceProviderHwrt.h
@@ -28,6 +28,7 @@ private:
 	FRHIShaderBindingTable* RayTracingSBT = nullptr; // TODO: Hold reference?
 	FRHIShaderResourceView* TLASSRV = nullptr;
 	TUniformBufferRef<FViewUniformShaderParameters> ViewUniformBuffer;
+	FRHIShaderResourceView* InlineRayTracingBindingDataBufferSRV = nullptr;
 };
 
 #endif
diff --git a/Engine/Shaders/Private/RayTracing/VFXTraceRay.ush b/Engine/Shaders/Private/RayTracing/VFXTraceRay.ush
--- a/Engine/Shaders/Private/RayTracing/VFXTraceRay.ush
+++ b/Engine/Shaders/Private/RayTracing/VFXTraceRay.ush
@@ -75,6 +75,8 @@ void VfxTraceRay(
 
 #if COMPUTESHADER
 
+StructuredBuffer<FRayTracingSceneMetadataRecord> RayTracingSceneMetadata;
+
 void VfxTraceRayInline(
 	RaytracingAccelerationStructure AccelerationStructure,
 	uint RayFlags,
@@ -90,7 +92,22 @@ void VfxTraceRayInline(
 		Payload.HitT = TraceResult.HitT;
 		Payload.GPUSceneInstanceId = TraceResult.InstanceID;
 		Payload.Barycentrics = TraceResult.Barycentrics;
+
+#if PLATFORM_SUPPORTS_INLINE_RAY_TRACING_TRIANGLE_NORMALS
 		Payload.WorldNormal = TraceResult.WorldGeometryNormal;
+#elif PLATFORM_SUPPORTS_INLINE_RAY_TRACING_TRIANGLE_ATTRIBUTES
+		FTriangleBaseAttributes Attributes = LoadInlineRayTracingTriangleAttributes(
+			RayTracingSceneMetadata,
+			TraceResult.InstanceContributionToHitGroupIndex,
+			RAY_TRACING_NUM_SHADER_SLOTS,
+			TraceResult.GeometryIndex,
+			TraceResult.PrimitiveIndex
+		);
+		float3 LocalNormal = cross(Attributes.LocalPositions[2] - Attributes.LocalPositions[0], Attributes.LocalPositions[1] - Attributes.LocalPositions[0]);
+		Payload.WorldNormal = normalize(mul(LocalNormal, (float3x3)TraceResult.WorldToObject3x4));	
+#else // !PLATFORM_SUPPORTS_INLINE_RAY_TRACING_TRIANGLE_ATTRIBUTES
+		Payload.WorldNormal = 0.0f;
+#endif // !PLATFORM_SUPPORTS_INLINE_RAY_TRACING_TRIANGLE_ATTRIBUTES
 		Payload.TranslatedWorldPosition = InRay.Origin + InRay.Direction * TraceResult.HitT;
 	}
 	else
diff --git a/Engine/Source/Runtime/Renderer/Private/FXRenderingUtils.cpp b/Engine/Source/Runtime/Renderer/Private/FXRenderingUtils.cpp
--- a/Engine/Source/Runtime/Renderer/Private/FXRenderingUtils.cpp
+++ b/Engine/Source/Runtime/Renderer/Private/FXRenderingUtils.cpp
@@ -263,6 +263,17 @@ FRHIRayTracingScene* UE::FXRenderingUtils::RayTracing::GetRayTracingScene(const
 	return nullptr;
 }
 
+FRHIShaderResourceView* UE::FXRenderingUtils::RayTracing::GetInlineRayTracingBindingDataBuffer(const FSceneInterface& InScene)
+{
+	FRHIShaderResourceView *InlineRayTracingBindingDataBufferSRV = nullptr;
+	if (const FScene* Scene = InScene.GetRenderScene())
+	{
+		InlineRayTracingBindingDataBufferSRV = Scene->RayTracingSBT.GetPersistentInlineBindingDataBufferSRV(Scene->RayTracingPersistentSBTState.InlineSBT.ID);
+
+	}
+	return InlineRayTracingBindingDataBufferSRV;
+}
+
 FRHIShaderResourceView* UE::FXRenderingUtils::RayTracing::GetRayTracingSceneView(FRHICommandListBase& RHICmdList, const FSceneInterface& InScene, const FSceneView& InView)
 {
 	if (const FScene* Scene = InScene.GetRenderScene())
diff --git a/Engine/Source/Runtime/Renderer/Private/RayTracing/RayTracingShaderBindingTable.h b/Engine/Source/Runtime/Renderer/Private/RayTracing/RayTracingShaderBindingTable.h
--- a/Engine/Source/Runtime/Renderer/Private/RayTracing/RayTracingShaderBindingTable.h
+++ b/Engine/Source/Runtime/Renderer/Private/RayTracing/RayTracingShaderBindingTable.h
@@ -122,6 +122,16 @@ public:
 	FRHIShaderBindingTable* GetPersistentSBT(FRayTracingPersistentShaderBindingTableID PersistentSBTID) { return PersistentSBTs[PersistentSBTID].ShaderBindingTable; }
 	FRDGBufferRef GetPersistentInlineBindingDataBuffer(FRDGBuilder& GraphBuilder, FRayTracingPersistentShaderBindingTableID PersistentSBTID) { return GraphBuilder.RegisterExternalBuffer(PersistentSBTs[PersistentSBTID].InlineBindingDataPooledBuffer); }
 
+	FRHIShaderResourceView* GetPersistentInlineBindingDataBufferSRV(FRayTracingPersistentShaderBindingTableID PersistentSBTID) const
+	{ 
+		FRHIShaderResourceView* InlineBindingDataBufferSRV = nullptr;
+		if (PersistentSBTs.IsValidIndex(PersistentSBTID))
+		{
+			InlineBindingDataBufferSRV = PersistentSBTs[PersistentSBTID].InlineBindingDataPooledBuffer->GetSRV();
+		}
+		return InlineBindingDataBufferSRV;
+	}
+
 	uint32 GetNumShaderSlotsPerSegment() const { return NumShaderSlotsPerGeometrySegment; }
 
 	/**
diff --git a/Engine/Source/Runtime/Renderer/Public/FXRenderingUtils.h b/Engine/Source/Runtime/Renderer/Public/FXRenderingUtils.h
--- a/Engine/Source/Runtime/Renderer/Public/FXRenderingUtils.h
+++ b/Engine/Source/Runtime/Renderer/Public/FXRenderingUtils.h
@@ -74,6 +74,8 @@ namespace UE::FXRenderingUtils
 		RENDERER_API FRHIRayTracingScene* GetRayTracingScene(const FSceneInterface& Scene, const FSceneView& View);
 		RENDERER_API FRHIShaderResourceView* GetRayTracingSceneView(FRHICommandListBase& RHICmdList, const FSceneInterface& Scene, const FSceneView& View);
 
+		RENDERER_API FRHIShaderResourceView* GetInlineRayTracingBindingDataBuffer(const FSceneInterface& Scene);
+
 		UE_DEPRECATED(5.7, "Please provide Scene reference instead.")
 		inline bool HasRayTracingScene(const FSceneInterface* Scene)
 		{
