- Add mesh capture tab to Niagara debugger
- Allows you to capture what is visible to an existing static mesh or a new static mesh

#tests Open the debugger tab, select some Niagara systems and capture them.  Note that capturing static meshes does not work currently.

#Synced-CL 48951098

diff --git a/Engine/Plugins/FX/Niagara/Source/Niagara/Internal/NiagaraRendererReadback.h b/Engine/Plugins/FX/Niagara/Source/Niagara/Internal/NiagaraRendererReadback.h
--- a/Engine/Plugins/FX/Niagara/Source/Niagara/Internal/NiagaraRendererReadback.h
+++ b/Engine/Plugins/FX/Niagara/Source/Niagara/Internal/NiagaraRendererReadback.h
@@ -101,6 +101,8 @@ namespace NiagaraRendererReadback
 	// The Callback will also be invoked on the GameThread once complete
 	NIAGARA_API void EnqueueReadback(TConstArrayView<UNiagaraComponent*> Components, FNiagaraRendererReadbackComplete Callback, const FNiagaraRendererReadbackParameters& Parameters = FNiagaraRendererReadbackParameters());
 	NIAGARA_API void EnqueueReadback(UNiagaraComponent* Component, FNiagaraRendererReadbackComplete Callback, const FNiagaraRendererReadbackParameters& Parameters = FNiagaraRendererReadbackParameters());
+	// Capture the next frame's renderer data for the provided world
+	NIAGARA_API void EnqueueReadback(UWorld* World, FNiagaraRendererReadbackComplete Callback, const FNiagaraRendererReadbackParameters& Parameters = FNiagaraRendererReadbackParameters());
 
 	inline void BeginCapture() { ++GIsCapturing; }
 	inline bool IsCapturing() { return GIsCapturing != 0; }
diff --git a/Engine/Plugins/FX/Niagara/Source/Niagara/Private/NiagaraRendererReadback.cpp b/Engine/Plugins/FX/Niagara/Source/Niagara/Private/NiagaraRendererReadback.cpp
--- a/Engine/Plugins/FX/Niagara/Source/Niagara/Private/NiagaraRendererReadback.cpp
+++ b/Engine/Plugins/FX/Niagara/Source/Niagara/Private/NiagaraRendererReadback.cpp
@@ -53,6 +53,7 @@ namespace NiagaraRendererReadback
 	{
 		bool													bExportMaterials = true;
 		bool													bApplyWPO = false;
+		bool													bCaptureAllPrimitives = false;
 		TOptional<int>											ViewIndexToCapture;
 		FNiagaraRendererReadbackResult							Result;
 
@@ -340,7 +341,7 @@ namespace NiagaraRendererReadback
 		{
 			for (const FRendererReadbackRequestPtr& Request : ReadbackRequests)
 			{
-				if (!Request->PrimitivesToCapture.Contains(SceneProxy->GetPrimitiveComponentId()))
+				if (Request->bCaptureAllPrimitives == false && !Request->PrimitivesToCapture.Contains(SceneProxy->GetPrimitiveComponentId()))
 				{
 					continue;
 				}
@@ -392,103 +393,133 @@ namespace NiagaraRendererReadback
 	};
 
 	template<typename T>
-	T GetVertexValue(const FNiagaraRendererReadbackResult& Result, uint32 Vertex, uint32 ComponentOffset, const T& DefaultValue)
+	T GetVertexValue(const FNiagaraRendererReadbackResult& Result, uint32 Vertex, uint32 ComponentOffset, uint32 ElementOffset, const T& DefaultValue)
 	{
 		T Value = DefaultValue;
-		if (Vertex < Result.NumVertices)
+		if (Vertex < Result.NumVertices && ComponentOffset != INDEX_NONE)
 		{
-			const uint32 Offset = (Result.VertexStride * Vertex) + ComponentOffset;
+			const uint32 Offset = (Result.VertexStride * Vertex) + ComponentOffset + ElementOffset;
 			FMemory::Memcpy(&Value, Result.VertexData.GetData() + Offset, sizeof(T));
 		}
 		return Value;
 	}
-} //namespace NiagaraRendererReadback
 
-//////////////////////////////////////////////////////////////////////////////////////////////////////////////////////////////////////////////////////////////////////////////
+	template<typename T>
+	T GetVertexValue(const FNiagaraRendererReadbackResult& Result, uint32 Vertex, uint32 ComponentOffset, const T& DefaultValue)
+	{
+		return GetVertexValue(Result, Vertex, ComponentOffset, 0, DefaultValue);
+	}
 
-void NiagaraRendererReadback::EnqueueReadback(TConstArrayView<UNiagaraComponent*> Components, FNiagaraRendererReadbackComplete Callback, const FNiagaraRendererReadbackParameters& Parameters)
-{
-	using namespace NiagaraRendererReadback;
+	void InternalEnqueueReadback(FNiagaraGpuComputeDispatchInterface* ComputeInterface, FRendererReadbackRequestPtr ReadbackRequest, const FNiagaraRendererReadbackParameters& Parameters)
+	{
+		using namespace NiagaraRendererReadback;
 
-	FNiagaraGpuComputeDispatchInterface* ComputeInterface = nullptr;
-	FRendererReadbackRequestPtr ReadbackRequest = MakeShared<FRendererReadbackRequest>();
-	ReadbackRequest->CompleteCallback = MoveTemp(Callback);
+		// Copy over parameters
+		ReadbackRequest->bExportMaterials	= Parameters.bExportMaterials;
+		ReadbackRequest->bApplyWPO			= Parameters.bApplyWPO;
+		ReadbackRequest->ViewIndexToCapture	= Parameters.ViewIndexToCapture;
 
-	for (UNiagaraComponent* Component : Components)
-	{
-		UWorld* ComponentWorld = ::IsValid(Component) ? Component->GetWorld() : nullptr;
-		if (ComponentWorld == nullptr)
+		// Build Vertex Output
 		{
-			continue;
+			FNiagaraRendererReadbackResult& Result = ReadbackRequest->Result;
+			if (Parameters.bExportPosition)
+			{
+				Result.VertexPositionOffset = Result.VertexStride;
+				Result.VertexStride += sizeof(FVector3f);
+			}
+			if (Parameters.bExportTangentBasis)
+			{
+				Result.VertexTangentBasisOffset = Result.VertexStride;
+				Result.VertexStride += sizeof(FVector3f) * 3;
+			}
+			if (Parameters.bExportColor)
+			{
+				Result.VertexColorOffset = Result.VertexStride;
+				Result.VertexStride += sizeof(FVector4f);
+			}
+			if (Parameters.ExportNumTexCoords > 0)
+			{
+				const int32 NumTexCoords = FMath::Min<int32>(Parameters.ExportNumTexCoords, MAX_STATIC_TEXCOORDS);
+				Result.VertexTexCoordNum = NumTexCoords;
+				Result.VertexTexCoordOffset = Result.VertexStride;
+				Result.VertexStride += sizeof(FVector2f) * NumTexCoords;
+			}
 		}
 
-		FNiagaraGpuComputeDispatchInterface* ComponentComputeInterface = FNiagaraGpuComputeDispatchInterface::Get(ComponentWorld);
-		ComputeInterface = ComputeInterface == nullptr ? ComponentComputeInterface : ComputeInterface;
+		ENQUEUE_RENDER_COMMAND(EnqueueRendererReadback)(
+			[ComputeInterface, ReadbackRequest](FRHICommandListImmediate&)
+			{
+				FRendererReadbackComputeManager& ReadbackManager = ComputeInterface->GetOrCreateDataManager<FRendererReadbackComputeManager>();
+				ReadbackManager.ReadbackRequests.Emplace(ReadbackRequest);
+			}
+		);
+	}
+
+	void EnqueueReadbackForComponents(TConstArrayView<UNiagaraComponent*> Components, FNiagaraRendererReadbackComplete Callback, const FNiagaraRendererReadbackParameters& Parameters)
+	{
+		FNiagaraGpuComputeDispatchInterface* ComputeInterface = nullptr;
+		FRendererReadbackRequestPtr ReadbackRequest = MakeShared<FRendererReadbackRequest>();
+		ReadbackRequest->CompleteCallback = MoveTemp(Callback);
+
+		ReadbackRequest->bCaptureAllPrimitives = false;
 
-		if (ComponentComputeInterface == nullptr)
+		for (UNiagaraComponent* Component : Components)
 		{
-			ReadbackRequest->Result.Errors.AddUnique(TEXT("Request failed due to no compute interface."));
-			ExecuteCompleteCallback(ReadbackRequest);
-			return;
+			UWorld* ComponentWorld = ::IsValid(Component) ? Component->GetWorld() : nullptr;
+			if (ComponentWorld == nullptr)
+			{
+				continue;
+			}
+
+			FNiagaraGpuComputeDispatchInterface* ComponentComputeInterface = FNiagaraGpuComputeDispatchInterface::Get(ComponentWorld);
+			ComputeInterface = ComputeInterface == nullptr ? ComponentComputeInterface : ComputeInterface;
+
+			if (ComponentComputeInterface == nullptr)
+			{
+				ReadbackRequest->Result.Errors.AddUnique(TEXT("Request failed due to no compute interface."));
+				ExecuteCompleteCallback(ReadbackRequest);
+				return;
+			}
+			else if (ComputeInterface != ComponentComputeInterface)
+			{
+				ReadbackRequest->Result.Errors.AddUnique(TEXT("Request failed due to components with mismatching compute interface."));
+				ExecuteCompleteCallback(ReadbackRequest);
+				return;
+			}
+
+			ReadbackRequest->PrimitivesToCapture.Add(Component->GetPrimitiveSceneId());
 		}
-		else if (ComputeInterface != ComponentComputeInterface)
+
+		// Nothing added then there is nothing to do complete
+		if (ReadbackRequest->PrimitivesToCapture.Num() == 0)
 		{
-			ReadbackRequest->Result.Errors.AddUnique(TEXT("Request failed due to components with mismatching compute interface."));
+			ReadbackRequest->Result.Errors.AddUnique(TEXT("No primitives found to capture."));
 			ExecuteCompleteCallback(ReadbackRequest);
 			return;
 		}
 
-		ReadbackRequest->PrimitivesToCapture.Add(Component->GetPrimitiveSceneId());
+		InternalEnqueueReadback(ComputeInterface, ReadbackRequest, Parameters);
 	}
 
-	// Nothing added then there is nothing to do complete
-	if (ReadbackRequest->PrimitivesToCapture.Num() == 0)
+	void EnqueueReadbackForWorld(UWorld* World, FNiagaraRendererReadbackComplete Callback, const FNiagaraRendererReadbackParameters& Parameters)
 	{
-		ReadbackRequest->Result.Errors.AddUnique(TEXT("No primitives found to capture."));
-		ExecuteCompleteCallback(ReadbackRequest);
-		return;
-	}
-
-	// Copy over parameters
-	ReadbackRequest->bExportMaterials	= Parameters.bExportMaterials;
-	ReadbackRequest->bApplyWPO			= Parameters.bApplyWPO;
-	ReadbackRequest->ViewIndexToCapture	= Parameters.ViewIndexToCapture;
+		FRendererReadbackRequestPtr ReadbackRequest = MakeShared<FRendererReadbackRequest>();
+		ReadbackRequest->CompleteCallback = MoveTemp(Callback);
+		ReadbackRequest->bCaptureAllPrimitives = true;
 
-	// Build Vertex Output
-	{
-		FNiagaraRendererReadbackResult& Result = ReadbackRequest->Result;
-		if (Parameters.bExportPosition)
-		{
-			Result.VertexPositionOffset = Result.VertexStride;
-			Result.VertexStride += sizeof(FVector3f);
-		}
-		if (Parameters.bExportTangentBasis)
-		{
-			Result.VertexTangentBasisOffset = Result.VertexStride;
-			Result.VertexStride += sizeof(FVector3f) * 3;
-		}
-		if (Parameters.bExportColor)
-		{
-			Result.VertexColorOffset = Result.VertexStride;
-			Result.VertexStride += sizeof(FVector4f);
-		}
-		if (Parameters.ExportNumTexCoords > 0)
+		FNiagaraGpuComputeDispatchInterface* ComputeInterface = FNiagaraGpuComputeDispatchInterface::Get(World);
+		if (ComputeInterface == nullptr)
 		{
-			const int32 NumTexCoords = FMath::Min<int32>(Parameters.ExportNumTexCoords, MAX_STATIC_TEXCOORDS);
-			Result.VertexTexCoordNum = NumTexCoords;
-			Result.VertexTexCoordOffset = Result.VertexStride;
-			Result.VertexStride += sizeof(FVector2f) * NumTexCoords;
+			ReadbackRequest->Result.Errors.AddUnique(TEXT("Request failed due to no compute interface."));
+			ExecuteCompleteCallback(ReadbackRequest);
+			return;
 		}
+
+		InternalEnqueueReadback(ComputeInterface, ReadbackRequest, Parameters);
 	}
+} //namespace NiagaraRendererReadback
 
-	ENQUEUE_RENDER_COMMAND(EnqueueRendererReadback)(
-		[ComputeInterface, ReadbackRequest](FRHICommandListImmediate&)
-		{
-			FRendererReadbackComputeManager& ReadbackManager = ComputeInterface->GetOrCreateDataManager<FRendererReadbackComputeManager>();
-			ReadbackManager.ReadbackRequests.Emplace(ReadbackRequest);
-		}
-	);
-}
+//////////////////////////////////////////////////////////////////////////////////////////////////////////////////////////////////////////////////////////////////////////////
 
 FVector3f FNiagaraRendererReadbackResult::GetPosition(uint32 Vertex) const
 {
@@ -507,12 +538,12 @@ FVector3f FNiagaraRendererReadbackResult::GetTangentX(uint32 Vertex) const
 
 FVector3f FNiagaraRendererReadbackResult::GetTangentY(uint32 Vertex) const
 {
-	return NiagaraRendererReadback::GetVertexValue(*this, Vertex, VertexTangentBasisOffset + sizeof(FVector3f), FVector3f::YAxisVector);
+	return NiagaraRendererReadback::GetVertexValue(*this, Vertex, VertexTangentBasisOffset, sizeof(FVector3f), FVector3f::YAxisVector);
 }
 
 FVector3f FNiagaraRendererReadbackResult::GetTangentZ(uint32 Vertex) const
 {
-	return NiagaraRendererReadback::GetVertexValue(*this, Vertex, VertexTangentBasisOffset + (sizeof(FVector3f) * 2), FVector3f::ZAxisVector);
+	return NiagaraRendererReadback::GetVertexValue(*this, Vertex, VertexTangentBasisOffset, sizeof(FVector3f) * 2, FVector3f::ZAxisVector);
 }
 
 FVector2f FNiagaraRendererReadbackResult::GetTexCoord(uint32 Vertex, uint32 TexCoordIndex) const
@@ -521,12 +552,22 @@ FVector2f FNiagaraRendererReadbackResult::GetTexCoord(uint32 Vertex, uint32 TexC
 	{
 		return FVector2f::ZeroVector;
 	}
-	return NiagaraRendererReadback::GetVertexValue(*this, Vertex, VertexTexCoordOffset + (sizeof(FVector2f) * TexCoordIndex), FVector2f::ZeroVector);
+	return NiagaraRendererReadback::GetVertexValue(*this, Vertex, VertexTexCoordOffset, sizeof(FVector2f) * TexCoordIndex, FVector2f::ZeroVector);
+}
+
+void NiagaraRendererReadback::EnqueueReadback(TConstArrayView<UNiagaraComponent*> Components, FNiagaraRendererReadbackComplete Callback, const FNiagaraRendererReadbackParameters& Parameters)
+{
+	EnqueueReadbackForComponents(Components, MoveTemp(Callback), Parameters);
 }
 
 void NiagaraRendererReadback::EnqueueReadback(UNiagaraComponent* Component, TFunction<void(const FNiagaraRendererReadbackResult& Result)> Callback, const FNiagaraRendererReadbackParameters& Parameters)
 {
-	EnqueueReadback(MakeArrayView(&Component, 1), MoveTemp(Callback), Parameters);
+	EnqueueReadbackForComponents(MakeArrayView(&Component, 1), MoveTemp(Callback), Parameters);
+}
+
+void NiagaraRendererReadback::EnqueueReadback(UWorld* World, FNiagaraRendererReadbackComplete Callback, const FNiagaraRendererReadbackParameters& Parameters)
+{
+	EnqueueReadbackForWorld(World, MoveTemp(Callback), Parameters);
 }
 
 void NiagaraRendererReadback::CaptureMeshBatch(const FSceneView* View, const FNiagaraSceneProxy* SceneProxy, const FMeshBatch& MeshBatch, uint32 NumInstances, uint32 NumVerticesPerInstance)
diff --git a/Engine/Plugins/FX/Niagara/Source/NiagaraEditor/Private/Widgets/SNiagaraDebugger.cpp b/Engine/Plugins/FX/Niagara/Source/NiagaraEditor/Private/Widgets/SNiagaraDebugger.cpp
--- a/Engine/Plugins/FX/Niagara/Source/NiagaraEditor/Private/Widgets/SNiagaraDebugger.cpp
+++ b/Engine/Plugins/FX/Niagara/Source/NiagaraEditor/Private/Widgets/SNiagaraDebugger.cpp
@@ -16,6 +16,7 @@
 #include "NiagaraEditorStyle.h"
 #include "PropertyEditorModule.h"
 #include "SNiagaraDebuggerSpawn.h"
+#include "SNiagaraDebuggerMeshCapture.h"
 #include "Styling/AppStyle.h"
 #include "UObject/UObjectIterator.h"
 #include "Widgets/Docking/SDockTab.h"
@@ -131,8 +132,6 @@ namespace NiagaraDebugSpawnTab
 
 	static void RegisterTabSpawner(const TSharedPtr<FTabManager>& TabManager, TSharedPtr<FNiagaraDebugger>& Debugger)
 	{
-		FPropertyEditorModule& PropertyModule = FModuleManager::LoadModuleChecked<FPropertyEditorModule>("PropertyEditor");
-
 		TabManager->RegisterTabSpawner(
 			TabName,
 			FOnSpawnTab::CreateLambda(
@@ -153,6 +152,31 @@ namespace NiagaraDebugSpawnTab
 	}
 }
 
+namespace NiagaraDebuggerMeshCaptureSpawnTab
+{
+	static const FName TabName = FName(TEXT("NiagaraDebuggerMeshCapture"));
+
+	static void RegisterTabSpawner(const TSharedPtr<FTabManager>& TabManager)
+	{
+		TabManager->RegisterTabSpawner(
+			TabName,
+			FOnSpawnTab::CreateLambda(
+				[=](const FSpawnTabArgs&)
+				{
+					return SNew(SDockTab)
+						.TabRole(ETabRole::PanelTab)
+						.Label(LOCTEXT("NiagaraDebuggerMeshCapture", "Mesh Capture"))
+						[
+							SNew(SNiagaraDebuggerMeshCapture)
+						];
+				}
+			)
+		)
+		.SetDisplayName(LOCTEXT("NiagaraDebuggerMeshCaptureTitle", "Mesh Capture"))
+		.SetTooltipText(LOCTEXT("NiagaraDebuggerMeshCaptureTooltipText", "Open the mesh capture tab."));
+	}
+}
+
 namespace NiagaraPerformanceTab
 {
 	static const FName TabName = FName(TEXT("PerformanceTab"));
@@ -766,6 +790,7 @@ void SNiagaraDebugger::Construct(const FArguments& InArgs)
 #if WITH_UNREAL_TARGET_DEVELOPER_TOOLS
 	NiagaraDebugSpawnTab::RegisterTabSpawner(TabManager, Debugger);
 	NiagaraPerformanceTab::RegisterTabSpawner(TabManager, FOnExecConsoleCommand::CreateSP(Debugger.ToSharedRef(), &FNiagaraDebugger::ExecConsoleCommand));
+	NiagaraDebuggerMeshCaptureSpawnTab::RegisterTabSpawner(TabManager);
 #endif
 
 #if WITH_SESSION_FRONTEND
@@ -797,6 +822,7 @@ void SNiagaraDebugger::Construct(const FArguments& InArgs)
 #if WITH_UNREAL_TARGET_DEVELOPER_TOOLS
 					->AddTab(NiagaraPerformanceTab::TabName, ETabState::OpenedTab)
 					->AddTab(NiagaraDebugSpawnTab::TabName, ETabState::OpenedTab)
+					->AddTab(NiagaraDebuggerMeshCaptureSpawnTab::TabName, ETabState::OpenedTab)
 #endif
 					->SetForegroundTab(NiagaraDebugHudTab::TabName)
 				)
diff --git a/Engine/Plugins/FX/Niagara/Source/NiagaraEditor/Private/Widgets/SNiagaraDebuggerMeshCapture.cpp b/Engine/Plugins/FX/Niagara/Source/NiagaraEditor/Private/Widgets/SNiagaraDebuggerMeshCapture.cpp
new file mode 100644
--- /dev/null
+++ b/Engine/Plugins/FX/Niagara/Source/NiagaraEditor/Private/Widgets/SNiagaraDebuggerMeshCapture.cpp
@@ -0,0 +1,215 @@
+// Copyright Epic Games, Inc. All Rights Reserved.
+
+#include "SNiagaraDebuggerMeshCapture.h"
+#include "NiagaraBakerRendererOutputStaticMesh.h"
+#include "NiagaraEditorStyle.h"
+
+#include "Engine/Selection.h"
+#include "Engine/StaticMesh.h"
+#include "Framework/MultiBox/MultiBoxBuilder.h"
+#include "Misc/MessageDialog.h"
+#include "Editor.h"
+#include "IAssetTools.h"
+#include "IStructureDetailsView.h"
+#include "PropertyEditorModule.h"
+
+#include UE_INLINE_GENERATED_CPP_BY_NAME(SNiagaraDebuggerMeshCapture)
+
+#if WITH_NIAGARA_DEBUGGER
+#define LOCTEXT_NAMESPACE "SNiagaraDebuggerMeshCapture"
+
+namespace NiagaraDebuggerMeshCapturePrivate
+{
+	void ProcessMeshCapture(const FNiagaraRendererReadbackResult& ReadbackResult, TWeakObjectPtr<UStaticMesh> WeakMeshOutput)
+	{
+		// If a mesh was not specified pop open a dialog
+		UStaticMesh* MeshOutput = WeakMeshOutput.Get();
+		if (MeshOutput == nullptr)
+		{
+			UNiagaraBakerStaticMeshFactoryNew* Factory = NewObject<UNiagaraBakerStaticMeshFactoryNew>();
+			MeshOutput = Cast<UStaticMesh>(IAssetTools::Get().CreateAssetWithDialog(UStaticMesh::StaticClass(), Factory));
+			if (MeshOutput == nullptr)
+			{
+				return;
+			}
+		}
+
+		// Failed or no data
+		if (!ReadbackResult.NumVertices)
+		{
+			FMessageDialog::Open(
+				EAppMsgType::Ok,
+				LOCTEXT("CaptureFailedMessage", "Capture process failed, this could be due to nothing being renderered or an internal error."),
+				LOCTEXT("CaptureFailedTitle", "Failed to capture any data!")
+			);
+			return;
+		}
+
+		FNiagaraBakerRendererOutputStaticMesh::ConvertReadbackResultsToStaticMesh(ReadbackResult, MeshOutput);
+	}
+}
+
+void SNiagaraDebuggerMeshCapture::Construct(const FArguments& InArgs)
+{
+	// Create details view for the export parameters
+	FPropertyEditorModule& PropertyModule = FModuleManager::LoadModuleChecked<FPropertyEditorModule>("PropertyEditor");
+
+	FDetailsViewArgs DetailsViewArgs;
+	DetailsViewArgs.bHideSelectionTip = true;
+
+	FStructureDetailsViewArgs StructureViewArgs;
+	StructureViewArgs.bShowObjects = true;
+	StructureViewArgs.bShowAssets = true;
+	StructureViewArgs.bShowClasses = true;
+	StructureViewArgs.bShowInterfaces = true;
+
+	TSharedRef<IStructureDetailsView> StructureDetailsView = PropertyModule.CreateStructureDetailView(DetailsViewArgs, StructureViewArgs, nullptr);
+
+	TSharedPtr<FStructOnScope> StructOnScope = MakeShared<FStructOnScope>(FNiagaraDebuggerMeshCaptureSettings::StaticStruct(), reinterpret_cast<uint8*>(&CaptureSettings));
+	StructureDetailsView->SetStructureData(StructOnScope);
+
+	ChildSlot
+	[
+		SNew(SVerticalBox)
+		+ SVerticalBox::Slot()
+		.AutoHeight()
+		[
+			MakeToolbar()
+		]
+ 		+ SVerticalBox::Slot()
+		.Padding(2.0)
+		[
+			StructureDetailsView->GetWidget().ToSharedRef()
+		]
+	];
+}
+
+TSharedRef<SWidget> SNiagaraDebuggerMeshCapture::MakeToolbar()
+{
+	FToolBarBuilder ToolbarBuilder(MakeShareable(new FUICommandList), FMultiBoxCustomization::None);
+	ToolbarBuilder.BeginSection("DebuggerSpawn");
+
+	ToolbarBuilder.AddToolBarButton(
+		FUIAction(
+			FExecuteAction::CreateSP(this, &SNiagaraDebuggerMeshCapture::OnCaptureSelected),
+			FCanExecuteAction::CreateSP(this, &SNiagaraDebuggerMeshCapture::CanCaptureSelected)
+		),
+		NAME_None,
+		LOCTEXT("CaptureSelected", "Captured Selected"),
+		LOCTEXT("CaptureSelectedTooltip", "Captures the visible Niagara components that are children of the outliner selection.")
+	);
+
+	ToolbarBuilder.AddToolBarButton(
+		FUIAction(
+			FExecuteAction::CreateSP(this, &SNiagaraDebuggerMeshCapture::OnCaptureAll),
+			FCanExecuteAction::CreateSP(this, &SNiagaraDebuggerMeshCapture::CanCaptureAll)
+		),
+		NAME_None,
+		LOCTEXT("CaptureAll", "Captured All"),
+		LOCTEXT("CaptureAllTooltip", "Captures all the visible Niagara components that are next rendered to the scene.")
+	);
+
+	ToolbarBuilder.EndSection();
+
+	return ToolbarBuilder.MakeWidget();
+}
+
+TArray<UNiagaraComponent*> SNiagaraDebuggerMeshCapture::GetSelectComponents(bool bFindFirstOnly) const
+{
+	TArray<UNiagaraComponent*> SelectedComponents;
+	if (GEditor)
+	{
+		if (USelection* ActorSelection = GEditor->GetSelectedActors())
+		{
+			for (FSelectionIterator It(*ActorSelection); It; ++It)
+			{
+				if (AActor* Actor = Cast<AActor>(*It))
+				{
+					TArray<UNiagaraComponent*> FoundComponents;
+					Actor->GetComponents(FoundComponents);
+					SelectedComponents.Append(FoundComponents);
+					if (SelectedComponents.Num() > 0 && bFindFirstOnly)
+					{
+						return SelectedComponents;
+					}
+				}
+			}
+		}
+		if (USelection* ComponentSelection = GEditor->GetSelectedComponents())
+		{
+			for (FSelectionIterator It(*ComponentSelection); It; ++It)
+			{
+				if (UNiagaraComponent* Component = Cast<UNiagaraComponent>(*It))
+				{
+					SelectedComponents.Add(Component);
+					if (bFindFirstOnly)
+					{
+						return SelectedComponents;
+					}
+				}
+			}
+		}
+	}
+	return SelectedComponents;
+}
+
+bool SNiagaraDebuggerMeshCapture::IsOutputLocationValid() const
+{
+	//return CaptureSettings.MeshOutput != nullptr;
+	return true;
+}
+
+bool SNiagaraDebuggerMeshCapture::CanCaptureSelected() const
+{
+	return IsOutputLocationValid() && GetSelectComponents(true).Num() > 0;
+}
+
+void SNiagaraDebuggerMeshCapture::OnCaptureSelected()
+{
+	TArray<UNiagaraComponent*> SelectedComponents = GetSelectComponents();
+	if (SelectedComponents.Num() == 0)
+	{
+		return;
+	}
+
+#if WITH_NIAGARA_RENDERER_READBACK
+	NiagaraRendererReadback::EnqueueReadback(
+		SelectedComponents,
+		[WeakMeshOutput=MakeWeakObjectPtr(CaptureSettings.MeshOutput)](const FNiagaraRendererReadbackResult& ReadbackResult)
+		{
+			NiagaraDebuggerMeshCapturePrivate::ProcessMeshCapture(ReadbackResult, WeakMeshOutput);
+		},
+		CaptureSettings.ExportParameters
+	);
+#endif //WITH_NIAGARA_RENDERER_READBACK
+}
+
+bool SNiagaraDebuggerMeshCapture::CanCaptureAll() const
+{
+#if WITH_NIAGARA_RENDERER_READBACK
+	UWorld* World = GEditor ? GEditor->GetEditorWorldContext().World() : nullptr;
+	return World != nullptr;
+#else
+	return false;
+#endif //WITH_NIAGARA_RENDERER_READBACK
+}
+
+void SNiagaraDebuggerMeshCapture::OnCaptureAll()
+{
+#if WITH_NIAGARA_RENDERER_READBACK
+	UWorld* World = GEditor ? GEditor->GetEditorWorldContext().World() : nullptr;
+
+	NiagaraRendererReadback::EnqueueReadback(
+		World,
+		[WeakMeshOutput=MakeWeakObjectPtr(CaptureSettings.MeshOutput)](const FNiagaraRendererReadbackResult& ReadbackResult)
+		{
+			NiagaraDebuggerMeshCapturePrivate::ProcessMeshCapture(ReadbackResult, WeakMeshOutput);
+		},
+		CaptureSettings.ExportParameters
+	);
+#endif //WITH_NIAGARA_RENDERER_READBACK
+}
+
+#undef LOCTEXT_NAMESPACE
+#endif //WITH_NIAGARA_DEBUGGER
+
diff --git a/Engine/Plugins/FX/Niagara/Source/NiagaraEditor/Private/Widgets/SNiagaraDebuggerMeshCapture.h b/Engine/Plugins/FX/Niagara/Source/NiagaraEditor/Private/Widgets/SNiagaraDebuggerMeshCapture.h
new file mode 100644
--- /dev/null
+++ b/Engine/Plugins/FX/Niagara/Source/NiagaraEditor/Private/Widgets/SNiagaraDebuggerMeshCapture.h
@@ -0,0 +1,53 @@
+// Copyright Epic Games, Inc. All Rights Reserved.
+
+#pragma once
+
+#include "Widgets/SCompoundWidget.h"
+
+#include "NiagaraDebugger.h"
+#include "NiagaraRendererReadback.h"
+
+//#include "TickableEditorObject.h"
+
+#include "SNiagaraDebuggerMeshCapture.generated.h"
+
+USTRUCT()
+struct FNiagaraDebuggerMeshCaptureSettings
+{
+	GENERATED_BODY()
+	
+	UPROPERTY(EditAnywhere, Category = "Settings")
+	TObjectPtr<UStaticMesh> MeshOutput;
+
+	UPROPERTY(EditAnywhere, Category = "Settings", meta = (ShowOnlyInnerProperties))
+	FNiagaraRendererReadbackParameters ExportParameters;
+};
+
+#if WITH_NIAGARA_DEBUGGER
+class SNiagaraDebuggerMeshCapture : public SCompoundWidget
+{
+public:
+	static const FName TabName;
+
+	SLATE_BEGIN_ARGS(SNiagaraDebuggerMeshCapture) {}
+	SLATE_END_ARGS();
+
+	void Construct(const FArguments& InArgs);
+
+private:
+	TSharedRef<SWidget> MakeToolbar();
+
+	TArray<UNiagaraComponent*> GetSelectComponents(bool bFindFirstOnly = false) const;
+	
+	bool IsOutputLocationValid() const;
+
+	bool CanCaptureSelected() const;
+	void OnCaptureSelected();
+
+	bool CanCaptureAll() const;
+	void OnCaptureAll();
+
+private:
+	FNiagaraDebuggerMeshCaptureSettings	CaptureSettings;
+};
+#endif //WITH_NIAGARA_DEBUGGER
