- Fix TSan issue for allocation estimation

#Synced-CL 49036664

diff --git a/Engine/Plugins/FX/Niagara/Source/Niagara/Classes/NiagaraEmitter.h b/Engine/Plugins/FX/Niagara/Source/Niagara/Classes/NiagaraEmitter.h
--- a/Engine/Plugins/FX/Niagara/Source/Niagara/Classes/NiagaraEmitter.h
+++ b/Engine/Plugins/FX/Niagara/Source/Niagara/Classes/NiagaraEmitter.h
@@ -227,9 +227,36 @@ struct FNiagaraDetailsLevelScaleOverrides
 struct FMemoryRuntimeEstimation
 {
 	TMap<uint64, int32> RuntimeAllocations;
-	bool IsEstimationDirty = false;
+	UE_DEPRECATED(5.8, "Direct use of IsEstimationDirty is deprecated please use the get / set methods")
+	std::atomic<bool> IsEstimationDirty = false;
 	int32 AllocationEstimate = 0;
 
+	FMemoryRuntimeEstimation() = default;
+	FMemoryRuntimeEstimation(const FMemoryRuntimeEstimation& Other) { *this = Other; }
+
+	FMemoryRuntimeEstimation& operator=(const FMemoryRuntimeEstimation& Other)
+	{
+		EstimationCriticalSection = Other.EstimationCriticalSection;
+		RuntimeAllocations = Other.RuntimeAllocations;
+		SetIsEstimationDirty(Other.GetIsEstimationDirty());
+		AllocationEstimate = Other.AllocationEstimate;
+		return *this;
+	}
+
+	void SetIsEstimationDirty(bool bIsDirty)
+	{
+		PRAGMA_DISABLE_DEPRECATION_WARNINGS
+		IsEstimationDirty.store(bIsDirty, std::memory_order_relaxed);
+		PRAGMA_ENABLE_DEPRECATION_WARNINGS
+	}
+
+	bool GetIsEstimationDirty() const
+	{
+		PRAGMA_DISABLE_DEPRECATION_WARNINGS
+		return IsEstimationDirty.load(std::memory_order_relaxed);
+		PRAGMA_ENABLE_DEPRECATION_WARNINGS
+	}
+
 	NIAGARA_API FCriticalSection* GetCriticalSection();
 	NIAGARA_API void Init();
 private:
diff --git a/Engine/Plugins/FX/Niagara/Source/Niagara/Private/NiagaraEmitter.cpp b/Engine/Plugins/FX/Niagara/Source/Niagara/Private/NiagaraEmitter.cpp
--- a/Engine/Plugins/FX/Niagara/Source/Niagara/Private/NiagaraEmitter.cpp
+++ b/Engine/Plugins/FX/Niagara/Source/Niagara/Private/NiagaraEmitter.cpp
@@ -2346,7 +2346,7 @@ void FVersionedNiagaraEmitterData::OnPostCompile(const UNiagaraEmitter& InEmitte
 
 	RuntimeEstimation.RuntimeAllocations.Empty();
 	RuntimeEstimation.AllocationEstimate = 0;
-	RuntimeEstimation.IsEstimationDirty = false;
+	RuntimeEstimation.SetIsEstimationDirty(false);
 #if STATS
 	StatDatabase.Init();
 	StatDatabase.ClearStatCaptures();
@@ -3418,12 +3418,12 @@ void FVersionedNiagaraEmitterData::ClearRuntimeAllocationEstimate(uint64 ReportH
 	{
 		RuntimeEstimation.AllocationEstimate = 0;
 		RuntimeEstimation.RuntimeAllocations.Empty();
-		RuntimeEstimation.IsEstimationDirty = true;
+		RuntimeEstimation.SetIsEstimationDirty(true);
 	}
 	else
 	{
 		RuntimeEstimation.RuntimeAllocations.Remove(ReportHandle);
-		RuntimeEstimation.IsEstimationDirty = true;
+		RuntimeEstimation.SetIsEstimationDirty(true);
 	}
 }
 
@@ -3434,7 +3434,7 @@ int32 FVersionedNiagaraEmitterData::AddRuntimeAllocation(uint64 ReporterHandle,
 	if (!Estimate || *Estimate < AllocationCount)
 	{
 		RuntimeEstimation.RuntimeAllocations.Add(ReporterHandle, AllocationCount);
-		RuntimeEstimation.IsEstimationDirty = true;
+		RuntimeEstimation.SetIsEstimationDirty(true);
 
 		// Remove a random entry when there are enough logged allocations already
 		if (RuntimeEstimation.RuntimeAllocations.Num() > 10)
@@ -3455,7 +3455,7 @@ int32 FVersionedNiagaraEmitterData::GetMaxParticleCountEstimate()
 		return PreAllocationCount;
 	}
 	
-	if (RuntimeEstimation.IsEstimationDirty)
+	if (RuntimeEstimation.GetIsEstimationDirty())
 	{
 		FScopeLock lock(RuntimeEstimation.GetCriticalSection());
 		int32 EstimationCount = RuntimeEstimation.RuntimeAllocations.Num();
@@ -3474,7 +3474,7 @@ int32 FVersionedNiagaraEmitterData::GetMaxParticleCountEstimate()
 				}
 				i++;
 			}
-			RuntimeEstimation.IsEstimationDirty = false;
+			RuntimeEstimation.SetIsEstimationDirty(false);
 		}
 	}
 	return RuntimeEstimation.AllocationEstimate;
