- Allow SubUV cutout generation to use a single frame for all frames

#test Create a NS with a sub UV sprite animation, observe the wireframe and enable / disable bCutoutIgnoreSubUV

#Synced-CL 49036664

diff --git a/Engine/Plugins/FX/Niagara/Source/Niagara/Private/NiagaraSpriteRendererProperties.cpp b/Engine/Plugins/FX/Niagara/Source/Niagara/Private/NiagaraSpriteRendererProperties.cpp
--- a/Engine/Plugins/FX/Niagara/Source/Niagara/Private/NiagaraSpriteRendererProperties.cpp
+++ b/Engine/Plugins/FX/Niagara/Source/Niagara/Private/NiagaraSpriteRendererProperties.cpp
@@ -482,13 +482,14 @@ void UNiagaraSpriteRendererProperties::PostEditChangeProperty(struct FPropertyCh
 	if (bUseMaterialCutoutTexture || CutoutTexture || DerivedData.BoundingGeometry.Num())
 	{
 		const bool bUpdateCutoutDDC =
-			PropertyName == TEXT("bUseMaterialCutoutTexture") ||
-			PropertyName == TEXT("CutoutTexture") ||
-			PropertyName == TEXT("BoundingMode") ||
-			PropertyName == TEXT("OpacitySourceMode") ||
-			PropertyName == TEXT("AlphaThreshold") ||
-			MemberPropertyName == TEXT("SubImageSize") ||
-			(bUseMaterialCutoutTexture && PropertyName == TEXT("Material"));
+			PropertyName == GET_MEMBER_NAME_CHECKED(UNiagaraSpriteRendererProperties, bUseMaterialCutoutTexture) ||
+			PropertyName == GET_MEMBER_NAME_CHECKED(UNiagaraSpriteRendererProperties, bCutoutIgnoreSubUV) ||
+			PropertyName == GET_MEMBER_NAME_CHECKED(UNiagaraSpriteRendererProperties, CutoutTexture) ||
+			PropertyName == GET_MEMBER_NAME_CHECKED(UNiagaraSpriteRendererProperties, BoundingMode) ||
+			PropertyName == GET_MEMBER_NAME_CHECKED(UNiagaraSpriteRendererProperties, OpacitySourceMode) ||
+			PropertyName == GET_MEMBER_NAME_CHECKED(UNiagaraSpriteRendererProperties, AlphaThreshold) ||
+			MemberPropertyName == GET_MEMBER_NAME_CHECKED(UNiagaraSpriteRendererProperties, SubImageSize) ||
+			(bUseMaterialCutoutTexture && PropertyName == GET_MEMBER_NAME_CHECKED(UNiagaraSpriteRendererProperties, Material));
 
 		if (bUpdateCutoutDDC)
 		{
@@ -713,7 +714,9 @@ void UNiagaraSpriteRendererProperties::CacheDerivedData()
 {
 	if (CutoutTexture)
 	{
-		const FString KeyString = FSubUVDerivedData::GetDDCKeyString(CutoutTexture->Source.GetId(), (int32)SubImageSize.X, (int32)SubImageSize.Y, (int32)BoundingMode, AlphaThreshold, (int32)OpacitySourceMode);
+		const FIntVector2 DerivedSubImageSize = bCutoutIgnoreSubUV ? FIntVector2(1, 1) : FIntVector2((int32)SubImageSize.X, (int32)SubImageSize.Y);
+
+		const FString KeyString = FSubUVDerivedData::GetDDCKeyString(CutoutTexture->Source.GetId(), DerivedSubImageSize.X, DerivedSubImageSize.Y, (int32)BoundingMode, AlphaThreshold, (int32)OpacitySourceMode);
 		TArray<uint8> Data;
 
 		COOK_STAT(auto Timer = NiagaraCutoutCookStats::UsageStats.TimeSyncWork());
@@ -726,7 +729,7 @@ void UNiagaraSpriteRendererProperties::CacheDerivedData()
 		}
 		else
 		{
-			DerivedData.Build(CutoutTexture, (int32)SubImageSize.X, (int32)SubImageSize.Y, BoundingMode, AlphaThreshold, OpacitySourceMode);
+			DerivedData.Build(CutoutTexture, DerivedSubImageSize.X, DerivedSubImageSize.Y, BoundingMode, AlphaThreshold, OpacitySourceMode);
 
 			Data.Empty(DerivedData.BoundingGeometry.Num() * sizeof(FVector2f));
 			Data.AddUninitialized(DerivedData.BoundingGeometry.Num() * sizeof(FVector2f));
@@ -734,6 +737,18 @@ void UNiagaraSpriteRendererProperties::CacheDerivedData()
 			GetDerivedDataCacheRef().Put(*KeyString, Data, GetPathName());
 			COOK_STAT(Timer.AddMiss(Data.Num()));
 		}
+
+		if (bCutoutIgnoreSubUV)
+		{
+			const int32 NumFrames = (int32)SubImageSize.X * (int32)SubImageSize.Y;
+			const int32 NumVerts = DerivedData.BoundingGeometry.Num();
+
+			DerivedData.BoundingGeometry.SetNumUninitialized(NumVerts * NumFrames);
+			for ( int32 i=1; i < NumFrames; ++i )
+			{
+				FMemory::Memcpy(&DerivedData.BoundingGeometry[NumVerts * i], &DerivedData.BoundingGeometry[0], sizeof(FVector2f) * NumVerts);
+			}
+		}
 	}
 	else
 	{
diff --git a/Engine/Plugins/FX/Niagara/Source/Niagara/Public/NiagaraSpriteRendererProperties.h b/Engine/Plugins/FX/Niagara/Source/Niagara/Public/NiagaraSpriteRendererProperties.h
--- a/Engine/Plugins/FX/Niagara/Source/Niagara/Public/NiagaraSpriteRendererProperties.h
+++ b/Engine/Plugins/FX/Niagara/Source/Niagara/Public/NiagaraSpriteRendererProperties.h
@@ -390,6 +390,10 @@ public:
 	UPROPERTY(EditAnywhere, BlueprintReadOnly, Category = "Cutout")
 	bool bUseMaterialCutoutTexture;
 
+	/** When enabled cutouts for SubUVs will generated like there is a single frame rather than atlasing the texture. */
+	UPROPERTY(EditAnywhere, BlueprintReadOnly, Category = "Cutout")
+	bool bCutoutIgnoreSubUV = false;
+
 	/** Texture to generate bounding geometry from.	*/
 	UPROPERTY(EditAnywhere, Category="Cutout", meta = (EditCondition = "!bUseMaterialCutoutTexture", EditConditionHides))
 	TObjectPtr<UTexture2D> CutoutTexture;
