Niagara - Refactor the curve model handling in the distribution UI to edit the curves directly instead of making a copy.  The rich curve model handles transactions itself which was confliciting with the distribution adapters transaction hanling.

#Synced-CL 48514417

diff --git a/Engine/Plugins/FX/Niagara/Source/NiagaraEditor/Private/Stateless/NiagaraDistributionPropertyCustomization.cpp b/Engine/Plugins/FX/Niagara/Source/NiagaraEditor/Private/Stateless/NiagaraDistributionPropertyCustomization.cpp
--- a/Engine/Plugins/FX/Niagara/Source/NiagaraEditor/Private/Stateless/NiagaraDistributionPropertyCustomization.cpp
+++ b/Engine/Plugins/FX/Niagara/Source/NiagaraEditor/Private/Stateless/NiagaraDistributionPropertyCustomization.cpp
@@ -662,6 +662,15 @@ public:
 		return nullptr;
 	}
 
+	virtual FRichCurve* GetCurveValue(int32 ChannelIndex) override
+	{
+		if (ChannelIndex < EditorNumChannels && ChannelIndex < SourceDistribution->ChannelCurves.Num())
+		{
+			return &SourceDistribution->ChannelCurves[ChannelIndex];
+		}
+		return nullptr;
+	}
+
 	virtual void SetCurveValue(int32 ChannelIndex, const FRichCurve& InValue, bool bShouldTransact) override
 	{
 		UObject* OwnerObject = OwnerObjectWeak.Get();
diff --git a/Engine/Plugins/FX/Niagara/Source/NiagaraEditor/Private/Widgets/INiagaraDistributionAdapter.h b/Engine/Plugins/FX/Niagara/Source/NiagaraEditor/Private/Widgets/INiagaraDistributionAdapter.h
--- a/Engine/Plugins/FX/Niagara/Source/NiagaraEditor/Private/Widgets/INiagaraDistributionAdapter.h
+++ b/Engine/Plugins/FX/Niagara/Source/NiagaraEditor/Private/Widgets/INiagaraDistributionAdapter.h
@@ -131,6 +131,9 @@ public:
 	/** Gets a curve value for the specified channel */
 	virtual const FRichCurve* GetCurveValue(int32 ChannelIndex) const = 0;
 
+	/** Gets a curve value for the specified channel */
+	virtual FRichCurve* GetCurveValue(int32 ChannelIndex) = 0;
+
 	/** Set a curve value for the specified channel */
 	virtual void SetCurveValue(int32 ChannelIndex, const FRichCurve& InValue, bool bShouldTransact = true) = 0;
 
diff --git a/Engine/Plugins/FX/Niagara/Source/NiagaraEditor/Private/Widgets/SNiagaraDistributionCurveEditor.cpp b/Engine/Plugins/FX/Niagara/Source/NiagaraEditor/Private/Widgets/SNiagaraDistributionCurveEditor.cpp
--- a/Engine/Plugins/FX/Niagara/Source/NiagaraEditor/Private/Widgets/SNiagaraDistributionCurveEditor.cpp
+++ b/Engine/Plugins/FX/Niagara/Source/NiagaraEditor/Private/Widgets/SNiagaraDistributionCurveEditor.cpp
@@ -34,6 +34,34 @@
 
 namespace
 {
+	
+	class FNiagaraDistributionCurveEditorModel : public FRichCurveEditorModel
+	{
+	public:
+		FNiagaraDistributionCurveEditorModel(TSharedRef<INiagaraDistributionAdapter> InDistributionAdapter, int32 InChannelIndex)
+			: FRichCurveEditorModel(GetOwnerFromDistribution(InDistributionAdapter.Get()))
+			, DistributionAdapter(InDistributionAdapter)
+			, ChannelIndex(InChannelIndex)
+		{
+		}
+
+		virtual bool IsValid() const override { return DistributionAdapter.IsValid() && DistributionAdapter->GetCurveValue(ChannelIndex) != nullptr; }
+		virtual FRichCurve& GetRichCurve() { return *DistributionAdapter->GetCurveValue(ChannelIndex); }
+		virtual const FRichCurve& GetReadOnlyRichCurve() const override { return *DistributionAdapter->GetCurveValue(ChannelIndex); }
+
+	private:
+		static UObject* GetOwnerFromDistribution(INiagaraDistributionAdapter& InDistribution)
+		{
+			TArray<UObject*> Owners;
+			InDistribution.GetOwners(Owners);
+			return Owners.Num() == 1 ? Owners[0] : nullptr;
+		}
+
+	private:
+		TSharedPtr<INiagaraDistributionAdapter>	DistributionAdapter;
+		int32 ChannelIndex = 0;
+	};
+
 	struct FNiagaraDistributionCurveEditorBounds : public ICurveEditorBounds
 	{
 		FNiagaraDistributionCurveEditorBounds(TSharedRef<FNiagaraDistributionCurveEditorOptions> InCurveEditorOptions)
@@ -63,12 +91,9 @@ namespace
 			DistributionAdapter = InDistributionAdapter;
 			ChannelIndex = InChannelIndex;
 
-			Curve			= *DistributionAdapter->GetCurveValue(ChannelIndex);
 			CurveName		= DistributionAdapter->GetChannelDisplayName(ChannelIndex);
 			CurveToolTip	= DistributionAdapter->GetChannelToolTip(ChannelIndex);
 			CurveColor		= DistributionAdapter->GetChannelColor(ChannelIndex).GetSpecifiedColor();
-
-			DistributionAdapter->OnDistributionValueChanged().AddSP(this, &FNiagaraDistributionCurveEditorDetailsTreeItem::OnDistributionValueChanged);
 		}
 
 		virtual TSharedPtr<SWidget> GenerateCurveEditorTreeWidget(const FName& InColumnName, TWeakPtr<FCurveEditor> InCurveEditor, FCurveEditorTreeItemID InTreeItemID, const TSharedRef<ITableRow>& TableRow) override
@@ -124,30 +149,17 @@ namespace
 				return;
 			}
 
-			TUniquePtr<FRichCurveEditorModelRaw> NewCurve = MakeUnique<FRichCurveEditorModelRaw>(&Curve, OwnerObjects[0]);
+			TUniquePtr<FNiagaraDistributionCurveEditorModel> NewCurve = MakeUnique<FNiagaraDistributionCurveEditorModel>(DistributionAdapter.ToSharedRef(), ChannelIndex);
 			NewCurve->SetShortDisplayName(CurveName);
 			NewCurve->SetColor(CurveColor);
-			NewCurve->OnCurveModified().AddSP(this, &FNiagaraDistributionCurveEditorDetailsTreeItem::OnCurveChanged);
 			OutCurveModels.Add(MoveTemp(NewCurve));
 		}
 
 		virtual bool PassesFilter(const FCurveEditorTreeFilter* InFilter) const override { return true; }
 
-	private:
-		void OnCurveChanged()
-		{
-			DistributionAdapter->SetCurveValue(ChannelIndex, Curve, false);
-		}
-
-		void OnDistributionValueChanged()
-		{
-			Curve = *DistributionAdapter->GetCurveValue(ChannelIndex);
-		}
-
 	private:
 		TSharedPtr<INiagaraDistributionAdapter>	DistributionAdapter;
 		int										ChannelIndex = 0;
-		FRichCurve								Curve;
 		FText									CurveName;
 		FText									CurveToolTip;
 		FLinearColor							CurveColor;
