Niagara - Fix an issue where curve edits on distributions weren't rebuilding the underlying runtime data.

#Synced-CL 48514417

diff --git a/Engine/Plugins/FX/Niagara/Source/NiagaraEditor/Private/Stateless/NiagaraDistributionPropertyCustomization.cpp b/Engine/Plugins/FX/Niagara/Source/NiagaraEditor/Private/Stateless/NiagaraDistributionPropertyCustomization.cpp
--- a/Engine/Plugins/FX/Niagara/Source/NiagaraEditor/Private/Stateless/NiagaraDistributionPropertyCustomization.cpp
+++ b/Engine/Plugins/FX/Niagara/Source/NiagaraEditor/Private/Stateless/NiagaraDistributionPropertyCustomization.cpp
@@ -671,6 +671,14 @@ public:
 		return nullptr;
 	}
 
+	virtual void NotifyCurveValueChangedExternally(int32 ChannelIndex) override
+	{
+		if (PropertyHandle.IsValid())
+		{
+			PropertyHandle->NotifyPostChange(EPropertyChangeType::ValueSet);
+		}
+	}
+
 	virtual void SetCurveValue(int32 ChannelIndex, const FRichCurve& InValue, bool bShouldTransact) override
 	{
 		UObject* OwnerObject = OwnerObjectWeak.Get();
diff --git a/Engine/Plugins/FX/Niagara/Source/NiagaraEditor/Private/Widgets/INiagaraDistributionAdapter.h b/Engine/Plugins/FX/Niagara/Source/NiagaraEditor/Private/Widgets/INiagaraDistributionAdapter.h
--- a/Engine/Plugins/FX/Niagara/Source/NiagaraEditor/Private/Widgets/INiagaraDistributionAdapter.h
+++ b/Engine/Plugins/FX/Niagara/Source/NiagaraEditor/Private/Widgets/INiagaraDistributionAdapter.h
@@ -134,6 +134,9 @@ public:
 	/** Gets a curve value for the specified channel */
 	virtual FRichCurve* GetCurveValue(int32 ChannelIndex) = 0;
 
+	/** Notifies this distribution adapter that a curve from GetCurveValue has been modified externally. */
+	virtual void NotifyCurveValueChangedExternally(int32 ChannelIndex) = 0;
+
 	/** Set a curve value for the specified channel */
 	virtual void SetCurveValue(int32 ChannelIndex, const FRichCurve& InValue, bool bShouldTransact = true) = 0;
 
diff --git a/Engine/Plugins/FX/Niagara/Source/NiagaraEditor/Private/Widgets/SNiagaraDistributionCurveEditor.cpp b/Engine/Plugins/FX/Niagara/Source/NiagaraEditor/Private/Widgets/SNiagaraDistributionCurveEditor.cpp
--- a/Engine/Plugins/FX/Niagara/Source/NiagaraEditor/Private/Widgets/SNiagaraDistributionCurveEditor.cpp
+++ b/Engine/Plugins/FX/Niagara/Source/NiagaraEditor/Private/Widgets/SNiagaraDistributionCurveEditor.cpp
@@ -152,11 +152,18 @@ namespace
 			TUniquePtr<FNiagaraDistributionCurveEditorModel> NewCurve = MakeUnique<FNiagaraDistributionCurveEditorModel>(DistributionAdapter.ToSharedRef(), ChannelIndex);
 			NewCurve->SetShortDisplayName(CurveName);
 			NewCurve->SetColor(CurveColor);
+			NewCurve->OnCurveModified().AddSP(this, &FNiagaraDistributionCurveEditorDetailsTreeItem::OnCurveChanged);
 			OutCurveModels.Add(MoveTemp(NewCurve));
 		}
 
 		virtual bool PassesFilter(const FCurveEditorTreeFilter* InFilter) const override { return true; }
 
+	private:
+		void OnCurveChanged()
+		{
+			DistributionAdapter->NotifyCurveValueChangedExternally(ChannelIndex);
+		}
+
 	private:
 		TSharedPtr<INiagaraDistributionAdapter>	DistributionAdapter;
 		int										ChannelIndex = 0;
