- Fix device preview crash in cooked editor
- ScriptRuntimeCookedDataMap will be incorrect when using a device preview

#Synced-CL 49161956

diff --git a/Engine/Plugins/FX/Niagara/Source/Niagara/Private/NiagaraSystem.cpp b/Engine/Plugins/FX/Niagara/Source/Niagara/Private/NiagaraSystem.cpp
--- a/Engine/Plugins/FX/Niagara/Source/Niagara/Private/NiagaraSystem.cpp
+++ b/Engine/Plugins/FX/Niagara/Source/Niagara/Private/NiagaraSystem.cpp
@@ -2887,12 +2887,9 @@ void UNiagaraSystem::UpdateHasGPUEmitters()
 
 void UNiagaraSystem::PostLoadScriptRuntimeData()
 {
-	if (FPlatformProperties::RequiresCookedData() == false && GetOutermost()->HasAnyPackageFlags(EPackageFlags::PKG_Cooked) == false)
-	{
-		ScriptRuntimeCookedDataMap.Empty();
-		return;
-	}
-
+#if WITH_EDITORONLY_DATA
+	ScriptRuntimeCookedDataMap.Empty();
+#else
 	ScriptRuntimeDataMap.Empty();
 	ForEachScriptWithOwningContext([this](const UNiagaraSystem* OwningSystem, const FNiagaraEmitterHandle* OwningEmitterHandle, UNiagaraScript* Script)
 	{
@@ -2924,6 +2921,7 @@ void UNiagaraSystem::PostLoadScriptRuntimeData()
 		TSharedRef<FNiagaraScriptRuntimeTransientData> RuntimeTransientData = MakeShared<FNiagaraScriptRuntimeTransientData>(*Script, ScriptCookedRuntimeData->CookedScriptRuntimeCompiledData, ScriptCookedRuntimeData->CookedExecutionParameterStore);
 		ScriptRuntimeDataMap.Add(ScriptDataKey, MakeShared<FNiagaraScriptRuntimeData>(RuntimeCompiledDataHandle, RuntimeTransientData)); 
 	});
+#endif
 }
 
 #if WITH_EDITORONLY_DATA
