- Floor offset option to Niagara system editor

#Synced-CL 49086386

diff --git a/Engine/Plugins/FX/Niagara/Source/NiagaraEditor/Private/Widgets/NiagaraSystemEditorViewportToolbarSections.cpp b/Engine/Plugins/FX/Niagara/Source/NiagaraEditor/Private/Widgets/NiagaraSystemEditorViewportToolbarSections.cpp
--- a/Engine/Plugins/FX/Niagara/Source/NiagaraEditor/Private/Widgets/NiagaraSystemEditorViewportToolbarSections.cpp
+++ b/Engine/Plugins/FX/Niagara/Source/NiagaraEditor/Private/Widgets/NiagaraSystemEditorViewportToolbarSections.cpp
@@ -14,156 +14,163 @@
 
 #define LOCTEXT_NAMESPACE "NiagaraSystemEditorViewportToolbarSections"
 
-FToolMenuEntry UE::NiagaraSystemEditor::CreateShowSubmenu()
-{
-	return UE::UnrealEd::CreateShowSubmenu(FNewToolMenuDelegate::CreateStatic(&UE::NiagaraSystemEditor::FillShowSubmenu, true));
-}
-
-void UE::NiagaraSystemEditor::FillShowSubmenu(UToolMenu* InMenu, bool bInShowViewportStatsToggle)
+namespace UE::NiagaraSystemEditor
 {
-	const FNiagaraEditorCommands& Commands = FNiagaraEditorCommands::Get();
-
-	if (bInShowViewportStatsToggle)
+	void FillShowSubmenu(UToolMenu* InMenu, bool bInShowViewportStatsToggle)
 	{
-		FToolMenuSection& UnnamedSection = InMenu->FindOrAddSection(NAME_None);
-		UnnamedSection.AddSubMenu(
-			"ViewportStats",
-			LOCTEXT("ViewportStatsSubMenu", "Viewport Stats"),
-			LOCTEXT("CameraSubmenuTooltip", "Camera options"),
-			FNewToolMenuDelegate::CreateLambda(
-				[](UToolMenu* Submenu) -> void
-				{
-					FToolMenuSection& CommonStatsSection =
-						Submenu->FindOrAddSection("CommonStats", LOCTEXT("CommonStatsLabel", "Common Stats"));
+		const FNiagaraEditorCommands& Commands = FNiagaraEditorCommands::Get();
 
-					CommonStatsSection.AddMenuEntry(
-						FEditorViewportCommands::Get().ToggleStats, LOCTEXT("ViewportStatsLabel", "Show Stats")
-					);
+		if (bInShowViewportStatsToggle)
+		{
+			FToolMenuSection& UnnamedSection = InMenu->FindOrAddSection(NAME_None);
+			UnnamedSection.AddSubMenu(
+				"ViewportStats",
+				LOCTEXT("ViewportStatsSubMenu", "Viewport Stats"),
+				LOCTEXT("CameraSubmenuTooltip", "Camera options"),
+				FNewToolMenuDelegate::CreateLambda(
+					[](UToolMenu* Submenu) -> void
+					{
+						FToolMenuSection& CommonStatsSection =
+							Submenu->FindOrAddSection("CommonStats", LOCTEXT("CommonStatsLabel", "Common Stats"));
+
+						CommonStatsSection.AddMenuEntry(
+							FEditorViewportCommands::Get().ToggleStats, LOCTEXT("ViewportStatsLabel", "Show Stats")
+						);
+
+						CommonStatsSection.AddSeparator(NAME_None);
+
+						CommonStatsSection.AddMenuEntry(
+							FEditorViewportCommands::Get().ToggleFPS, LOCTEXT("ViewportFPSLabel", "Show FPS")
+						);
+					}
+				)
+			);
+		}
 
-					CommonStatsSection.AddSeparator(NAME_None);
+		FToolMenuSection& CommonShowFlagsSection =
+			InMenu->FindOrAddSection("CommonShowFlags", LOCTEXT("CommonShowFlagsLabel", "Common Show Flags"));
 
-					CommonStatsSection.AddMenuEntry(
-						FEditorViewportCommands::Get().ToggleFPS, LOCTEXT("ViewportFPSLabel", "Show FPS")
-					);
-				}
-			)
-		);
+		CommonShowFlagsSection.AddMenuEntry(Commands.ToggleEmitterExecutionOrder);
+		CommonShowFlagsSection.AddMenuEntry(Commands.ToggleGpuTickInformation);
+		CommonShowFlagsSection.AddMenuEntry(Commands.ToggleInstructionCounts);
+		CommonShowFlagsSection.AddMenuEntry(Commands.ToggleMemoryInfo);
+		CommonShowFlagsSection.AddMenuEntry(Commands.ToggleParticleCounts);
+		CommonShowFlagsSection.AddMenuEntry(Commands.ToggleStatelessInfo);
 	}
 
-	FToolMenuSection& CommonShowFlagsSection =
-		InMenu->FindOrAddSection("CommonShowFlags", LOCTEXT("CommonShowFlagsLabel", "Common Show Flags"));
-
-	CommonShowFlagsSection.AddMenuEntry(Commands.ToggleEmitterExecutionOrder);
-	CommonShowFlagsSection.AddMenuEntry(Commands.ToggleGpuTickInformation);
-	CommonShowFlagsSection.AddMenuEntry(Commands.ToggleInstructionCounts);
-	CommonShowFlagsSection.AddMenuEntry(Commands.ToggleMemoryInfo);
-	CommonShowFlagsSection.AddMenuEntry(Commands.ToggleParticleCounts);
-	CommonShowFlagsSection.AddMenuEntry(Commands.ToggleStatelessInfo);
-}
-
-TSharedRef<SWidget> UE::NiagaraSystemEditor::CreateMotionMenuWidget(const TSharedRef<SNiagaraSystemViewport>& InNiagaraSystemEditorViewport
-)
-{
-	// We generate a menu via UToolMenus, so we can use FillShowSubmenu call from both old and new toolbar
-	FName OldShowMenuName = "NiagaraSystemEditor.OldViewportToolbar.Motion";
+	TSharedPtr<SNiagaraSystemViewport> ResolveSystemViewport(UToolMenu* InMenu)
+	{
+		if ( UUnrealEdViewportToolbarContext* ViewportToolbarContext = InMenu->FindContext<UUnrealEdViewportToolbarContext>() )
+		{
+			return StaticCastSharedPtr<SNiagaraSystemViewport>(ViewportToolbarContext->Viewport.Pin());
+		}
+		return nullptr;
+	}
 
-	if (!UToolMenus::Get()->IsMenuRegistered(OldShowMenuName))
+	void AddMotionSettingsToSection(FToolMenuSection& InSection)
 	{
-		UToolMenu* Menu = UToolMenus::Get()->RegisterMenu(OldShowMenuName, NAME_None, EMultiBoxType::Menu, false);
-		Menu->AddDynamicSection(
-			"BaseSection",
+		InSection.AddSubMenu(
+			"MotionOptions",
+			LOCTEXT("MotionOptionsSubMenu", "Motion Options"),
+			LOCTEXT("MotionOptionsSubMenu_ToolTip", "Set Motion Options for the Niagara Component"),
 			FNewToolMenuDelegate::CreateLambda(
 				[](UToolMenu* InMenu)
 				{
 					FToolMenuSection& UnnamedSection = InMenu->FindOrAddSection(NAME_None);
-					UE::NiagaraSystemEditor::AddMotionSettingsToSection(UnnamedSection);
-				}
-			)
-		);
-	}
-
-	FToolMenuContext MenuContext;
-	{
-		MenuContext.AppendCommandList(InNiagaraSystemEditorViewport->GetCommandList());
-
-		// Add the UnrealEd viewport toolbar context.
-		{
-			UUnrealEdViewportToolbarContext* const ContextObject =
-				UE::UnrealEd::CreateViewportToolbarDefaultContext(InNiagaraSystemEditorViewport);
-
-			MenuContext.AddObject(ContextObject);
-		}
-	}
-
-	return UToolMenus::Get()->GenerateWidget(OldShowMenuName, MenuContext);
-}
-
-void UE::NiagaraSystemEditor::FillSettingsSubmenu(UToolMenu* InMenu)
-{
-	FToolMenuSection& ViewportControlsSection =
-		InMenu->FindOrAddSection("ViewportControls", LOCTEXT("ViewportControlsLabel", "Viewport Controls"));
 
-	AddMotionSettingsToSection(ViewportControlsSection);
-}
+					UnnamedSection.AddMenuEntry(FNiagaraEditorCommands::Get().ToggleMotion);
 
-void UE::NiagaraSystemEditor::AddMotionSettingsToSection(FToolMenuSection& InSection)
-{
-	InSection.AddSubMenu(
-		"MotionOptions",
-		LOCTEXT("MotionOptionsSubMenu", "Motion Options"),
-		LOCTEXT("MotionOptionsSubMenu_ToolTip", "Set Motion Options for the Niagara Component"),
-		FNewToolMenuDelegate::CreateLambda(
-			[](UToolMenu* InMenu)
-			{
-				FToolMenuSection& UnnamedSection = InMenu->FindOrAddSection(NAME_None);
-
-				UnnamedSection.AddMenuEntry(FNiagaraEditorCommands::Get().ToggleMotion);
-
-				UUnrealEdViewportToolbarContext* ViewportToolbarContext =
-					InMenu->FindContext<UUnrealEdViewportToolbarContext>();
-				if (!ViewportToolbarContext)
-				{
-					return;
-				}
+					TSharedPtr<SNiagaraSystemViewport> NiagaraSystemViewport = ResolveSystemViewport(InMenu);
+					if (!NiagaraSystemViewport)
+					{
+						return;
+					}
 
-				TSharedPtr<SNiagaraSystemViewport> NiagaraSystemViewport =
-					StaticCastSharedPtr<SNiagaraSystemViewport>(ViewportToolbarContext->Viewport.Pin());
-				if (!NiagaraSystemViewport)
-				{
-					return;
-				}
-
-				FToolMenuEntry MotionRateEntry = FToolMenuEntry::InitWidget(
-					"MotionRate",
-					SNew(SSpinBox<float>)
+					FToolMenuEntry MotionRateEntry = FToolMenuEntry::InitWidget(
+						"MotionRate",
+						SNew(SSpinBox<float>)
 						.IsEnabled(NiagaraSystemViewport.Get(), &SNiagaraSystemViewport::IsMotionEnabled)
 						.Font(FAppStyle::GetFontStyle(TEXT("MenuItem.Font")))
 						.MinSliderValue(0.0f)
 						.MaxSliderValue(360.0f)
 						.Value(NiagaraSystemViewport.Get(), &SNiagaraSystemViewport::GetMotionRate)
 						.OnValueChanged(NiagaraSystemViewport.Get(), &SNiagaraSystemViewport::SetMotionRate),
-					LOCTEXT("MotionSpeed", "Motion Speed")
-				);
+						LOCTEXT("MotionSpeed", "Motion Speed")
+					);
 
-				FToolMenuEntry MotionRadiusEntry = FToolMenuEntry::InitWidget(
-					"MotionRadius",
-					SNew(SSpinBox<float>)
+					FToolMenuEntry MotionRadiusEntry = FToolMenuEntry::InitWidget(
+						"MotionRadius",
+						SNew(SSpinBox<float>)
 						.IsEnabled(NiagaraSystemViewport.Get(), &SNiagaraSystemViewport::IsMotionEnabled)
 						.Font(FAppStyle::GetFontStyle(TEXT("MenuItem.Font")))
 						.MinSliderValue(0.0f)
 						.MaxSliderValue(1000.0f)
 						.Value(NiagaraSystemViewport.Get(), &SNiagaraSystemViewport::GetMotionRadius)
 						.OnValueChanged(NiagaraSystemViewport.Get(), &SNiagaraSystemViewport::SetMotionRadius),
-					LOCTEXT("MotionRadius", "Motion Radius")
-				);
-
-				UnnamedSection.AddEntry(MotionRateEntry);
-				UnnamedSection.AddEntry(MotionRadiusEntry);
-			}
-		),
-		false,
-		FSlateIcon()
-	);
+						LOCTEXT("MotionRadius", "Motion Radius")
+					);
+
+					UnnamedSection.AddEntry(MotionRateEntry);
+					UnnamedSection.AddEntry(MotionRadiusEntry);
+				}
+			),
+			false,
+			FSlateIcon()
+		);
+	}
+
+	void AddFloorSettingsToSection(FToolMenuSection& InSection)
+	{
+		InSection.AddSubMenu(
+			"FloorOptions",
+			LOCTEXT("FloorOptionsSubMenu", "Floor Options"),
+			LOCTEXT("FloorOptionsSubMenu_ToolTip", "Set Floor Options for the preview window"),
+			FNewToolMenuDelegate::CreateLambda(
+				[](UToolMenu* InMenu)
+				{
+					TSharedPtr<SNiagaraSystemViewport> NiagaraSystemViewport = ResolveSystemViewport(InMenu);
+					if (!NiagaraSystemViewport)
+					{
+						return;
+					}
+
+					FToolMenuSection& UnnamedSection = InMenu->FindOrAddSection(NAME_None);
+
+					UnnamedSection.AddMenuEntry(
+						NAME_None,
+						LOCTEXT("FloorEnable", "Floor Enabled"),
+						LOCTEXT("FloorEnableToolTip", "Controls visibility of the floor"),
+						FSlateIcon(),
+						FUIAction(
+							FExecuteAction::CreateSP(NiagaraSystemViewport.Get(), &SNiagaraSystemViewport::ToggleFloorEnabled),
+							FCanExecuteAction(),
+							FIsActionChecked::CreateSP(NiagaraSystemViewport.Get(), &SNiagaraSystemViewport::IsFloorEnabled)
+						),
+						EUserInterfaceActionType::ToggleButton
+					);
+
+					UnnamedSection.AddEntry(
+						FToolMenuEntry::InitWidget(
+							NAME_None,
+							SNew(SSpinBox<float>)
+							.Font(FAppStyle::GetFontStyle(TEXT("MenuItem.Font")))
+							.IsEnabled(NiagaraSystemViewport.Get(), &SNiagaraSystemViewport::IsFloorEnabled)
+							.Value(NiagaraSystemViewport.Get(), &SNiagaraSystemViewport::GetFloorOffset)
+							.OnValueChanged(NiagaraSystemViewport.Get(), &SNiagaraSystemViewport::SetFloorOffset),
+							LOCTEXT("FloorHeight", "Floor Height")
+						)
+					);
+				}
+			),
+			false,
+			FSlateIcon()
+		);
+	}
+}
+
+FToolMenuEntry UE::NiagaraSystemEditor::CreateShowSubmenu()
+{
+	return UE::UnrealEd::CreateShowSubmenu(FNewToolMenuDelegate::CreateStatic(&UE::NiagaraSystemEditor::FillShowSubmenu, true));
 }
 
 void UE::NiagaraSystemEditor::ExtendPreviewSceneSettingsSubmenu(FName InSubmenuName)
@@ -180,11 +187,13 @@ void UE::NiagaraSystemEditor::ExtendPreviewSceneSettingsSubmenu(FName InSubmenuN
 	);
 
 	UE::NiagaraSystemEditor::AddMotionSettingsToSection(PreviewControlsSection);
+	UE::NiagaraSystemEditor::AddFloorSettingsToSection(PreviewControlsSection);
 
 	//Adds toggle origin axis to menu
 	FToolMenuSection& ProfileOptionsSection = Submenu->FindOrAddSection(
 		"PreviewSceneSettings", LOCTEXT("AssetViewerProfileOptionsSectionLabel", "Preview Scene Options")
 	);
+
 	ProfileOptionsSection.AddMenuEntry(
 		FNiagaraEditorCommands::Get().ToggleOriginAxis, 
 		TAttribute<FText>(), 
diff --git a/Engine/Plugins/FX/Niagara/Source/NiagaraEditor/Private/Widgets/NiagaraSystemEditorViewportToolbarSections.h b/Engine/Plugins/FX/Niagara/Source/NiagaraEditor/Private/Widgets/NiagaraSystemEditorViewportToolbarSections.h
--- a/Engine/Plugins/FX/Niagara/Source/NiagaraEditor/Private/Widgets/NiagaraSystemEditorViewportToolbarSections.h
+++ b/Engine/Plugins/FX/Niagara/Source/NiagaraEditor/Private/Widgets/NiagaraSystemEditorViewportToolbarSections.h
@@ -1,4 +1,4 @@
-﻿// Copyright Epic Games, Inc. All Rights Reserved.
+// Copyright Epic Games, Inc. All Rights Reserved.
 
 #pragma once
 
@@ -9,11 +9,6 @@ class SNiagaraSystemViewport;
 
 namespace UE::NiagaraSystemEditor
 {
-FToolMenuEntry CreateShowSubmenu();
-void FillShowSubmenu(UToolMenu* InMenu, bool bInShowViewportStatsToggle = true);
-
-TSharedRef<SWidget> CreateMotionMenuWidget(const TSharedRef<SNiagaraSystemViewport>& InNiagaraSystemEditorViewport);
-void FillSettingsSubmenu(UToolMenu* InMenu);
-void AddMotionSettingsToSection(FToolMenuSection& InSection);
-void ExtendPreviewSceneSettingsSubmenu(FName InSubmenuName);
+	extern FToolMenuEntry CreateShowSubmenu();
+	extern void ExtendPreviewSceneSettingsSubmenu(FName InSubmenuName);
 }
diff --git a/Engine/Plugins/FX/Niagara/Source/NiagaraEditor/Private/Widgets/SNiagaraSystemViewport.cpp b/Engine/Plugins/FX/Niagara/Source/NiagaraEditor/Private/Widgets/SNiagaraSystemViewport.cpp
--- a/Engine/Plugins/FX/Niagara/Source/NiagaraEditor/Private/Widgets/SNiagaraSystemViewport.cpp
+++ b/Engine/Plugins/FX/Niagara/Source/NiagaraEditor/Private/Widgets/SNiagaraSystemViewport.cpp
@@ -970,6 +970,33 @@ bool SNiagaraSystemViewport::IsMotionEnabled() const
 	return bMotionEnabled;
 }
 
+bool SNiagaraSystemViewport::IsFloorEnabled() const
+{
+	return AdvancedPreviewScene ? AdvancedPreviewScene->IsFloorEnabled() : false;
+}
+
+void SNiagaraSystemViewport::ToggleFloorEnabled()
+{
+	if ( AdvancedPreviewScene )
+	{
+		AdvancedPreviewScene->SetFloorVisibility(!AdvancedPreviewScene->IsFloorEnabled());
+	}
+}
+
+float SNiagaraSystemViewport::GetFloorOffset() const
+{
+	return AdvancedPreviewScene ? AdvancedPreviewScene->GetFloorOffset() : 0.0f;
+}
+
+void SNiagaraSystemViewport::SetFloorOffset(float Offset)
+{
+	if (AdvancedPreviewScene)
+	{
+		AdvancedPreviewScene->SetFloorOffset(Offset);
+		RefreshViewport();
+	}
+}
+
 void SNiagaraSystemViewport::RefreshViewport()
 {
 	//reregister the preview components, so if the preview material changed it will be propagated to the render thread
diff --git a/Engine/Plugins/FX/Niagara/Source/NiagaraEditor/Private/Widgets/SNiagaraSystemViewport.h b/Engine/Plugins/FX/Niagara/Source/NiagaraEditor/Private/Widgets/SNiagaraSystemViewport.h
--- a/Engine/Plugins/FX/Niagara/Source/NiagaraEditor/Private/Widgets/SNiagaraSystemViewport.h
+++ b/Engine/Plugins/FX/Niagara/Source/NiagaraEditor/Private/Widgets/SNiagaraSystemViewport.h
@@ -107,6 +107,10 @@ public:
 	float GetMotionRadius() const { return MotionRadius; }
 	void SetMotionRadius(float Radius) { MotionRadius = Radius; }
 
+	bool IsFloorEnabled() const;
+	void ToggleFloorEnabled();
+	float GetFloorOffset() const;
+	void SetFloorOffset(float Offset);
 
 protected:	
 	/** SEditorViewport interface */
