- Handle potential crash with new or old variable being invalid
- It's not clear why this occurs as we don't have a repro

#Synced-CL 49086386

diff --git a/Engine/Plugins/FX/Niagara/Source/NiagaraEditor/Private/ViewModels/HierarchyEditor/NiagaraScriptParametersHierarchyViewModel.cpp b/Engine/Plugins/FX/Niagara/Source/NiagaraEditor/Private/ViewModels/HierarchyEditor/NiagaraScriptParametersHierarchyViewModel.cpp
--- a/Engine/Plugins/FX/Niagara/Source/NiagaraEditor/Private/ViewModels/HierarchyEditor/NiagaraScriptParametersHierarchyViewModel.cpp
+++ b/Engine/Plugins/FX/Niagara/Source/NiagaraEditor/Private/ViewModels/HierarchyEditor/NiagaraScriptParametersHierarchyViewModel.cpp
@@ -214,20 +214,25 @@ void UNiagaraScriptParametersHierarchyViewModel::OnParametersChanged(TOptional<T
 		if(ChangedData.GetScriptStruct() == FNiagaraParameterRenamedData::StaticStruct())
 		{
 			const FNiagaraParameterRenamedData& RenamedData = ChangedData.Get<FNiagaraParameterRenamedData>();
-			if(RenamedData.OldScriptVariable->GetIsStaticSwitch() && RenamedData.NewScriptVariable->GetIsStaticSwitch())
+			const UNiagaraScriptVariable* OldScriptVariable = RenamedData.OldScriptVariable.Get();
+			const UNiagaraScriptVariable* NewScriptVariable = RenamedData.NewScriptVariable.Get();
+			if (ensureMsgf(OldScriptVariable && NewScriptVariable, TEXT("UNiagaraScriptParametersHierarchyViewModel::OnParametersChanged - OldVariable(%p) NewVariable(%p)"), OldScriptVariable, NewScriptVariable))
 			{
-				TArray<UNiagaraHierarchyScriptParameter*> AllHierarchyScriptParameters;
-				HierarchyRoot->GetChildrenOfType<UNiagaraHierarchyScriptParameter>(AllHierarchyScriptParameters, true);
-
-				UNiagaraHierarchyScriptParameter** FoundHierarchyScriptParameter = AllHierarchyScriptParameters.FindByPredicate([Guid = RenamedData.OldScriptVariable->Metadata.GetVariableGuid()](UNiagaraHierarchyScriptParameter* Candidate)
-				{
-					return Candidate->GetPersistentIdentity().Guids[0] == Guid;
-				});
-
-				if(FoundHierarchyScriptParameter)
+				if(OldScriptVariable->GetIsStaticSwitch() && NewScriptVariable->GetIsStaticSwitch())
 				{
-					HierarchyRoot->Modify();
-					(*FoundHierarchyScriptParameter)->Initialize(*RenamedData.NewScriptVariable);
+					TArray<UNiagaraHierarchyScriptParameter*> AllHierarchyScriptParameters;
+					HierarchyRoot->GetChildrenOfType<UNiagaraHierarchyScriptParameter>(AllHierarchyScriptParameters, true);
+
+					UNiagaraHierarchyScriptParameter** FoundHierarchyScriptParameter = AllHierarchyScriptParameters.FindByPredicate([Guid = OldScriptVariable->Metadata.GetVariableGuid()](UNiagaraHierarchyScriptParameter* Candidate)
+					{
+						return Candidate->GetPersistentIdentity().Guids[0] == Guid;
+					});
+
+					if(FoundHierarchyScriptParameter)
+					{
+						HierarchyRoot->Modify();
+						(*FoundHierarchyScriptParameter)->Initialize(*NewScriptVariable);
+					}
 				}
 			}
 		}
