- Fix for crash in VT data interface as AcquireAllocatedVT is no longer allowed to be called during set shader parameters
- We call AcquireAllocatedVT when setting the instance data to ensure the VT is allocated otherwise it won't be

#Synced-CL 49086386

diff --git a/Engine/Plugins/FX/Niagara/Source/Niagara/Private/NiagaraDataInterfaceVirtualTextureSample.cpp b/Engine/Plugins/FX/Niagara/Source/Niagara/Private/NiagaraDataInterfaceVirtualTextureSample.cpp
--- a/Engine/Plugins/FX/Niagara/Source/Niagara/Private/NiagaraDataInterfaceVirtualTextureSample.cpp
+++ b/Engine/Plugins/FX/Niagara/Source/Niagara/Private/NiagaraDataInterfaceVirtualTextureSample.cpp
@@ -241,6 +241,12 @@ bool UNiagaraDataInterfaceVirtualTextureSample::PerInstanceTick(void* PerInstanc
 				InstanceData.TextureSize = RT_TextureSize;
 				InstanceData.MipLevels = RT_MipLevels;
 				InstanceData.VirtualTextureResource = static_cast<FVirtualTexture2DResource*>(RT_Texture->GetResource());
+				if (InstanceData.VirtualTextureResource && InstanceData.VirtualTextureResource->IsInitialized())
+				{
+					// We need to acquire here as it's not valid to acquire during set shader parameters and will cause a crash
+					InstanceData.VirtualTextureResource->AcquireAllocatedVT();
+				}
+
 				InstanceData.VirtualTextureResource_SRGB = RT_Texture->SRGB;				
 			}
 		);
@@ -345,12 +351,9 @@ void UNiagaraDataInterfaceVirtualTextureSample::SetShaderParameters(const FNiaga
 	if (InstanceData)
 	{
 		FVirtualTexture2DResource* VirtualTextureResource = InstanceData->VirtualTextureResource;
-
 		if (VirtualTextureResource && VirtualTextureResource->IsInitialized())
 		{
-			IAllocatedVirtualTexture* AllocatedVT = VirtualTextureResource->AcquireAllocatedVT();
-
-			if (AllocatedVT)
+			if (IAllocatedVirtualTexture* AllocatedVT = VirtualTextureResource->GetAllocatedVT())
 			{
 				FRHIShaderResourceView* InPhysicalTexture = AllocatedVT->GetPhysicalTextureSRV((uint32)0, InstanceData->VirtualTextureResource_SRGB == 1);
 
