- Add system information to log when data interface is null

#tests Hacked in code to verify the log works

#Synced-CL 49086386

diff --git a/Engine/Plugins/FX/Niagara/Source/Niagara/Private/NiagaraSystemInstance.cpp b/Engine/Plugins/FX/Niagara/Source/Niagara/Private/NiagaraSystemInstance.cpp
--- a/Engine/Plugins/FX/Niagara/Source/Niagara/Private/NiagaraSystemInstance.cpp
+++ b/Engine/Plugins/FX/Niagara/Source/Niagara/Private/NiagaraSystemInstance.cpp
@@ -1730,7 +1730,7 @@ void FNiagaraSystemInstance::InitDataInterfaces()
 		}
 		else
 		{
-			UE_LOG(LogNiagara, Error, TEXT("A data interface currently in use by an System has been destroyed."));
+			UE_LOG(LogNiagara, Error, TEXT("A data interface currently in use by '%s' has been destroyed."), *GetNameSafe(System));
 			bDataInterfacesInitialized = false;
 		}
 	}
