- Add bounds toggle buttons to Niagara stack overview
- Note: not added to summary view currently
- Also made the menu drop down refresh to the viewport to make it more visible

#Synced-CL 49086386

diff --git a/Engine/Plugins/FX/Niagara/Source/NiagaraEditor/Private/Toolkits/NiagaraSystemToolkit.cpp b/Engine/Plugins/FX/Niagara/Source/NiagaraEditor/Private/Toolkits/NiagaraSystemToolkit.cpp
--- a/Engine/Plugins/FX/Niagara/Source/NiagaraEditor/Private/Toolkits/NiagaraSystemToolkit.cpp
+++ b/Engine/Plugins/FX/Niagara/Source/NiagaraEditor/Private/Toolkits/NiagaraSystemToolkit.cpp
@@ -963,6 +963,10 @@ void FNiagaraSystemToolkit::ToggleShowEmitterBounds(TSharedRef<FNiagaraEmitterHa
 	if (FNiagaraEmitterHandle* EmitterHandle = EmitterViewModel->GetEmitterHandle())
 	{
 		EmitterHandle->SetDebugShowBounds(!EmitterHandle->GetDebugShowBounds());
+		if (Viewport.IsValid())
+		{
+			Viewport->RefreshViewport();
+		}
 	}
 }
 
diff --git a/Engine/Plugins/FX/Niagara/Source/NiagaraEditorWidgets/Private/SNiagaraOverviewStackNode.cpp b/Engine/Plugins/FX/Niagara/Source/NiagaraEditorWidgets/Private/SNiagaraOverviewStackNode.cpp
--- a/Engine/Plugins/FX/Niagara/Source/NiagaraEditorWidgets/Private/SNiagaraOverviewStackNode.cpp
+++ b/Engine/Plugins/FX/Niagara/Source/NiagaraEditorWidgets/Private/SNiagaraOverviewStackNode.cpp
@@ -611,14 +611,25 @@ TSharedRef<SWidget> SNiagaraOverviewStackNode::CreateTopContentBar_Default()
 	{		
 		// Isolate toggle button
 		TopContentBar->AddSlot()
-			.AutoWidth()
-			.HAlign(HAlign_Left)
-			.VAlign(VAlign_Center)
-			.Padding(2, 0, 0, 0)
-			[
-				CreateIsolateButton()
-			];
+		.AutoWidth()
+		.HAlign(HAlign_Left)
+		.VAlign(VAlign_Center)
+		.Padding(2, 0, 0, 0)
+		[
+			CreateIsolateButton()
+		];
 
+		// Bounds toggle button
+		TopContentBar->AddSlot()
+		.AutoWidth()
+		.HAlign(HAlign_Left)
+		.VAlign(VAlign_Center)
+		.Padding(2, 0, 0, 0)
+		[
+			CreateBoundsButton()
+		];
+
+		// Renderer thumbnails
 		TopContentBar->AddSlot()
 		.AutoWidth()
 		.HAlign(HAlign_Left)
@@ -1004,6 +1015,24 @@ TSharedRef<SButton> SNiagaraOverviewStackNode::CreateIsolateButton()
 	];
 }
 
+TSharedRef<SButton> SNiagaraOverviewStackNode::CreateBoundsButton()
+{
+	return SNew(SButton)
+	.ButtonStyle(FAppStyle::Get(), "HoverHintOnly")
+	.HAlign(HAlign_Center)
+	.ContentPadding(1)
+	.ToolTipText(this, &SNiagaraOverviewStackNode::GetToggleBoundsToolTip)
+	.OnClicked(this, &SNiagaraOverviewStackNode::OnToggleBoundsButtonClicked)
+	.Visibility(this, &SNiagaraOverviewStackNode::GetToggleBoundsVisibility)
+	.IsFocusable(false)
+	.Content()
+	[
+		SNew(SImage)
+		.Image(this, &SNiagaraOverviewStackNode::GetToggleBoundsImage)
+		.ColorAndOpacity(this, &SNiagaraOverviewStackNode::GetToggleBoundsImageColor)
+	];
+}
+
 TSharedRef<SButton> SNiagaraOverviewStackNode::CreateCaptureThumbnailButton()
 {
 	return SNew(SButton)
@@ -1205,6 +1234,43 @@ FSlateColor SNiagaraOverviewStackNode::GetToggleIsolateImageColor() const
 		: FLinearColor::Gray;
 }
 
+FText SNiagaraOverviewStackNode::GetToggleBoundsToolTip() const
+{
+	return LOCTEXT("ToggleBoundsTooltip", "Toggle debug emitter bounds visibility on / off");
+}
+
+FReply SNiagaraOverviewStackNode::OnToggleBoundsButtonClicked()
+{
+	TSharedPtr<FNiagaraEmitterHandleViewModel> EmitterHandleViewModel = EmitterHandleViewModelWeak.Pin();
+	FNiagaraEmitterHandle* EmitterHandle = EmitterHandleViewModel.IsValid() ? EmitterHandleViewModel->GetEmitterHandle() : nullptr;
+	if (EmitterHandle)
+	{
+		EmitterHandle->SetDebugShowBounds(!EmitterHandle->GetDebugShowBounds());
+	}
+	return FReply::Handled();
+}
+
+EVisibility SNiagaraOverviewStackNode::GetToggleBoundsVisibility() const
+{
+	TSharedPtr<FNiagaraEmitterHandleViewModel> EmitterHandleViewModel = EmitterHandleViewModelWeak.Pin();
+	return EmitterHandleViewModel.IsValid() &&
+		EmitterHandleViewModel->GetOwningSystemEditMode() == ENiagaraSystemViewModelEditMode::SystemAsset
+		? EVisibility::Visible
+		: EVisibility::Collapsed;
+}
+
+const FSlateBrush* SNiagaraOverviewStackNode::GetToggleBoundsImage() const
+{
+	return FAppStyle::Get().GetBrush("AssetEditor.ToggleShowBounds");
+}
+
+FSlateColor SNiagaraOverviewStackNode::GetToggleBoundsImageColor() const
+{
+	TSharedPtr<FNiagaraEmitterHandleViewModel> EmitterHandleViewModel = EmitterHandleViewModelWeak.Pin();
+	FNiagaraEmitterHandle* EmitterHandle = EmitterHandleViewModel.IsValid() ? EmitterHandleViewModel->GetEmitterHandle() : nullptr;
+	return EmitterHandle && EmitterHandle->GetDebugShowBounds() ? FAppStyle::GetSlateColor("SelectionColor") : FLinearColor::Gray;
+}
+
 FSlateColor SNiagaraOverviewStackNode::GetScalabilityTintAlpha() const
 {
 	if(EmitterHandleViewModelWeak.IsValid())
diff --git a/Engine/Plugins/FX/Niagara/Source/NiagaraEditorWidgets/Private/SNiagaraOverviewStackNode.h b/Engine/Plugins/FX/Niagara/Source/NiagaraEditorWidgets/Private/SNiagaraOverviewStackNode.h
--- a/Engine/Plugins/FX/Niagara/Source/NiagaraEditorWidgets/Private/SNiagaraOverviewStackNode.h
+++ b/Engine/Plugins/FX/Niagara/Source/NiagaraEditorWidgets/Private/SNiagaraOverviewStackNode.h
@@ -62,6 +62,7 @@ protected:
 
 	TSharedRef<SWidget> CreateEnabledCheckbox();
 	TSharedRef<SButton> CreateIsolateButton();
+	TSharedRef<SButton> CreateBoundsButton();
 	TSharedRef<SWidget> CreateEmitterThumbnail();
 	TSharedRef<SWidget> CreateEmitterRendererThumbnails();
 	TSharedRef<SWidget> CreateSummaryViewToggle();
@@ -96,6 +97,12 @@ private:
 	const FSlateBrush* GetToggleIsolateImage() const;
 	FSlateColor GetToggleIsolateImageColor() const;
 
+	FText GetToggleBoundsToolTip() const;
+	FReply OnToggleBoundsButtonClicked();
+	EVisibility GetToggleBoundsVisibility() const;
+	const FSlateBrush* GetToggleBoundsImage() const;
+	FSlateColor GetToggleBoundsImageColor() const;
+
 	EVisibility GetScalabilityIndicatorVisibility() const;
 	FSlateColor GetScalabilityTintAlpha() const;
 	void OnScalabilityModeChanged(bool bActive);
