- Add manual focus distance to system editor
- This allows users to customize the focus distance for an asset, avoiding issues with fixed bounds being too large, etc

#Synced-CL 49227054

diff --git a/Engine/Plugins/FX/Niagara/Source/NiagaraEditor/Private/Widgets/NiagaraSystemEditorViewportToolbarSections.cpp b/Engine/Plugins/FX/Niagara/Source/NiagaraEditor/Private/Widgets/NiagaraSystemEditorViewportToolbarSections.cpp
--- a/Engine/Plugins/FX/Niagara/Source/NiagaraEditor/Private/Widgets/NiagaraSystemEditorViewportToolbarSections.cpp
+++ b/Engine/Plugins/FX/Niagara/Source/NiagaraEditor/Private/Widgets/NiagaraSystemEditorViewportToolbarSections.cpp
@@ -166,6 +166,54 @@ namespace UE::NiagaraSystemEditor
 			FSlateIcon()
 		);
 	}
+
+	void AddFocusSettingsToSection(FToolMenuSection& InSection)
+	{
+		InSection.AddSubMenu(
+			"FocusOptions",
+			LOCTEXT("FocusOptionsSubMenu", "Focus Options"),
+			LOCTEXT("FocusOptionsSubMenu_ToolTip", "Adjusts how we focus on the Niagara component"),
+			FNewToolMenuDelegate::CreateLambda(
+				[](UToolMenu* InMenu)
+				{
+					TSharedPtr<SNiagaraSystemViewport> NiagaraSystemViewport = ResolveSystemViewport(InMenu);
+					if (!NiagaraSystemViewport)
+					{
+						return;
+					}
+
+					FToolMenuSection& UnnamedSection = InMenu->FindOrAddSection(NAME_None);
+
+					UnnamedSection.AddMenuEntry(
+						NAME_None,
+						LOCTEXT("FocusOverrideEnable", "Focus Override Enabled"),
+						LOCTEXT("FocusOverrideEnableToolTip", "When enable we will override the default focus method, which is to use the bounding box."),
+						FSlateIcon(),
+						FUIAction(
+							FExecuteAction::CreateSP(NiagaraSystemViewport.Get(), &SNiagaraSystemViewport::ToggleManualFocusDistance),
+							FCanExecuteAction(),
+							FIsActionChecked::CreateSP(NiagaraSystemViewport.Get(), &SNiagaraSystemViewport::IsManualFocusDistanceEnabled)
+						),
+						EUserInterfaceActionType::ToggleButton
+					);
+
+					UnnamedSection.AddEntry(
+						FToolMenuEntry::InitWidget(
+							NAME_None,
+							SNew(SSpinBox<float>)
+							.Font(FAppStyle::GetFontStyle(TEXT("MenuItem.Font")))
+							.IsEnabled(NiagaraSystemViewport.Get(), &SNiagaraSystemViewport::IsManualFocusDistanceEnabled)
+							.Value(NiagaraSystemViewport.Get(), &SNiagaraSystemViewport::GetManualFocusDistance)
+							.OnValueChanged(NiagaraSystemViewport.Get(), &SNiagaraSystemViewport::SetManualFocusDistance),
+							LOCTEXT("FocusDistance", "Focus Distance")
+						)
+					);
+				}
+			),
+			false,
+			FSlateIcon()
+		);
+	}
 }
 
 FToolMenuEntry UE::NiagaraSystemEditor::CreateShowSubmenu()
@@ -188,6 +236,7 @@ void UE::NiagaraSystemEditor::ExtendPreviewSceneSettingsSubmenu(FName InSubmenuN
 
 	UE::NiagaraSystemEditor::AddMotionSettingsToSection(PreviewControlsSection);
 	UE::NiagaraSystemEditor::AddFloorSettingsToSection(PreviewControlsSection);
+	UE::NiagaraSystemEditor::AddFocusSettingsToSection(PreviewControlsSection);
 
 	//Adds toggle origin axis to menu
 	FToolMenuSection& ProfileOptionsSection = Submenu->FindOrAddSection(
diff --git a/Engine/Plugins/FX/Niagara/Source/NiagaraEditor/Private/Widgets/SNiagaraSystemViewport.cpp b/Engine/Plugins/FX/Niagara/Source/NiagaraEditor/Private/Widgets/SNiagaraSystemViewport.cpp
--- a/Engine/Plugins/FX/Niagara/Source/NiagaraEditor/Private/Widgets/SNiagaraSystemViewport.cpp
+++ b/Engine/Plugins/FX/Niagara/Source/NiagaraEditor/Private/Widgets/SNiagaraSystemViewport.cpp
@@ -95,6 +95,8 @@ public:
 	virtual bool CanSetWidgetMode(UE::Widget::EWidgetMode NewMode) const override { return false; }
 	virtual bool CanCycleWidgetMode() const override { return false; }
 
+	void FocusViewport(UNiagaraComponent* Component, bool bInstant);
+
 	void SetOrbitModeFromSettings();
 	virtual void SetIsSimulateInEditorViewport(bool bInIsSimulateInEditorViewport) override;
 	UNiagaraSystemEditorData* GetSystemEditorData() const;
@@ -198,11 +200,11 @@ void FNiagaraSystemViewportClient::Draw(FViewport* InViewport, FCanvas* Canvas)
 	FNiagaraSystemInstanceControllerConstPtr SystemInstanceController = NiagaraViewport.IsValid() ? NiagaraViewport->GetPreviewComponent()->GetSystemInstanceController() : nullptr;
 	if (Component && SystemInstanceController.IsValid() && SystemInstanceController->GetAge() > 0.0)
 	{
-		FBox Bounds = Component->GetLocalBounds().GetBox();
+		const FBox Bounds = Component->GetLocalBounds().GetBox();
 		if (bUpdateViewportFocus && Bounds.IsValid)
 		{
-			bool bCachedShouldOrbitCamera = ShouldOrbitCamera();
-			FocusViewportOnBox(Bounds);
+			const bool bCachedShouldOrbitCamera = ShouldOrbitCamera();
+			FocusViewport(Component, true);
 			ToggleOrbitCamera(bCachedShouldOrbitCamera);
 			SetUpdateViewportFocus(false);
 		}
@@ -700,6 +702,24 @@ FLinearColor FNiagaraSystemViewportClient::GetBackgroundColor() const
 	return BackgroundColor;
 }
 
+void FNiagaraSystemViewportClient::FocusViewport(UNiagaraComponent* Component, bool bInstant)
+{
+	FBox FocusBounds = Component->Bounds.GetBox();
+	if (FocusBounds.IsValid == false)
+	{
+		return;
+	}
+
+	UNiagaraSystemEditorData* SystemEditorData = GetSystemEditorData();
+	if (SystemEditorData && SystemEditorData->GetViewportManualFocusDistanceEnabled())
+	{
+		const float Distance = SystemEditorData->GetViewportManualFocusDistance();
+		const FVector HDistance = FVector(Distance * 0.5f);
+		FocusBounds = FBox(FocusBounds.GetCenter() - HDistance, FocusBounds.GetCenter() + HDistance);
+	}
+	FocusViewportOnBox(FocusBounds, bInstant);
+}
+
 FSceneView* FNiagaraSystemViewportClient::CalcSceneView(FSceneViewFamily* ViewFamily, const int32 StereoViewIndex)
 {
 	FSceneView* SceneView = FEditorViewportClient::CalcSceneView(ViewFamily);
@@ -997,6 +1017,34 @@ void SNiagaraSystemViewport::SetFloorOffset(float Offset)
 	}
 }
 
+bool SNiagaraSystemViewport::IsManualFocusDistanceEnabled() const
+{
+	UNiagaraSystemEditorData* SystemEditorData = SystemViewportClient->GetSystemEditorData();
+	return SystemEditorData ? SystemEditorData->GetViewportManualFocusDistanceEnabled() : false;
+}
+
+void SNiagaraSystemViewport::ToggleManualFocusDistance()
+{
+	if ( UNiagaraSystemEditorData* SystemEditorData = SystemViewportClient->GetSystemEditorData() )
+	{
+		SystemEditorData->SetViewportManualFocusDistanceEnabled(!SystemEditorData->GetViewportManualFocusDistanceEnabled());
+	}
+}
+
+void SNiagaraSystemViewport::SetManualFocusDistance(float Distance)
+{
+	if (UNiagaraSystemEditorData* SystemEditorData = SystemViewportClient->GetSystemEditorData())
+	{
+		SystemEditorData->SetViewportManualFocusDistance(Distance);
+	}
+}
+
+float SNiagaraSystemViewport::GetManualFocusDistance() const
+{
+	UNiagaraSystemEditorData* SystemEditorData = SystemViewportClient->GetSystemEditorData();
+	return SystemEditorData ? SystemEditorData->GetViewportManualFocusDistance() : 0.0f;
+}
+
 void SNiagaraSystemViewport::RefreshViewport()
 {
 	//reregister the preview components, so if the preview material changed it will be propagated to the render thread
@@ -1209,9 +1257,9 @@ void SNiagaraSystemViewport::OnFocusViewportToSelection()
 	if( PreviewComponent )
 	{
 		// FocusViewportOnBox disables orbit, so remember our state
-		bool bIsOrbit = SystemViewportClient->ShouldOrbitCamera();
+		const bool bIsOrbit = SystemViewportClient->ShouldOrbitCamera();
 
-		SystemViewportClient->FocusViewportOnBox(PreviewComponent->Bounds.GetBox());
+		SystemViewportClient->FocusViewport(PreviewComponent, false);
 
 		// this will reactivate orbit mode after the transition is done, if needed
 		bIsViewTransitioning = true;
diff --git a/Engine/Plugins/FX/Niagara/Source/NiagaraEditor/Private/Widgets/SNiagaraSystemViewport.h b/Engine/Plugins/FX/Niagara/Source/NiagaraEditor/Private/Widgets/SNiagaraSystemViewport.h
--- a/Engine/Plugins/FX/Niagara/Source/NiagaraEditor/Private/Widgets/SNiagaraSystemViewport.h
+++ b/Engine/Plugins/FX/Niagara/Source/NiagaraEditor/Private/Widgets/SNiagaraSystemViewport.h
@@ -112,6 +112,11 @@ public:
 	float GetFloorOffset() const;
 	void SetFloorOffset(float Offset);
 
+	bool IsManualFocusDistanceEnabled() const;
+	void ToggleManualFocusDistance();
+	void SetManualFocusDistance(float Distance);
+	float GetManualFocusDistance() const;
+
 protected:	
 	/** SEditorViewport interface */
 	virtual TSharedRef<FEditorViewportClient> MakeEditorViewportClient() override;
diff --git a/Engine/Plugins/FX/Niagara/Source/NiagaraEditor/Public/NiagaraSystemEditorData.h b/Engine/Plugins/FX/Niagara/Source/NiagaraEditor/Public/NiagaraSystemEditorData.h
--- a/Engine/Plugins/FX/Niagara/Source/NiagaraEditor/Public/NiagaraSystemEditorData.h
+++ b/Engine/Plugins/FX/Niagara/Source/NiagaraEditor/Public/NiagaraSystemEditorData.h
@@ -133,6 +133,12 @@ public:
 	NIAGARAEDITOR_API void SetAssetViewportSettings(FNiagaraPerAssetViewportSettings InSettings);
 	NIAGARAEDITOR_API void SetUseOrbitMode(bool bInUseOrbitMode);
 	
+	bool GetViewportManualFocusDistanceEnabled() const { return bViewportManualFocusDistanceEnabled; }
+	void SetViewportManualFocusDistanceEnabled(bool bEnabled) { bViewportManualFocusDistanceEnabled = bEnabled; }
+
+	float GetViewportManualFocusDistance() const { return ViewportManualFocusDistance; }
+	void SetViewportManualFocusDistance(float Distance) { ViewportManualFocusDistance = FMath::Max(Distance, 1.0f); }
+
 	/** Contains the root ids for organizing user parameters. */
 	UPROPERTY()
 	TObjectPtr<UHierarchyRoot> UserParameterHierarchy;
@@ -179,4 +185,10 @@ private:
 	TArray<TObjectPtr<UNiagaraScriptVariable>> UserParameterMetaData;
 	
 	FOnUserParameterScriptVariablesSynced OnUserParameterScriptVariablesSyncedDelegate;
+
+	UPROPERTY()
+	bool bViewportManualFocusDistanceEnabled = false;
+
+	UPROPERTY()
+	float ViewportManualFocusDistance = 100.0f;
 };
