- Tweak how focus bounds works
- Provide a minimum size for the initial bounds to avoid focusing on a 1x1x1 square
- Use Bounds for initial bounds not CalcLocalBounds, as calculating the local bounds makes no sense as it's in the wrong space

#Synced-CL 49227054

diff --git a/Engine/Plugins/FX/Niagara/Source/NiagaraEditor/Private/Widgets/SNiagaraSystemViewport.cpp b/Engine/Plugins/FX/Niagara/Source/NiagaraEditor/Private/Widgets/SNiagaraSystemViewport.cpp
--- a/Engine/Plugins/FX/Niagara/Source/NiagaraEditor/Private/Widgets/SNiagaraSystemViewport.cpp
+++ b/Engine/Plugins/FX/Niagara/Source/NiagaraEditor/Private/Widgets/SNiagaraSystemViewport.cpp
@@ -95,7 +95,7 @@ public:
 	virtual bool CanSetWidgetMode(UE::Widget::EWidgetMode NewMode) const override { return false; }
 	virtual bool CanCycleWidgetMode() const override { return false; }
 
-	void FocusViewport(UNiagaraComponent* Component, bool bInstant);
+	void FocusViewport(FBox FocusBounds, bool bInstant);
 
 	void SetOrbitModeFromSettings();
 	virtual void SetIsSimulateInEditorViewport(bool bInIsSimulateInEditorViewport) override;
@@ -200,11 +200,15 @@ void FNiagaraSystemViewportClient::Draw(FViewport* InViewport, FCanvas* Canvas)
 	FNiagaraSystemInstanceControllerConstPtr SystemInstanceController = NiagaraViewport.IsValid() ? NiagaraViewport->GetPreviewComponent()->GetSystemInstanceController() : nullptr;
 	if (Component && SystemInstanceController.IsValid() && SystemInstanceController->GetAge() > 0.0)
 	{
-		const FBox Bounds = Component->GetLocalBounds().GetBox();
+		FBox Bounds = Component->Bounds.GetBox();
 		if (bUpdateViewportFocus && Bounds.IsValid)
 		{
+			// Provide some minimum bounds for the instant snap
+			Bounds += Bounds.GetCenter() - FVector(100.0);
+			Bounds += Bounds.GetCenter() + FVector(100.0);
+
 			const bool bCachedShouldOrbitCamera = ShouldOrbitCamera();
-			FocusViewport(Component, true);
+			FocusViewport(Bounds, true);
 			ToggleOrbitCamera(bCachedShouldOrbitCamera);
 			SetUpdateViewportFocus(false);
 		}
@@ -702,9 +706,8 @@ FLinearColor FNiagaraSystemViewportClient::GetBackgroundColor() const
 	return BackgroundColor;
 }
 
-void FNiagaraSystemViewportClient::FocusViewport(UNiagaraComponent* Component, bool bInstant)
+void FNiagaraSystemViewportClient::FocusViewport(FBox FocusBounds, bool bInstant)
 {
-	FBox FocusBounds = Component->Bounds.GetBox();
 	if (FocusBounds.IsValid == false)
 	{
 		return;
@@ -1259,7 +1262,7 @@ void SNiagaraSystemViewport::OnFocusViewportToSelection()
 		// FocusViewportOnBox disables orbit, so remember our state
 		const bool bIsOrbit = SystemViewportClient->ShouldOrbitCamera();
 
-		SystemViewportClient->FocusViewport(PreviewComponent, false);
+		SystemViewportClient->FocusViewport(PreviewComponent->Bounds.GetBox(), false);
 
 		// this will reactivate orbit mode after the transition is done, if needed
 		bIsViewTransitioning = true;
