Niagra: expose functions for enabling scripting capabilities

#Synced-CL 48697252

diff --git a/Engine/Plugins/FX/Niagara/Source/NiagaraEditor/Public/NiagaraGraph.h b/Engine/Plugins/FX/Niagara/Source/NiagaraEditor/Public/NiagaraGraph.h
--- a/Engine/Plugins/FX/Niagara/Source/NiagaraEditor/Public/NiagaraGraph.h
+++ b/Engine/Plugins/FX/Niagara/Source/NiagaraEditor/Public/NiagaraGraph.h
@@ -280,10 +280,10 @@ public:
 	NIAGARAEDITOR_API bool GetPropertyMetadata(FName PropertyName, FString& OutValue) const;
 
 	/** Signal to listeners that the graph has changed */
-	void NotifyGraphNeedsRecompile();
+	NIAGARAEDITOR_API void NotifyGraphNeedsRecompile();
 		
 	/** Notifies the graph that a contained data interface has changed. */
-	void NotifyGraphDataInterfaceChanged();
+	NIAGARAEDITOR_API void NotifyGraphDataInterfaceChanged();
 
 	/** Get all referenced graphs in this specified graph, including this graph. */
 	void GetAllReferencedGraphs(TArray<const UNiagaraGraph*>& Graphs) const;
@@ -295,7 +295,7 @@ public:
 	bool IsOtherSynchronized(const FGuid& InChangeId) const;
 
 	/** Identify that this graph has undergone changes that will require synchronization with a compiled script.*/
-	void MarkGraphRequiresSynchronization(FString Reason);
+	NIAGARAEDITOR_API void MarkGraphRequiresSynchronization(FString Reason);
 
 	/** A change was made to the graph that external parties should take note of. The ChangeID will be updated.*/
 	virtual void NotifyGraphChanged() override;
diff --git a/Engine/Plugins/FX/Niagara/Source/NiagaraEditor/Public/NiagaraNodeInput.h b/Engine/Plugins/FX/Niagara/Source/NiagaraEditor/Public/NiagaraNodeInput.h
--- a/Engine/Plugins/FX/Niagara/Source/NiagaraEditor/Public/NiagaraNodeInput.h
+++ b/Engine/Plugins/FX/Niagara/Source/NiagaraEditor/Public/NiagaraNodeInput.h
@@ -110,8 +110,8 @@ public:
 	bool IsDataInterface() const;
 	bool IsObjectAsset() const;
 
-	UNiagaraDataInterface* GetDataInterface() const;
-	void SetDataInterface(UNiagaraDataInterface* InDataInterface);
+	NIAGARAEDITOR_API UNiagaraDataInterface* GetDataInterface() const;
+	NIAGARAEDITOR_API void SetDataInterface(UNiagaraDataInterface* InDataInterface);
 	
 	TObjectPtr<UObject> GetObjectAsset() const;
 	void SetObjectAsset(TObjectPtr<UObject> Object);
diff --git a/Engine/Plugins/FX/Niagara/Source/NiagaraEditor/Public/ViewModels/Stack/NiagaraStackGraphUtilities.h b/Engine/Plugins/FX/Niagara/Source/NiagaraEditor/Public/ViewModels/Stack/NiagaraStackGraphUtilities.h
--- a/Engine/Plugins/FX/Niagara/Source/NiagaraEditor/Public/ViewModels/Stack/NiagaraStackGraphUtilities.h
+++ b/Engine/Plugins/FX/Niagara/Source/NiagaraEditor/Public/ViewModels/Stack/NiagaraStackGraphUtilities.h
@@ -42,7 +42,7 @@ namespace FNiagaraStackGraphUtilities
 	void MakeLinkTo(UEdGraphPin* PinA, UEdGraphPin* PinB);
 	void BreakAllPinLinks(UEdGraphPin* PinA);
 
-	void RelayoutGraph(UEdGraph& Graph);
+	NIAGARAEDITOR_API void RelayoutGraph(UEdGraph& Graph);
 
 	void ConnectPinToInputNode(UEdGraphPin& Pin, UNiagaraNodeInput& InputNode);
 
@@ -217,13 +217,13 @@ namespace FNiagaraStackGraphUtilities
 
 	TArray<UEdGraphPin*> GetOverridePinsForFunction(UNiagaraNodeParameterMapSet& OverrideNode, UNiagaraNodeFunctionCall& FunctionCallNode);
 
-	void RemoveNodesForStackFunctionInputOverridePin(UEdGraphPin& StackFunctionInputOverridePin);
+	NIAGARAEDITOR_API void RemoveNodesForStackFunctionInputOverridePin(UEdGraphPin& StackFunctionInputOverridePin);
 
-	void RemoveNodesForStackFunctionInputOverridePin(UEdGraphPin& StackFunctinoInputOverridePin, TArray<TWeakObjectPtr<UNiagaraDataInterface>>& OutRemovedDataObjects);
+	NIAGARAEDITOR_API void RemoveNodesForStackFunctionInputOverridePin(UEdGraphPin& StackFunctinoInputOverridePin, TArray<TWeakObjectPtr<UNiagaraDataInterface>>& OutRemovedDataObjects);
 
 	UEdGraphPin* GetLinkedValueHandleForFunctionInput(const UEdGraphPin& OverridePin);
 
-	void GetParametersForContext(UEdGraph* Graph, UNiagaraSystem& System, TSet<FNiagaraVariableBase>& OutParameters);
+	NIAGARAEDITOR_API void GetParametersForContext(UEdGraph* Graph, UNiagaraSystem& System, TSet<FNiagaraVariableBase>& OutParameters);
 	NIAGARAEDITOR_API void SetLinkedParameterValueForFunctionInput(UEdGraphPin& OverridePin, const FNiagaraVariableBase& LinkedParameter, const TSet<FNiagaraVariableBase>& KnownParameters, ENiagaraDefaultMode DesiredDefaultMode = ENiagaraDefaultMode::FailIfPreviouslyNotSet, const FGuid& NewNodePersistentId = FGuid());
 
 	NIAGARAEDITOR_API void SetDataInterfaceValueForFunctionInput(UEdGraphPin& OverridePin, UClass* DataObjectType, FString InputNodeInputName, UNiagaraDataInterface*& OutDataObject, const FGuid& NewNodePersistentId = FGuid());
@@ -232,7 +232,7 @@ namespace FNiagaraStackGraphUtilities
 
 	NIAGARAEDITOR_API void SetDynamicInputForFunctionInput(UEdGraphPin& OverridePin, UNiagaraScript* DynamicInput, UNiagaraNodeFunctionCall*& OutDynamicInputFunctionCall, const FGuid& NewNodePersistentId = FGuid(), FString SuggestedName = FString(), const FGuid& InScriptVersion = FGuid());
 
-	void SetCustomExpressionForFunctionInput(UEdGraphPin& OverridePin, const FString& CustomExpression, UNiagaraNodeCustomHlsl*& OutDynamicInputFunctionCall, const FGuid& NewNodePersistentId = FGuid());
+	NIAGARAEDITOR_API void SetCustomExpressionForFunctionInput(UEdGraphPin& OverridePin, const FString& CustomExpression, UNiagaraNodeCustomHlsl*& OutDynamicInputFunctionCall, const FGuid& NewNodePersistentId = FGuid());
 
 	bool RemoveModuleFromStack(UNiagaraSystem& OwningSystem, FGuid OwningEmitterId, UNiagaraNodeFunctionCall& ModuleNode);
 
@@ -322,10 +322,10 @@ namespace FNiagaraStackGraphUtilities
 	void MoveModule(UNiagaraScript& SourceScript, UNiagaraNodeFunctionCall& ModuleToMove, UNiagaraSystem& TargetSystem, FGuid TargetEmitterHandleId, ENiagaraScriptUsage TargetUsage, FGuid TargetUsageId, int32 TargetModuleIndex, bool bForceCopy, UNiagaraNodeFunctionCall*& OutMovedModue);
 
 	/** Whether a parameter is allowed to be used in a certain execution category. 
-		Used to check if parameter can be dropped on a module or funciton stack entry. */
+		Used to check if parameter can be dropped on a module or function stack entry. */
 	NIAGARAEDITOR_API bool ParameterAllowedInExecutionCategory(const FName InParameterName, const FName ExecutionCategory);
 
-	void RebuildEmitterNodes(UNiagaraSystem& System);
+	NIAGARAEDITOR_API void RebuildEmitterNodes(UNiagaraSystem& System);
 
 	void FindAffectedScripts(UNiagaraSystem* System, FVersionedNiagaraEmitter Emitter, UNiagaraNodeFunctionCall& ModuleNode, TArray<TWeakObjectPtr<UNiagaraScript>>& OutAffectedScripts);
 
