fix packaged builds missing secondary particles

warnings instead of errors for bad emitter/DI on grid reader
valid bounds for secondary emitter

#virtualized
#Synced-CL 49380099

diff --git a/Engine/Plugins/FX/Niagara/Source/Niagara/Private/NiagaraDataInterfaceGrid2DCollectionReader.cpp b/Engine/Plugins/FX/Niagara/Source/Niagara/Private/NiagaraDataInterfaceGrid2DCollectionReader.cpp
--- a/Engine/Plugins/FX/Niagara/Source/Niagara/Private/NiagaraDataInterfaceGrid2DCollectionReader.cpp
+++ b/Engine/Plugins/FX/Niagara/Source/Niagara/Private/NiagaraDataInterfaceGrid2DCollectionReader.cpp
@@ -115,25 +115,31 @@ bool UNiagaraDataInterfaceGrid2DCollectionReader::InitPerInstanceData(void* PerI
 			// error if the reader is trying to access a grid that has both named and unnamed attributes
 			if (InstanceData->OtherInstanceData != nullptr && GridInterface->NumAttributes > 0 && GridInterface->NumAttributes != InstanceData->OtherInstanceData->NumAttributes)
 			{
-				UE_LOG(LogNiagara, Error, TEXT("Grid2DCollectionReaders are not compatible with Emitter %s and DataInterface %s.  Do not use unnamed and named grids together."), *this->EmitterName, *this->DIName);
-				return false;
+				UE_LOG(LogNiagara, Warning, TEXT("Grid2DCollectionReaders are not compatible with Emitter %s and DataInterface %s.  Do not use unnamed and named grids together."), *this->EmitterName, *this->DIName);				
 			}
-
-			// Push Updates to Proxy.
-			FNiagaraDataInterfaceProxyGrid2DCollectionProxy* RT_Proxy = GetProxyAs<FNiagaraDataInterfaceProxyGrid2DCollectionProxy>();
-			ENQUEUE_RENDER_COMMAND(FUpdateData)(
-				[RT_Proxy, InstanceID = SystemInstance->GetId(), RT_ProxyToUse = ProxyToUse](FRHICommandListImmediate& RHICmdList)
-			{
-				check(!RT_Proxy->SystemInstancesToProxyData_RT.Contains(InstanceID));				
-
-				FGrid2DCollectionRWInstanceData_RenderThread* TargetData = &RT_Proxy->SystemInstancesToProxyData_RT.Add(InstanceID);
-				TargetData->OtherProxy = RT_ProxyToUse;
-			});
-			return true;
 		}
+		else
+		{
+			UE_LOG(LogNiagara, Warning, TEXT("Grid2DCollectionReader %s refers to an invalid Grid2D Collection DI"), *this->DIName);				
+		}
+	}
+	else 
+	{
+		UE_LOG(LogNiagara, Warning, TEXT("Grid2DCollectionReader %s has invalid emitter"), *this->DIName);
 	}
 
-	return false;
+	// Push Updates to Proxy.
+	FNiagaraDataInterfaceProxyGrid2DCollectionProxy* RT_Proxy = GetProxyAs<FNiagaraDataInterfaceProxyGrid2DCollectionProxy>();
+	ENQUEUE_RENDER_COMMAND(FUpdateData)(
+		[RT_Proxy, InstanceID = SystemInstance->GetId(), RT_ProxyToUse = ProxyToUse](FRHICommandListImmediate& RHICmdList)
+	{
+		check(!RT_Proxy->SystemInstancesToProxyData_RT.Contains(InstanceID));				
+
+		FGrid2DCollectionRWInstanceData_RenderThread* TargetData = &RT_Proxy->SystemInstancesToProxyData_RT.Add(InstanceID);
+		TargetData->OtherProxy = RT_ProxyToUse;
+	});
+
+	return true;
 }
 
 void UNiagaraDataInterfaceGrid2DCollectionReader::DestroyPerInstanceData(void* PerInstanceData, FNiagaraSystemInstance* SystemInstance)
@@ -215,9 +221,12 @@ bool UNiagaraDataInterfaceGrid2DCollectionReader::PerInstanceTickPostSimulate(vo
 	if (InstanceData->OtherInstanceData == nullptr && InstanceData->OtherDI)
 	{
 		UNiagaraDataInterfaceGrid2DCollection* OtherGridDI = static_cast<UNiagaraDataInterfaceGrid2DCollection*>(InstanceData->OtherDI);
-		FGrid2DCollectionRWInstanceData_GameThread* Grid2DInstanceData = OtherGridDI->GetSystemInstancesToProxyData_GT().FindRef(SystemInstance->GetId());
-
-		InstanceData->OtherInstanceData = Grid2DInstanceData;
+		
+		if (OtherGridDI)
+		{
+			FGrid2DCollectionRWInstanceData_GameThread* Grid2DInstanceData = OtherGridDI->GetSystemInstancesToProxyData_GT().FindRef(SystemInstance->GetId());
+			InstanceData->OtherInstanceData = Grid2DInstanceData;
+		}
 	}
 
 	return false;
