- Fix for NaN if spline has no length when using LUT mode
- Also add debug HUD drawing for spline, shows regular spline in red and LUT in green

#tests Create a BP which contains a spline and NS, the NS samples the spline and make sure the spline has no points.  Use the debugger / HUD to observe NaNs when reading the data.

#Synced-CL 49372769

diff --git a/Engine/Plugins/FX/Niagara/Source/Niagara/Classes/NiagaraDataInterfaceSpline.h b/Engine/Plugins/FX/Niagara/Source/Niagara/Classes/NiagaraDataInterfaceSpline.h
--- a/Engine/Plugins/FX/Niagara/Source/Niagara/Classes/NiagaraDataInterfaceSpline.h
+++ b/Engine/Plugins/FX/Niagara/Source/Niagara/Classes/NiagaraDataInterfaceSpline.h
@@ -109,7 +109,7 @@ struct FNDISpline_InstanceData
 	
 
 	// We cache the version of the current spline curves so that we can reset the curves structure if we're using the LUT
-	uint32 SplineCurvesVersion;
+	TOptional<uint32> SplineCurvesVersion;
 
 	inline bool ResetRequired(UNiagaraDataInterfaceSpline* Interface, FNiagaraSystemInstance* SystemInstance) const;
 	
@@ -201,6 +201,9 @@ public:
 	NIAGARA_API virtual bool Equals(const UNiagaraDataInterface* Other) const override;
 	virtual bool CanExecuteOnTarget(ENiagaraSimTarget Target)const override { return true; }
 	virtual bool HasPreSimulateTick() const override { return true; }
+#if WITH_NIAGARA_DEBUGGER
+	NIAGARA_API virtual void DrawDebugHud(FNDIDrawDebugHudContext& DebugHudContext) const override;
+#endif
 	//UNiagaraDataInterface Interface End
 
 	template<typename UseLUT, typename TransformHandlerType, typename SplineSampleType>
diff --git a/Engine/Plugins/FX/Niagara/Source/Niagara/Private/NiagaraDataInterfaceSpline.cpp b/Engine/Plugins/FX/Niagara/Source/Niagara/Private/NiagaraDataInterfaceSpline.cpp
--- a/Engine/Plugins/FX/Niagara/Source/Niagara/Private/NiagaraDataInterfaceSpline.cpp
+++ b/Engine/Plugins/FX/Niagara/Source/Niagara/Private/NiagaraDataInterfaceSpline.cpp
@@ -1,7 +1,12 @@
 // Copyright Epic Games, Inc. All Rights Reserved.
 
 #include "NiagaraDataInterfaceSpline.h"
+#include "Engine/Canvas.h"
+#include "Internationalization/Internationalization.h"
 #include "Misc/LargeWorldRenderPosition.h"
+#include "CanvasItem.h"
+#include "ShaderCompilerCore.h"
+
 #include "NiagaraComponent.h"
 #include "NiagaraCompileHashVisitor.h"
 #include "NiagaraRenderer.h"
@@ -9,8 +14,6 @@
 #include "NiagaraStats.h"
 #include "NiagaraSystem.h"
 #include "NiagaraSystemInstance.h"
-#include "Internationalization/Internationalization.h"
-#include "ShaderCompilerCore.h"
 
 #include UE_INLINE_GENERATED_CPP_BY_NAME(NiagaraDataInterfaceSpline)
 
@@ -535,6 +538,80 @@ bool UNiagaraDataInterfaceSpline::Equals(const UNiagaraDataInterface* Other) con
 	return OtherTyped->SoftSourceActor == SoftSourceActor && OtherTyped->SplineUserParameter == SplineUserParameter && OtherTyped->bUseLUT == bUseLUT && OtherTyped->NumLUTSteps == NumLUTSteps;
 }
 
+#if WITH_NIAGARA_DEBUGGER
+void UNiagaraDataInterfaceSpline::DrawDebugHud(FNDIDrawDebugHudContext& DebugHudContext) const
+{
+	const FNDISpline_InstanceData* InstanceData_GT = DebugHudContext.GetSystemInstance()->FindTypedDataInterfaceInstanceData<FNDISpline_InstanceData>(this);
+	if (InstanceData_GT == nullptr)
+	{
+		return;
+	}
+
+	USplineComponent* SplineComponent = InstanceData_GT->Component.Get();
+	AActor* SceneActor = SplineComponent ? SplineComponent->GetAttachmentRootActor() : nullptr;
+	DebugHudContext.GetOutputString().Appendf(TEXT("SplineComponent(%s:%s)"), SceneActor ? *SceneActor->GetActorNameOrLabel() : TEXT("None"), *GetNameSafe(SplineComponent));
+
+	if (!DebugHudContext.IsVerbose() || SplineComponent == nullptr)
+	{
+		return;
+	}
+
+	UCanvas* Canvas = DebugHudContext.GetCanvas();
+
+	auto DrawDebugSpline =
+		[Canvas](int32 NumPoints, TFunction<FVector(int)> GetPoint, const FColor& Color)
+		{
+			TOptional<FVector> LastScreenPosition;
+			for (int32 i = 0; i < NumPoints; ++i)
+			{
+				const FVector WorldPosition = GetPoint(i);
+				const FVector ScreenPosition = Canvas->Project(WorldPosition, false);
+				if (ScreenPosition.Z > 0.0f)
+				{
+					if (LastScreenPosition.IsSet())
+					{
+						FCanvasLineItem LineItem(FVector2D(ScreenPosition), FVector2D(LastScreenPosition.GetValue()));
+						LineItem.SetColor(Color);
+						Canvas->Canvas->DrawItem(LineItem);
+					}
+
+					Canvas->Canvas->DrawNGon(FVector2D(ScreenPosition), Color, 8, 4.0f);
+					LastScreenPosition.Emplace(ScreenPosition);
+				}
+				else
+				{
+					LastScreenPosition.Reset();
+				}
+			}
+		};
+
+	// Draw points in spline
+	DrawDebugSpline(
+		SplineComponent->GetNumberOfSplinePoints(),
+		[SplineComponent](int32 PointIndex)
+		{
+			const FSplinePoint SplinePoint = SplineComponent->GetSplinePointAt(PointIndex, ESplineCoordinateSpace::World);
+			return SplinePoint.Position;
+		},
+		FColor::Red
+	);
+
+	// Draw LUT points
+	if (bUseLUT)
+	{
+		DrawDebugSpline(
+			InstanceData_GT->SplineLUT.Positions.Num(),
+			[InstanceData_GT, Transform=SplineComponent->GetComponentTransform()](int32 PointIndex)
+			{
+				const FVector Position = InstanceData_GT->SplineLUT.Positions[PointIndex];
+				return Transform.TransformPosition(Position);
+			},
+			FColor::Green
+		);
+	}
+}
+#endif
+
 int32 UNiagaraDataInterfaceSpline::PerInstanceDataSize()const
 {
 	return sizeof(FNDISpline_InstanceData);
@@ -552,7 +629,6 @@ bool UNiagaraDataInterfaceSpline::InitPerInstanceData(void* PerInstanceData, FNi
 	InstData->ComponentTransform = FTransform::Identity;
 	InstData->DefaultUpVector = FVector::UpVector;
 	InstData->bSyncedGPUCopy = false;
-	InstData->SplineCurvesVersion = INDEX_NONE;
 
 	InstData->SplineLUT.Reset();
 
@@ -660,7 +736,7 @@ bool UNiagaraDataInterfaceSpline::PerInstanceTick(void* PerInstanceData, FNiagar
 
 		bool bShouldBuildLUT = (bUseLUT || IsUsedWithGPUScript()) && InstData->SplineLUT.MaxIndex < 0;
 		
-		if (InstData->SplineCurvesVersion != SplineComponent->GetVersion())
+		if (!InstData->SplineCurvesVersion.IsSet() || (InstData->SplineCurvesVersion != SplineComponent->GetVersion()))
 		{
 			InstData->SplineCurves = SplineComponent->GetSplineCurves();
 			InstData->SplineCurvesVersion = InstData->SplineCurves.Version;
@@ -775,7 +851,7 @@ void FNiagaraDataInterfaceSplineLUT::BuildLUT(const FSplineCurves& SplineCurves,
 
 	SplineLength = SplineCurves.GetSplineLength();
 	SplineDistanceStep = MaxIndex ? ((1.0f / MaxIndex) * SplineLength) : 0.0f;
-	InvSplineDistanceStep = 1.0f / SplineDistanceStep;
+	InvSplineDistanceStep = SplineDistanceStep > 0.0f ? 1.0f / SplineDistanceStep : 0.0f;
 	
 	for (int32 Index = 0; Index < NumSteps; Index++)
 	{
@@ -1058,7 +1134,8 @@ float FNDISpline_InstanceData::EvaluateFindNearestPosition<TIntegralConstant<boo
 			KeyToNearest = i;
 		}
 	}
-	return KeyToNearest > 0 ? float(KeyToNearest) / float(SplineLUT.Positions.Num() - 1) : 0.0f;
+	const int32 NumPointsMinusOne = SplineLUT.Positions.Num() - 1;
+	return KeyToNearest > 0 && NumPointsMinusOne > 0 ? float(KeyToNearest) / float(NumPointsMinusOne) : 0.0f;
 }
 
 template<typename UseLUT>
@@ -1445,7 +1522,7 @@ void UNiagaraDataInterfaceSpline::FindClosestUnitDistanceFromPositionWS(FVectorV
 			// This first call finds the key time, but this is not in 0..1 range for the spline. 
 			const float KeyTime = InstData->FindInputKeyClosestToWorldLocation<UseLUT>(WorldPos);
 			// We need to convert into the range by dividing through by the overall duration of the spline according to the keys.
-			const float UnitDistance = KeyTime / FinalKeyTime;
+			const float UnitDistance = FinalKeyTime > 0.0f ? KeyTime / FinalKeyTime : 0.0f;
 
 			*OutUnitDistance.GetDest() = UnitDistance;
 			OutUnitDistance.Advance();
