- Add Resize And Fill function to arrays

#tests Call function on arrays, visualize in the attribute spreadsheet

#Synced-CL 49372769

diff --git a/Engine/Plugins/FX/Niagara/Source/Niagara/Classes/NiagaraDataInterfaceArrayImpl.h b/Engine/Plugins/FX/Niagara/Source/Niagara/Classes/NiagaraDataInterfaceArrayImpl.h
--- a/Engine/Plugins/FX/Niagara/Source/Niagara/Classes/NiagaraDataInterfaceArrayImpl.h
+++ b/Engine/Plugins/FX/Niagara/Source/Niagara/Classes/NiagaraDataInterfaceArrayImpl.h
@@ -185,6 +185,7 @@ struct FNiagaraDataInterfaceArrayImplInternal
 
 	NIAGARA_API static const FName Function_ClearName;
 	NIAGARA_API static const FName Function_ResizeName;
+	NIAGARA_API static const FName Function_ResizeAndFillName;
 	NIAGARA_API static const FName Function_SetArrayElemName;
 	NIAGARA_API static const FName Function_AddName;
 	NIAGARA_API static const FName Function_RemoveLastElemName;
@@ -683,6 +684,10 @@ struct FNDIArrayProxyImpl : public INDIArrayProxyBase
 			check(BindingInfo.GetNumInputs() == 2 && BindingInfo.GetNumOutputs() == 0);
 			OutFunc.BindRaw(this, &FNDIArrayProxyImpl::VMResize);
 		}
+		else if (BindingInfo.Name == FNiagaraDataInterfaceArrayImplInternal::Function_ResizeAndFillName)
+		{
+			OutFunc.BindRaw(this, &FNDIArrayProxyImpl::VMResizeAndFill);
+		}
 		else if (BindingInfo.Name == FNiagaraDataInterfaceArrayImplInternal::Function_SetArrayElemName)
 		{
 			OutFunc.BindRaw(this, &FNDIArrayProxyImpl::VMSetValue);
@@ -794,6 +799,7 @@ struct FNDIArrayProxyImpl : public INDIArrayProxyBase
 
 			if ((FunctionInfo.DefinitionName == FNiagaraDataInterfaceArrayImplInternal::Function_ClearName) ||
 				(FunctionInfo.DefinitionName == FNiagaraDataInterfaceArrayImplInternal::Function_ResizeName) ||
+				(FunctionInfo.DefinitionName == FNiagaraDataInterfaceArrayImplInternal::Function_ResizeAndFillName) ||
 				(FunctionInfo.DefinitionName == FNiagaraDataInterfaceArrayImplInternal::Function_SetArrayElemName) ||
 				(FunctionInfo.DefinitionName == FNiagaraDataInterfaceArrayImplInternal::Function_AddName) ||
 				(FunctionInfo.DefinitionName == FNiagaraDataInterfaceArrayImplInternal::Function_RemoveLastElemName))
@@ -1170,6 +1176,33 @@ struct FNDIArrayProxyImpl : public INDIArrayProxyBase
 		InstanceData->bIsRenderDirty |= bShouldSyncToGpu;
 	}
 
+	void VMResizeAndFill(FVectorVMExternalFunctionContext& Context)
+	{
+		ensureMsgf(Context.GetNumInstances() == 1, TEXT("Setting the number of values in an array with more than one instance, which doesn't make sense"));
+		VectorVM::FUserPtrHandler<FNDIArrayInstanceData_GameThread<TArrayType>> InstanceData(Context);
+		FNDIInputParam<int32>			InNewNum(Context);
+		FNDIInputParam<bool>			InFillNewOnly(Context);
+		FNDIInputParam<TVMArrayType>	InDefaultValue(Context);
+
+		FWriteArrayRef ArrayData(Owner, InstanceData);
+
+		const int32 OldNum = ArrayData.GetArray().Num();
+		const int32 NewNum = FMath::Min(InNewNum.GetAndAdvance(), kSafeMaxElements);
+		ArrayData.GetArray().SetNumUninitialized(NewNum);
+
+		const int32 FirstToFill	= InFillNewOnly.GetAndAdvance() ? OldNum : 0;
+		if (NewNum > FirstToFill)
+		{
+			const TArrayType DefaultValue = InDefaultValue.GetAndAdvance();
+			for (int32 i = FirstToFill; i < NewNum; ++i)
+			{
+				ArrayData.GetArray()[i] = DefaultValue;
+			}
+		}
+
+		InstanceData->bIsRenderDirty |= bShouldSyncToGpu;
+	}
+
 	void VMSetValue(FVectorVMExternalFunctionContext& Context)
 	{
 		VectorVM::FUserPtrHandler<FNDIArrayInstanceData_GameThread<TArrayType>> InstanceData(Context);
diff --git a/Engine/Plugins/FX/Niagara/Source/Niagara/Private/NiagaraDataInterfaceArrayImpl.cpp b/Engine/Plugins/FX/Niagara/Source/Niagara/Private/NiagaraDataInterfaceArrayImpl.cpp
--- a/Engine/Plugins/FX/Niagara/Source/Niagara/Private/NiagaraDataInterfaceArrayImpl.cpp
+++ b/Engine/Plugins/FX/Niagara/Source/Niagara/Private/NiagaraDataInterfaceArrayImpl.cpp
@@ -14,6 +14,7 @@ const FName FNiagaraDataInterfaceArrayImplInternal::Function_GetName(TEXT("Get")
 
 const FName FNiagaraDataInterfaceArrayImplInternal::Function_ClearName(TEXT("Clear"));
 const FName FNiagaraDataInterfaceArrayImplInternal::Function_ResizeName(TEXT("Resize"));
+const FName FNiagaraDataInterfaceArrayImplInternal::Function_ResizeAndFillName(TEXT("ResizeAndFill"));
 const FName FNiagaraDataInterfaceArrayImplInternal::Function_SetArrayElemName(TEXT("SetArrayElem"));
 const FName FNiagaraDataInterfaceArrayImplInternal::Function_AddName(TEXT("Add"));
 const FName FNiagaraDataInterfaceArrayImplInternal::Function_RemoveLastElemName(TEXT("RemoveLastElem"));
@@ -103,6 +104,16 @@ void FNiagaraDataInterfaceArrayImplInternal::GetMutableFunctions(TArray<FNiagara
 		Sig.Description = NSLOCTEXT("Niagara", "Array_ResizeDesc", "Resizes the array to the specified size, initializing new elements with the default value.");
 	}
 
+	{
+		FNiagaraFunctionSignature& Sig = OutFunctions.Add_GetRef(DefaultMutableSig);
+		Sig.Name = FNiagaraDataInterfaceArrayImplInternal::Function_ResizeAndFillName;
+		Sig.ModuleUsageBitmask = ENiagaraScriptUsageMask::System | ENiagaraScriptUsageMask::Emitter;
+		Sig.Inputs.Emplace(FNiagaraTypeDefinition::GetIntDef(), TEXT("Num"));
+		Sig.Inputs.Emplace(FNiagaraTypeDefinition::GetBoolDef(), TEXT("FillNewOnly"));
+		Sig.Inputs.Emplace(ValueTypeDef, TEXT("Value"));
+		Sig.Description = NSLOCTEXT("Niagara", "Array_ResizeAndFillDesc", "Resizes the array to the specified size.  Will initialize all elements or only the new elements depending on Fill New Only.");
+	}
+
 	{
 		FNiagaraFunctionSignature& Sig = OutFunctions.Add_GetRef(DefaultMutableSig);
 		Sig.Name = FNiagaraDataInterfaceArrayImplInternal::Function_SetArrayElemName;
@@ -225,6 +236,7 @@ bool FNiagaraDataInterfaceArrayImplInternal::IsRWFunction(const FName FunctionNa
 	{
 		Function_ClearName,
 		Function_ResizeName,
+		Function_ResizeAndFillName,
 		Function_SetArrayElemName,
 		Function_AddName,
 		Function_RemoveLastElemName,
