- Fix missing null pointer check for invalid graph

#Synced-CL 49372769

diff --git a/Engine/Plugins/FX/Niagara/Source/NiagaraEditor/Private/ViewModels/Stack/NiagaraStackGraphUtilities.cpp b/Engine/Plugins/FX/Niagara/Source/NiagaraEditor/Private/ViewModels/Stack/NiagaraStackGraphUtilities.cpp
--- a/Engine/Plugins/FX/Niagara/Source/NiagaraEditor/Private/ViewModels/Stack/NiagaraStackGraphUtilities.cpp
+++ b/Engine/Plugins/FX/Niagara/Source/NiagaraEditor/Private/ViewModels/Stack/NiagaraStackGraphUtilities.cpp
@@ -1454,14 +1454,17 @@ UNiagaraNodeFunctionCall* FNiagaraStackGraphUtilities::FindFunctionCallNode(FGui
 {
 	TSharedRef<FNiagaraScriptGraphViewModel> GraphViewModel = EmitterViewModel->GetSharedScriptViewModel()->GetGraphViewModel();
 
-	TArray<UNiagaraNodeFunctionCall*> FunctionCalls;
-	GraphViewModel->GetGraph()->GetNodesOfClass<UNiagaraNodeFunctionCall>(FunctionCalls);
-
-	for(UNiagaraNodeFunctionCall* FunctionCall : FunctionCalls)
+	if (UNiagaraGraph* Graph = GraphViewModel->GetGraph())
 	{
-		if(FunctionCall->NodeGuid == FunctionCallGuid)
+		TArray<UNiagaraNodeFunctionCall*> FunctionCalls;
+		Graph->GetNodesOfClass<UNiagaraNodeFunctionCall>(FunctionCalls);
+
+		for(UNiagaraNodeFunctionCall* FunctionCall : FunctionCalls)
 		{
-			return FunctionCall;
+			if(FunctionCall->NodeGuid == FunctionCallGuid)
+			{
+				return FunctionCall;
+			}
 		}
 	}
 
