- Allow emitter mode to work with GPU emitters for light & decal renderers
- Fix some incorrect defaults with the light renderer when in emitter mode

#Synced-CL 49372769

diff --git a/Engine/Plugins/FX/Niagara/Source/Niagara/Private/NiagaraRendererLights.cpp b/Engine/Plugins/FX/Niagara/Source/Niagara/Private/NiagaraRendererLights.cpp
--- a/Engine/Plugins/FX/Niagara/Source/Niagara/Private/NiagaraRendererLights.cpp
+++ b/Engine/Plugins/FX/Niagara/Source/Niagara/Private/NiagaraRendererLights.cpp
@@ -228,21 +228,19 @@ FNiagaraDynamicDataBase* FNiagaraRendererLights::GenerateDynamicData(const FNiag
 		{
 			const FVector3f SimPos = DefaultPos;
 			const FLinearColor LightColor = DefaultColor;
-			const float LightExponent = DefaultExponent;
-			const float LightScattering = DefaultScattering;
 			const float Brightness = Properties->bAlphaScalesBrightness ? LightColor.A : 1.0f;
 
 			SimpleLightData& LightData = DynamicData->LightArray.AddDefaulted_GetRef();
 			LightData.LightEntry.Radius = LightRadius;
 			LightData.LightEntry.Color = FVector3f(LightColor) * Brightness + Properties->ColorAdd;
-			LightData.LightEntry.Exponent = Properties->bUseInverseSquaredFalloff ? 0 : LightExponent;
+			LightData.LightEntry.Exponent = Properties->bUseInverseSquaredFalloff ? 0 : DefaultExponent;
 			LightData.LightEntry.InverseExposureBlend = InverseExposureBlend;
 			LightData.LightEntry.bAffectTranslucency = Properties->bAffectsTranslucency;
 			LightData.LightEntry.bAllowMegaLights = Properties->bAllowMegaLights;
 			LightData.LightEntry.bMegaLightsCastShadows = Properties->bMegaLightsCastShadows;
-			LightData.LightEntry.VolumetricScatteringIntensity = LightScattering;
-			LightData.LightEntry.SpecularScale = Properties->SpecularScale;
-			LightData.LightEntry.DiffuseScale = Properties->DiffuseScale;
+			LightData.LightEntry.VolumetricScatteringIntensity = DefaultScattering;
+			LightData.LightEntry.SpecularScale = DefaultSpecularScale;
+			LightData.LightEntry.DiffuseScale = DefaultDiffuseScale;
 			LightData.LightEntry.LightId = FSimpleLightId(0, SourcePrimitiveId);
 
 			LightData.PerViewEntry.Position = SimToWorld.TransformPosition(FVector(SimPos));
diff --git a/Engine/Plugins/FX/Niagara/Source/Niagara/Public/NiagaraDecalRendererProperties.h b/Engine/Plugins/FX/Niagara/Source/Niagara/Public/NiagaraDecalRendererProperties.h
--- a/Engine/Plugins/FX/Niagara/Source/Niagara/Public/NiagaraDecalRendererProperties.h
+++ b/Engine/Plugins/FX/Niagara/Source/Niagara/Public/NiagaraDecalRendererProperties.h
@@ -38,7 +38,7 @@ public:
 	virtual FNiagaraRenderer* CreateEmitterRenderer(ERHIFeatureLevel::Type FeatureLevel, const FNiagaraEmitterInstance* Emitter, const FNiagaraSystemInstanceController& InController) override;
 	virtual class FNiagaraBoundsCalculator* CreateBoundsCalculator() override;
 	virtual void GetUsedMaterials(const FNiagaraEmitterInstance* InEmitter, TArray<UMaterialInterface*>& OutMaterials) const override;
-	virtual bool IsSimTargetSupported(ENiagaraSimTarget InSimTarget) const override { return InSimTarget == ENiagaraSimTarget::CPUSim; };
+	virtual bool IsSimTargetSupported(ENiagaraSimTarget InSimTarget) const override { return InSimTarget == ENiagaraSimTarget::CPUSim || SourceMode == ENiagaraRendererSourceDataMode::Emitter; };
 #if WITH_EDITORONLY_DATA
 	virtual const TArray<FNiagaraVariable>& GetOptionalAttributes() override;
 	virtual const FSlateBrush* GetStackIcon() const override;
diff --git a/Engine/Plugins/FX/Niagara/Source/Niagara/Public/NiagaraLightRendererProperties.h b/Engine/Plugins/FX/Niagara/Source/Niagara/Public/NiagaraLightRendererProperties.h
--- a/Engine/Plugins/FX/Niagara/Source/Niagara/Public/NiagaraLightRendererProperties.h
+++ b/Engine/Plugins/FX/Niagara/Source/Niagara/Public/NiagaraLightRendererProperties.h
@@ -34,7 +34,7 @@ public:
 	virtual FNiagaraRenderer* CreateEmitterRenderer(ERHIFeatureLevel::Type FeatureLevel, const FNiagaraEmitterInstance* Emitter, const FNiagaraSystemInstanceController& InController) override;
 	virtual class FNiagaraBoundsCalculator* CreateBoundsCalculator() override;
 	virtual void GetUsedMaterials(const FNiagaraEmitterInstance* InEmitter, TArray<UMaterialInterface*>& OutMaterials) const override;
-	virtual bool IsSimTargetSupported(ENiagaraSimTarget InSimTarget) const override { return (InSimTarget == ENiagaraSimTarget::CPUSim); };
+	virtual bool IsSimTargetSupported(ENiagaraSimTarget InSimTarget) const override { return InSimTarget == ENiagaraSimTarget::CPUSim || SourceMode == ENiagaraRendererSourceDataMode::Emitter; };
 #if WITH_EDITORONLY_DATA
 	virtual TArray<FNiagaraVariable> GetBoundAttributes() const override;
 	virtual const TArray<FNiagaraVariable>& GetOptionalAttributes() override;
