- Fix for restricted folder plugins not being accounted as plugins

#tests Have a module / dynamic input inside a restricted plugin and toggle on/off the plugin selection in the menu when adding

#Synced-CL 49561835

diff --git a/Engine/Plugins/FX/Niagara/Source/NiagaraEditor/Private/NiagaraEditorUtilities.cpp b/Engine/Plugins/FX/Niagara/Source/NiagaraEditor/Private/NiagaraEditorUtilities.cpp
--- a/Engine/Plugins/FX/Niagara/Source/NiagaraEditor/Private/NiagaraEditorUtilities.cpp
+++ b/Engine/Plugins/FX/Niagara/Source/NiagaraEditor/Private/NiagaraEditorUtilities.cpp
@@ -1962,39 +1962,62 @@ bool FNiagaraEditorUtilities::DoesItemMatchFilterText(const FText& FilterText, c
 
 TTuple<EScriptSource, FText> FNiagaraEditorUtilities::GetScriptSource(const FAssetData& ScriptAssetData)
 {
-	FString PackagePathLocal ="";
+	FString PackagePathLocal;
 	FPackageName::TryConvertGameRelativePackagePathToLocalPath(ScriptAssetData.PackagePath.ToString(), PackagePathLocal);
 
-	if(FPaths::IsUnderDirectory(PackagePathLocal, FPaths::EnginePluginsDir() / TEXT("FX")))
-	{
-		int32 ContentFoundIndex = PackagePathLocal.Find(TEXT("/Content"));
-
-		if(ContentFoundIndex != INDEX_NONE)
+	const auto ResolvePluginName =
+		[PackagePathLocal]() -> TOptional<FText>
 		{
-			FString LeftPart = PackagePathLocal.Left(ContentFoundIndex);
-			bool bFound = LeftPart.FindLastChar('/', ContentFoundIndex);
-
-			if(bFound)
+			int32 ContentFoundIndex = PackagePathLocal.Find(TEXT("/Content"));
+			if (ContentFoundIndex != INDEX_NONE)
 			{
-				FString PluginName = LeftPart.RightChop(ContentFoundIndex + 1);
-				return TTuple<EScriptSource, FText>(EScriptSource::Niagara, FText::FromString(PluginName));
+				const FString LeftPart = PackagePathLocal.Left(ContentFoundIndex);
+				const bool bFound = LeftPart.FindLastChar('/', ContentFoundIndex);
+
+				if (bFound)
+				{
+					const FString PluginName = LeftPart.RightChop(ContentFoundIndex + 1);
+
+					return FText::FromString(PluginName);
+				}
 			}
+			return {};
+		};
+
+	if(FPaths::IsUnderDirectory(PackagePathLocal, FPaths::EnginePluginsDir() / TEXT("FX")))
+	{
+		TOptional<FText> PluginName = ResolvePluginName();
+		if (PluginName.IsSet())
+		{
+			return TTuple<EScriptSource, FText>(EScriptSource::Niagara, PluginName.GetValue());
 		}
 	}
 	
 	if(FPaths::IsUnderDirectory(PackagePathLocal, FPaths::EnginePluginsDir()) || FPaths::IsUnderDirectory(PackagePathLocal, FPaths::ProjectPluginsDir()))
 	{
-		int32 ContentFoundIndex = PackagePathLocal.Find(TEXT("/Content"));
-
-		if(ContentFoundIndex != INDEX_NONE)
+		TOptional<FText> PluginName = ResolvePluginName();
+		if (PluginName.IsSet())
 		{
-			FString LeftPart = PackagePathLocal.Left(ContentFoundIndex);
-			bool bFound = LeftPart.FindLastChar('/', ContentFoundIndex);
+			return TTuple<EScriptSource, FText>(EScriptSource::Plugins, PluginName.GetValue());
+		}
+	}
 
-			if(bFound)
+	if (FPaths::IsRestrictedPath(PackagePathLocal))
+	{
+		const TArray<FString> PluginSourcePaths = { FPaths::ProjectDir(), FPaths::EngineDir() };
+		for (const FString& RestrictedPath : FPaths::GetRestrictedFolderNames())
+		{
+			for (const FString& PluginSourcePath : PluginSourcePaths)
 			{
-				FString PluginName = LeftPart.RightChop(ContentFoundIndex + 1);
-				return TTuple<EScriptSource, FText>(EScriptSource::Plugins, FText::FromString(PluginName));
+				const FString& RestrictedPluginPath = FPaths::Combine(PluginSourcePath, TEXT("Restricted/"), RestrictedPath, TEXT("Plugins/"));
+				if (FPaths::IsUnderDirectory(PackagePathLocal, RestrictedPluginPath))
+				{
+					TOptional<FText> PluginName = ResolvePluginName();
+					if (PluginName.IsSet())
+					{
+						return TTuple<EScriptSource, FText>(EScriptSource::Plugins, PluginName.GetValue());
+					}
+				}
 			}
 		}
 	}
