- Move ensure to log warning in hope of verifying a rare asan issue

#rnx

#Synced-CL 49562370

diff --git a/Engine/Plugins/FX/Niagara/Source/Niagara/Private/NiagaraGPUSystemTick.cpp b/Engine/Plugins/FX/Niagara/Source/Niagara/Private/NiagaraGPUSystemTick.cpp
--- a/Engine/Plugins/FX/Niagara/Source/Niagara/Private/NiagaraGPUSystemTick.cpp
+++ b/Engine/Plugins/FX/Niagara/Source/Niagara/Private/NiagaraGPUSystemTick.cpp
@@ -100,8 +100,8 @@ void FNiagaraGPUSystemTick::Init(FNiagaraSystemInstance* InSystemInstance)
 
 	// we want to include interpolation parameters (current and previous frame) if any of the emitters in the system
 	// require it
-	const bool IncludeInterpolationParameters = InSystemInstance->GPUParamIncludeInterpolation;
-	const int32 InterpFactor = IncludeInterpolationParameters ? 2 : 1;
+	const bool bIncludeInterpolationParameters = InSystemInstance->GPUParamIncludeInterpolation;
+	const int32 InterpFactor = bIncludeInterpolationParameters ? 2 : 1;
 
 	GlobalParamData = ParamDataBufferPtr;
 	SystemParamData = GlobalParamData + InterpFactor * sizeof(FNiagaraGlobalParameters);
@@ -112,7 +112,7 @@ void FNiagaraGPUSystemTick::Init(FNiagaraSystemInstance* InSystemInstance)
 	FMemory::Memcpy(SystemParamData, &InSystemInstance->GetSystemParameters(), sizeof(FNiagaraSystemParameters));
 	FMemory::Memcpy(OwnerParamData, &InSystemInstance->GetOwnerParameters(), sizeof(FNiagaraOwnerParameters));
 
-	if (IncludeInterpolationParameters)
+	if (bIncludeInterpolationParameters)
 	{
 		FMemory::Memcpy(GlobalParamData + sizeof(FNiagaraGlobalParameters), &InSystemInstance->GetGlobalParameters(true), sizeof(FNiagaraGlobalParameters));
 		FMemory::Memcpy(SystemParamData + sizeof(FNiagaraSystemParameters), &InSystemInstance->GetSystemParameters(true), sizeof(FNiagaraSystemParameters));
@@ -169,8 +169,11 @@ void FNiagaraGPUSystemTick::Init(FNiagaraSystemInstance* InSystemInstance)
 			InstanceData->Context = GPUContext;
 			check(GPUContext->MainDataSet);
 
-			// Validate that we generated interpolation parameters
-			ensureMsgf(GPUContext->HasInterpolationParameters ? IncludeInterpolationParameters : true, TEXT("Mismatch between required interpolation parameters and calculated."));
+			// Validate that we generated interpolation parameters if we need them for the context
+			if (GPUContext->HasInterpolationParameters && !bIncludeInterpolationParameters)
+			{
+				UE_LOG(LogNiagara, Warning, TEXT("Mismatch between required interpolation parameters in GPUContext(%d) vs calculated(%d) Sim(%s)"), GPUContext->HasInterpolationParameters, bIncludeInterpolationParameters, GPUContext->GetDebugSimName());
+			}
 
 			InstanceData->SpawnInfo = GPUContext->GpuSpawnInfo_GT;
 
@@ -189,7 +192,7 @@ void FNiagaraGPUSystemTick::Init(FNiagaraSystemInstance* InSystemInstance)
 
 			// actually copy all of the data over
 			FMemory::Memcpy(InstanceData->EmitterParamData, &InSystemInstance->GetEmitterParameters(EmitterIdx), sizeof(FNiagaraEmitterParameters));
-			if (IncludeInterpolationParameters)
+			if (bIncludeInterpolationParameters)
 			{
 				FMemory::Memcpy(InstanceData->EmitterParamData + sizeof(FNiagaraEmitterParameters), &InSystemInstance->GetEmitterParameters(EmitterIdx, true), sizeof(FNiagaraEmitterParameters));
 			}
