PR #14227: [Niagara][Performance] Detach auto manage attachment vfx components when returned to the pool

#Synced-CL 49562370

diff --git a/Engine/Plugins/FX/Niagara/Source/Niagara/Private/NiagaraComponent.cpp b/Engine/Plugins/FX/Niagara/Source/Niagara/Private/NiagaraComponent.cpp
--- a/Engine/Plugins/FX/Niagara/Source/Niagara/Private/NiagaraComponent.cpp
+++ b/Engine/Plugins/FX/Niagara/Source/Niagara/Private/NiagaraComponent.cpp
@@ -1837,11 +1837,6 @@ void UNiagaraComponent::OnSystemComplete(bool bExternalCompletion)
 		}
 		else
 		{
-			if (bAutoManageAttachment)
-			{
-				CancelAutoAttachment(/*bDetachFromParent=*/ true);
-			}
-
 			// For components who do not auto destroy or auto release back into the pool ensure we remove the render state
 			// Without doing this we will 'leak' proxies until the user manually destroys or releases.
 			if (IsRenderStateCreated())
@@ -1849,6 +1844,14 @@ void UNiagaraComponent::OnSystemComplete(bool bExternalCompletion)
 				DestroyRenderState_Concurrent();
 			}
 		}
+		
+		if (!bAutoDestroy)
+		{
+			if (bAutoManageAttachment)
+			{
+				CancelAutoAttachment(/*bDetachFromParent=*/ true);
+			}
+		}
 
 		if (IsRegisteredWithScalabilityManager())
 		{
