- Allow parameter collections to be overridden on system instances
- When a custom override is enabled the instance is forced into solo mode

#Synced-CL 49561835

diff --git a/Engine/Plugins/FX/Niagara/Source/Niagara/Classes/NiagaraEmitterInstance.h b/Engine/Plugins/FX/Niagara/Source/Niagara/Classes/NiagaraEmitterInstance.h
--- a/Engine/Plugins/FX/Niagara/Source/Niagara/Classes/NiagaraEmitterInstance.h
+++ b/Engine/Plugins/FX/Niagara/Source/Niagara/Classes/NiagaraEmitterInstance.h
@@ -13,6 +13,7 @@ struct INiagaraComputeDataBufferInterface;
 struct FNiagaraComputeExecutionContext;
 struct FNiagaraScriptExecutionContext;
 class UNiagaraDataInterface;
+class UNiagaraParameterCollectionInstance;
 class FNiagaraSystemInstance;
 
 class FNiagaraEmitterInstanceImpl;
@@ -38,6 +39,7 @@ public:
 
 	virtual void BindParameters(bool bExternalOnly) = 0;
 	virtual void UnbindParameters(bool bExternalOnly) = 0;
+	virtual void RebindParameterCollection(UNiagaraParameterCollectionInstance* OldCollection, UNiagaraParameterCollectionInstance* NewCollection) = 0;
 
 	virtual bool ShouldTick() const = 0;
 	virtual void PreTick() {}
diff --git a/Engine/Plugins/FX/Niagara/Source/Niagara/Internal/NiagaraEmitterInstanceImpl.h b/Engine/Plugins/FX/Niagara/Source/Niagara/Internal/NiagaraEmitterInstanceImpl.h
--- a/Engine/Plugins/FX/Niagara/Source/Niagara/Internal/NiagaraEmitterInstanceImpl.h
+++ b/Engine/Plugins/FX/Niagara/Source/Niagara/Internal/NiagaraEmitterInstanceImpl.h
@@ -57,6 +57,7 @@ public:
 	virtual TConstArrayView<UNiagaraRendererProperties*> GetRenderers() const override;
 	virtual void BindParameters(bool bExternalOnly) override;
 	virtual void UnbindParameters(bool bExternalOnly) override;
+	virtual void RebindParameterCollection(UNiagaraParameterCollectionInstance* OldCollection, UNiagaraParameterCollectionInstance* NewCollection) override;
 	virtual bool ShouldTick() const override;
 	virtual void PreTick() override;
 	virtual void Tick(float DeltaSeconds) override;
diff --git a/Engine/Plugins/FX/Niagara/Source/Niagara/Internal/Stateless/NiagaraStatelessEmitterInstance.h b/Engine/Plugins/FX/Niagara/Source/Niagara/Internal/Stateless/NiagaraStatelessEmitterInstance.h
--- a/Engine/Plugins/FX/Niagara/Source/Niagara/Internal/Stateless/NiagaraStatelessEmitterInstance.h
+++ b/Engine/Plugins/FX/Niagara/Source/Niagara/Internal/Stateless/NiagaraStatelessEmitterInstance.h
@@ -78,6 +78,7 @@ public:
 	virtual TConstArrayView<UNiagaraRendererProperties*> GetRenderers() const override;
 	virtual void BindParameters(bool bExternalOnly) override;
 	virtual void UnbindParameters(bool bExternalOnly) override;
+	virtual void RebindParameterCollection(UNiagaraParameterCollectionInstance* OldCollection, UNiagaraParameterCollectionInstance* NewCollection) override;
 	virtual bool ShouldTick() const override;
 	virtual void Tick(float DeltaSeconds) override;
 	// FNiagaraEmitterInstance Impl
diff --git a/Engine/Plugins/FX/Niagara/Source/Niagara/Private/NiagaraComponent.cpp b/Engine/Plugins/FX/Niagara/Source/Niagara/Private/NiagaraComponent.cpp
--- a/Engine/Plugins/FX/Niagara/Source/Niagara/Private/NiagaraComponent.cpp
+++ b/Engine/Plugins/FX/Niagara/Source/Niagara/Private/NiagaraComponent.cpp
@@ -12,6 +12,7 @@
 #include "NiagaraDataInterface.h"
 #include "NiagaraFunctionLibrary.h"
 #include "NiagaraGpuComputeDispatchInterface.h"
+#include "NiagaraParameterCollectionResolver.h"
 #include "NiagaraRenderer.h"
 #include "NiagaraSceneComponentUtils.h"
 #include "NiagaraSceneProxy.h"
@@ -1925,6 +1926,8 @@ void UNiagaraComponent::OnPooledReuse(UWorld* NewWorld)
 {
 	check(IsValidChecked(this));
 
+	ParameterCollectionOverrides.Reset();
+
 	SetUserParametersToDefaultValues();
 
 	//Need to reset the component's visibility in case it's returned to the pool while marked invisible.
@@ -3037,6 +3040,94 @@ FMatrix UNiagaraComponent::GetVariableMatrix(FName InVariableName, bool& bIsVali
 	return FMatrix(Value.Get(FMatrix44f::Identity));
 }
 
+void UNiagaraComponent::SetParameterCollectionOverride(UNiagaraParameterCollectionInstance* OverrideCollection)
+{
+	if (OverrideCollection == nullptr)
+	{
+		return;
+	}
+
+	UNiagaraParameterCollection* OverrideCollectionParent = OverrideCollection->GetParent();
+	UNiagaraParameterCollectionInstance* ExistingCollection = nullptr;
+	if (SystemInstanceController)
+	{
+		// Find existing collection, if the same we can early out
+		const FNiagaraParameterCollectionResolver ParameterCollectionResolver = SystemInstanceController->CreateParameterCollectionResolver();
+		ExistingCollection = ParameterCollectionResolver.ResolveInstance(OverrideCollectionParent);
+		if (ExistingCollection == OverrideCollection)
+		{
+			return;
+		}
+	}
+
+	// Add collection, replacing a one that shared the same parent if we have one
+	const int32 ExistingIndex = ParameterCollectionOverrides.IndexOfByPredicate([OverrideCollectionParent](UNiagaraParameterCollectionInstance* Existing) { return Existing && Existing->GetParent() == OverrideCollectionParent; });
+	if (ExistingIndex == INDEX_NONE)
+	{
+		ParameterCollectionOverrides.Add(OverrideCollection);
+	}
+	else
+	{
+		ParameterCollectionOverrides[ExistingIndex] = OverrideCollection;
+	}
+
+	// Ensure we are in solo mode
+	// Note: This must be done BEFORE we modify any bindings
+	UpdateInstanceSoloMode();
+
+	// Bind the override collection
+	if (SystemInstanceController)
+	{
+		SystemInstanceController->RebindParameterCollection(ExistingCollection, OverrideCollection);
+	}
+}
+
+void UNiagaraComponent::RemoveParameterCollectionOverride(UNiagaraParameterCollectionInstance* OverrideCollection)
+{
+	if (OverrideCollection == nullptr)
+	{
+		return;
+	}
+
+	if (ParameterCollectionOverrides.RemoveSwap(OverrideCollection) == 0)
+	{
+		return;
+	}
+
+	if (SystemInstanceController)
+	{
+		const FNiagaraParameterCollectionResolver ParameterCollectionResolver = SystemInstanceController->CreateParameterCollectionResolver();
+
+		UNiagaraParameterCollection* OverrideCollectionParent = OverrideCollection->GetParent();
+		UNiagaraParameterCollectionInstance* ReplacementCollection = ParameterCollectionResolver.ResolveInstance(OverrideCollectionParent);
+		SystemInstanceController->RebindParameterCollection(OverrideCollection, ReplacementCollection);
+	}
+
+	// Update solo mode
+	// Note: This must be done AFTER we have potentially modified the bindings
+	UpdateInstanceSoloMode();
+}
+
+void UNiagaraComponent::ClearParameterCollectionOverrides()
+{
+	while (ParameterCollectionOverrides.Num() > 0)
+	{
+		RemoveParameterCollectionOverride(ParameterCollectionOverrides.Last());
+	}
+}
+
+UNiagaraParameterCollectionInstance* UNiagaraComponent::GetParameterCollectionOverride(UNiagaraParameterCollection* ParameterCollection) const
+{
+	const int32 ExistingIndex = ParameterCollectionOverrides.IndexOfByPredicate(
+		[Parent = ParameterCollection](const UNiagaraParameterCollectionInstance* Instance)
+		{
+			return Instance && Instance->GetParent() == Parent;
+		}
+	);
+
+	return ExistingIndex == INDEX_NONE ? nullptr : ParameterCollectionOverrides[ExistingIndex];
+}
+
 void FixInvalidUserParameters(FNiagaraUserRedirectionParameterStore& ParameterStore)
 {
 	static const FString UserPrefix = FNiagaraConstants::UserNamespace.ToString() + TEXT(".");
@@ -4529,7 +4620,8 @@ void UNiagaraComponent::SetForceSolo(bool bInForceSolo)
 
 bool UNiagaraComponent::RequiresSoloMode() const
 {
-	return bForceSolo || SimCache != nullptr || !FMath::IsNearlyEqual(CustomTimeDilation, 1.0f) || (Asset && Asset->HasFixedTickDelta());
+	//-OPT: If parameter collection overrides are only bound into the particle scripts we can avoid forcing solo mode
+	return bForceSolo || SimCache != nullptr || !FMath::IsNearlyEqual(CustomTimeDilation, 1.0f) || (Asset && Asset->HasFixedTickDelta()) || (ParameterCollectionOverrides.Num() > 0);
 }
 
 void UNiagaraComponent::UpdateInstanceSoloMode()
diff --git a/Engine/Plugins/FX/Niagara/Source/Niagara/Private/NiagaraEmitterInstanceImpl.cpp b/Engine/Plugins/FX/Niagara/Source/Niagara/Private/NiagaraEmitterInstanceImpl.cpp
--- a/Engine/Plugins/FX/Niagara/Source/Niagara/Private/NiagaraEmitterInstanceImpl.cpp
+++ b/Engine/Plugins/FX/Niagara/Source/Niagara/Private/NiagaraEmitterInstanceImpl.cpp
@@ -6,6 +6,7 @@
 #include "NiagaraStats.h"
 #include "NiagaraConstants.h"
 #include "NiagaraModule.h"
+#include "NiagaraParameterCollectionResolver.h"
 #include "NiagaraSystemInstance.h"
 #include "NiagaraComputeExecutionContext.h"
 #include "NiagaraDataInterface.h"
@@ -637,22 +638,6 @@ void FNiagaraEmitterInstanceImpl::DirtyDataInterfaces()
 	}
 }
 
-//Unsure on usage of this atm. Possibly useful in future.
-// void FNiagaraEmitterInstanceImpl::RebindParameterCollection(UNiagaraParameterCollectionInstance* OldInstance, UNiagaraParameterCollectionInstance* NewInstance)
-// {
-// 	OldInstance->GetParameterStore().Unbind(&SpawnExecContext.Parameters);
-// 	NewInstance->GetParameterStore().Bind(&SpawnExecContext.Parameters);
-// 
-// 	OldInstance->GetParameterStore().Unbind(&UpdateExecContext.Parameters);
-// 	NewInstance->GetParameterStore().Bind(&UpdateExecContext.Parameters);
-// 
-// 	for (FNiagaraScriptExecutionContext& EventContext : EventExecContexts)
-// 	{
-// 		OldInstance->GetParameterStore().Unbind(&EventContext.Parameters);
-// 		NewInstance->GetParameterStore().Bind(&EventContext.Parameters);
-// 	}
-// }
-
 TConstArrayView<UNiagaraRendererProperties*> FNiagaraEmitterInstanceImpl::GetRenderers() const
 {
 	FVersionedNiagaraEmitterData* EmitterData = VersionedEmitter.GetEmitterData();
@@ -663,29 +648,26 @@ void FNiagaraEmitterInstanceImpl::UnbindParameters(bool bExternalOnly)
 {
 	if (bExternalOnly && !IsDisabled())
 	{
-		for (UNiagaraParameterCollection* Collection : SpawnExecContext.Script->GetCachedParameterCollectionReferences())
+		const FNiagaraParameterCollectionResolver ParameterCollectionResolver = ParentSystemInstance->CreateParameterCollectionResolver();
+		ParameterCollectionResolver.UnbindCollections(SpawnExecContext.Script, SpawnExecContext.Parameters);
+		ParameterCollectionResolver.UnbindCollections(UpdateExecContext.Script, UpdateExecContext.Parameters);
+
+		for (FNiagaraScriptExecutionContext& EventContext : GetEventExecutionContexts())
 		{
-			if (UNiagaraParameterCollectionInstance* NPCInst = ParentSystemInstance->GetParameterCollectionInstance(Collection))
-			{
-				NPCInst->GetParameterStore().Unbind(&SpawnExecContext.Parameters);
-			}
-		}
-		for (UNiagaraParameterCollection* Collection : UpdateExecContext.Script->GetCachedParameterCollectionReferences())
-		{
-			if (UNiagaraParameterCollectionInstance* NPCInst = ParentSystemInstance->GetParameterCollectionInstance(Collection))
-			{
-				NPCInst->GetParameterStore().Unbind(&UpdateExecContext.Parameters);
-			}
+			ParameterCollectionResolver.UnbindCollections(EventContext.Script, EventContext.Parameters);
 		}
 
-		for (FNiagaraScriptExecutionContext& EventContext : GetEventExecutionContexts())
+		if (FVersionedNiagaraEmitterData* EmitterData = VersionedEmitter.GetEmitterData())
 		{
-			for (UNiagaraParameterCollection* Collection : EventContext.Script->GetCachedParameterCollectionReferences())
+			FNiagaraScriptInstanceParameterStore& TargetParamStore = EmitterData->SimTarget == ENiagaraSimTarget::GPUComputeSim ? GPUExecContext->CombinedParamStore : UpdateExecContext.Parameters;
+			for (const UNiagaraSimulationStageBase* SimStage : EmitterData->GetSimulationStages())
 			{
-				if (UNiagaraParameterCollectionInstance* NPCInst = ParentSystemInstance->GetParameterCollectionInstance(Collection))
+				if (SimStage->bEnabled == false)
 				{
-					NPCInst->GetParameterStore().Unbind(&EventContext.Parameters);
+					continue;
 				}
+
+				ParameterCollectionResolver.UnbindCollections(SimStage->Script, TargetParamStore);
 			}
 		}
 
@@ -716,53 +698,20 @@ void FNiagaraEmitterInstanceImpl::BindParameters(bool bExternalOnly)
 		return;
 	}
 
-	auto BindToParameterCollection = [&](UNiagaraParameterCollection* Collection, FNiagaraParameterStore& DestStore)
-	{
-		if (Collection)
-		{
-			if (UNiagaraParameterCollectionInstance* Inst = ParentSystemInstance->GetParameterCollectionInstance(Collection))
-			{
-				Inst->GetParameterStore().Bind(&DestStore);
-			}
-			else
-			{
-				UE_LOG(LogNiagara, Error, TEXT("Emitter attempting to bind to a null Parameter Collection Instance.\nEmitter:%s\nCollection:%s"), VersionedEmitter.Emitter ? *VersionedEmitter.Emitter->GetPathName() : TEXT("null emitter!"), *Collection->GetPathName());
-			}
-		}
-		else
-		{
-			UE_LOG(LogNiagara, Error, TEXT("Emitter attempting to bind to a null Parameter Collection.\nEmitter:%s"), VersionedEmitter.Emitter ? *VersionedEmitter.Emitter->GetPathName() : TEXT("null emitter!" ));
-		}
-	};
-
-	for (UNiagaraParameterCollection* Collection : SpawnExecContext.Script->GetCachedParameterCollectionReferences())
-	{
-		BindToParameterCollection(Collection, SpawnExecContext.Parameters);
-	}
-	for (UNiagaraParameterCollection* Collection : UpdateExecContext.Script->GetCachedParameterCollectionReferences())
-	{
-		BindToParameterCollection(Collection, UpdateExecContext.Parameters);
-	}
+	const FNiagaraParameterCollectionResolver ParameterCollectionResolver = ParentSystemInstance->CreateParameterCollectionResolver();
+	ParameterCollectionResolver.BindCollections(SpawnExecContext.Script, SpawnExecContext.Parameters);
+	ParameterCollectionResolver.BindCollections(UpdateExecContext.Script, UpdateExecContext.Parameters);
 
 	FVersionedNiagaraEmitterData* EmitterData = VersionedEmitter.GetEmitterData();
 	if (EmitterData->SimTarget == ENiagaraSimTarget::GPUComputeSim)
 	{
-		for (UNiagaraParameterCollection* Collection : SpawnExecContext.Script->GetCachedParameterCollectionReferences())
-		{
-			BindToParameterCollection(Collection, GPUExecContext->CombinedParamStore);
-		}
-		for (UNiagaraParameterCollection* Collection : UpdateExecContext.Script->GetCachedParameterCollectionReferences())
-		{
-			BindToParameterCollection(Collection, GPUExecContext->CombinedParamStore);
-		}
+		ParameterCollectionResolver.BindCollections(SpawnExecContext.Script, GPUExecContext->CombinedParamStore);
+		ParameterCollectionResolver.BindCollections(UpdateExecContext.Script, GPUExecContext->CombinedParamStore);
 	}
 
 	for (FNiagaraScriptExecutionContext& EventContext : GetEventExecutionContexts())
 	{
-		for (UNiagaraParameterCollection* Collection : EventContext.Script->GetCachedParameterCollectionReferences())
-		{
-			BindToParameterCollection(Collection, EventContext.Parameters);
-		}
+		ParameterCollectionResolver.BindCollections(EventContext.Script, EventContext.Parameters);
 	}
 
 	FNiagaraScriptInstanceParameterStore& TargetParamStore = EmitterData->SimTarget == ENiagaraSimTarget::GPUComputeSim ? GPUExecContext->CombinedParamStore : UpdateExecContext.Parameters;
@@ -773,10 +722,7 @@ void FNiagaraEmitterInstanceImpl::BindParameters(bool bExternalOnly)
 			continue;
 		}
 
-		for (UNiagaraParameterCollection* Collection : SimStage->Script->GetCachedParameterCollectionReferences())
-		{
-			BindToParameterCollection(Collection, TargetParamStore);
-		}
+		ParameterCollectionResolver.BindCollections(SimStage->Script, TargetParamStore);
 	}
 
 	if (!bExternalOnly)
@@ -835,6 +781,50 @@ void FNiagaraEmitterInstanceImpl::BindParameters(bool bExternalOnly)
 	}
 }
 
+void FNiagaraEmitterInstanceImpl::RebindParameterCollection(UNiagaraParameterCollectionInstance* OldCollection, UNiagaraParameterCollectionInstance* NewCollection)
+{
+	if (IsDisabled())
+	{
+		return;
+	}
+
+	const auto RebindToParameterStore =
+		[CollectionParent=OldCollection->GetParent(), OldCollection, NewCollection](UNiagaraScript* Script, FNiagaraParameterStore* ParameterStore)
+		{
+			if (Script && Script->GetCachedParameterCollectionReferences().Contains(CollectionParent))
+			{
+				OldCollection->GetParameterStore().Unbind(ParameterStore);
+				NewCollection->GetParameterStore().Bind(ParameterStore);
+			}
+		};
+
+	RebindToParameterStore(SpawnExecContext.Script, &SpawnExecContext.Parameters);
+	RebindToParameterStore(UpdateExecContext.Script, &UpdateExecContext.Parameters);
+
+	if (SimTarget == ENiagaraSimTarget::GPUComputeSim)
+	{
+		RebindToParameterStore(SpawnExecContext.Script, &GPUExecContext->CombinedParamStore);
+		RebindToParameterStore(UpdateExecContext.Script, &GPUExecContext->CombinedParamStore);
+	}
+
+	for (FNiagaraScriptExecutionContext& EventContext : GetEventExecutionContexts())
+	{
+		RebindToParameterStore(EventContext.Script, &EventContext.Parameters);
+	}
+
+	FVersionedNiagaraEmitterData* EmitterData = VersionedEmitter.GetEmitterData();
+	FNiagaraScriptInstanceParameterStore& TargetParamStore = SimTarget == ENiagaraSimTarget::GPUComputeSim ? GPUExecContext->CombinedParamStore : UpdateExecContext.Parameters;
+	for (const UNiagaraSimulationStageBase* SimStage : EmitterData->GetSimulationStages())
+	{
+		if (SimStage->bEnabled == false)
+		{
+			continue;
+		}
+
+		RebindToParameterStore(SimStage->Script, &TargetParamStore);
+	}
+}
+
 FBox FNiagaraEmitterInstanceImpl::InternalCalculateDynamicBounds(int32 ParticleCount) const
 {
 	FVersionedNiagaraEmitterData* EmitterData = VersionedEmitter.GetEmitterData();
diff --git a/Engine/Plugins/FX/Niagara/Source/Niagara/Private/NiagaraFunctionLibrary.cpp b/Engine/Plugins/FX/Niagara/Source/Niagara/Private/NiagaraFunctionLibrary.cpp
--- a/Engine/Plugins/FX/Niagara/Source/Niagara/Private/NiagaraFunctionLibrary.cpp
+++ b/Engine/Plugins/FX/Niagara/Source/Niagara/Private/NiagaraFunctionLibrary.cpp
@@ -743,6 +743,29 @@ UNiagaraParameterCollectionInstance* UNiagaraFunctionLibrary::GetNiagaraParamete
 	return nullptr;
 }
 
+UNiagaraParameterCollectionInstance* UNiagaraFunctionLibrary::CreateNiagaraParameterCollectionInstance(UObject* WorldContextObject, UNiagaraParameterCollection* ParentCollection)
+{
+	UWorld* World = GEngine->GetWorldFromContextObject(WorldContextObject, EGetWorldErrorMode::LogAndReturnNull);
+	if (!ParentCollection || !World)
+	{
+		return nullptr;
+	}
+
+	UObject* CollectionOuter = GetTransientPackage();
+	static const FName CollectionBaseName("RuntimeNPC");
+	const FName CollectionName = MakeUniqueObjectName(CollectionOuter, UNiagaraParameterCollectionInstance::StaticClass(), CollectionBaseName);
+
+	UNiagaraParameterCollectionInstance* DefaultInstance = ParentCollection->GetDefaultInstance();
+	UNiagaraParameterCollectionInstance* NewInstance = CastChecked<UNiagaraParameterCollectionInstance>(StaticDuplicateObject(DefaultInstance, CollectionOuter, CollectionName));
+
+#if WITH_EDITOR
+	// Bind to the default instance so that changes to the collection propagate through.
+	DefaultInstance->GetParameterStore().Bind(&NewInstance->GetParameterStore());
+#endif
+
+	NewInstance->Bind(World);
+	return NewInstance;
+}
 
 const TArray<FNiagaraFunctionSignature>& UNiagaraFunctionLibrary::GetVectorVMFastPathOps(bool bIgnoreConsoleVariable)
 {
diff --git a/Engine/Plugins/FX/Niagara/Source/Niagara/Private/NiagaraParameterCollectionResolver.cpp b/Engine/Plugins/FX/Niagara/Source/Niagara/Private/NiagaraParameterCollectionResolver.cpp
new file mode 100644
--- /dev/null
+++ b/Engine/Plugins/FX/Niagara/Source/Niagara/Private/NiagaraParameterCollectionResolver.cpp
@@ -0,0 +1,94 @@
+// Copyright Epic Games, Inc. All Rights Reserved.
+
+#include "NiagaraParameterCollectionResolver.h"
+
+#include "NiagaraComponent.h"
+#include "NiagaraParameterCollection.h"
+#include "NiagaraScript.h"
+#include "NiagaraSystem.h"
+#include "NiagaraSystemSimulation.h"
+#include "NiagaraWorldManager.h"
+
+FNiagaraParameterCollectionResolver::FNiagaraParameterCollectionResolver(UNiagaraComponent* InNiagaraComponent, UNiagaraSystem* InNiagaraSystem, UWorld* World)
+{
+	NiagaraComponent	= InNiagaraComponent;
+	NiagaraSystem		= InNiagaraSystem;
+	NiagaraWorldManager	= World ? FNiagaraWorldManager::Get(World) : nullptr;
+}
+
+FNiagaraParameterCollectionResolver::FNiagaraParameterCollectionResolver(UNiagaraComponent* InNiagaraComponent, UNiagaraSystem* InNiagaraSystem, FNiagaraWorldManager* InNiagaraWorldManager)
+{
+	NiagaraComponent	= InNiagaraComponent;
+	NiagaraSystem		= InNiagaraSystem;
+	NiagaraWorldManager	= InNiagaraWorldManager;
+}
+
+void FNiagaraParameterCollectionResolver::BindCollections(TConstArrayView<TObjectPtr<UNiagaraParameterCollection>> ParameterCollections, FNiagaraParameterStore& DestinationStore) const
+{
+	for (UNiagaraParameterCollection* Collection : ParameterCollections)
+	{
+		if (Collection)
+		{
+			if (UNiagaraParameterCollectionInstance* CollectionInstance = ResolveInstance(Collection))
+			{
+				CollectionInstance->GetParameterStore().Bind(&DestinationStore);
+			}
+			else
+			{
+				UE_LOG(LogNiagara, Error, TEXT("Attempting to bind system simulation to a null parameter collection instance | Collection: %s | System: %s |"), *GetPathNameSafe(Collection), *GetPathNameSafe(NiagaraSystem));
+			}
+		}
+		else
+		{
+			UE_LOG(LogNiagara, Error, TEXT("Attempting to bind system simulation to a null parameter collection | System: %s |"), *GetPathNameSafe(NiagaraSystem));
+		}
+	}
+}
+
+void FNiagaraParameterCollectionResolver::UnbindCollections(TConstArrayView<TObjectPtr<UNiagaraParameterCollection>> ParameterCollections, FNiagaraParameterStore& DestinationStore) const
+{
+	for (UNiagaraParameterCollection* Collection : ParameterCollections)
+	{
+		UNiagaraParameterCollectionInstance* CollectionInstance = Collection ? ResolveInstance(Collection) : nullptr;
+		if (CollectionInstance)
+		{
+			CollectionInstance->GetParameterStore().Unbind(&DestinationStore);
+		}
+	}
+}
+
+void FNiagaraParameterCollectionResolver::BindCollections(UNiagaraScript* Script, FNiagaraParameterStore& DestinationStore) const
+{
+	BindCollections(Script->GetCachedParameterCollectionReferences(), DestinationStore);
+}
+
+void FNiagaraParameterCollectionResolver::UnbindCollections(UNiagaraScript* Script, FNiagaraParameterStore& DestinationStore) const
+{
+	UnbindCollections(Script->GetCachedParameterCollectionReferences(), DestinationStore);
+}
+
+UNiagaraParameterCollectionInstance* FNiagaraParameterCollectionResolver::ResolveInstance(UNiagaraParameterCollection* Collection) const
+{
+	if (NiagaraComponent)
+	{
+		if (UNiagaraParameterCollectionInstance* Instance = NiagaraComponent->GetParameterCollectionOverride(Collection))
+		{
+			return Instance;
+		}
+	}
+
+	if (NiagaraSystem)
+	{
+		if (UNiagaraParameterCollectionInstance* Instance = NiagaraSystem->GetParameterCollectionOverride(Collection))
+		{
+			return Instance;
+		}
+	}
+
+	if (NiagaraWorldManager)
+	{
+		return NiagaraWorldManager->GetParameterCollection(Collection);
+	}
+
+	return nullptr;
+}
diff --git a/Engine/Plugins/FX/Niagara/Source/Niagara/Private/NiagaraParameterCollectionResolver.h b/Engine/Plugins/FX/Niagara/Source/Niagara/Private/NiagaraParameterCollectionResolver.h
new file mode 100644
--- /dev/null
+++ b/Engine/Plugins/FX/Niagara/Source/Niagara/Private/NiagaraParameterCollectionResolver.h
@@ -0,0 +1,36 @@
+// Copyright Epic Games, Inc. All Rights Reserved.
+
+#pragma once
+
+#include "Containers/ArrayView.h"
+#include "UObject/ObjectPtr.h"
+
+class UNiagaraComponent;
+class UNiagaraParameterCollection;
+class UNiagaraParameterCollectionInstance;
+struct FNiagaraParameterStore;
+class UNiagaraScript;
+class UNiagaraSystem;
+class FNiagaraSystemSimulation;
+class FNiagaraWorldManager;
+class UWorld;
+
+struct FNiagaraParameterCollectionResolver
+{
+public:
+	explicit FNiagaraParameterCollectionResolver(UNiagaraComponent* InNiagaraComponent, UNiagaraSystem* InNiagaraSystem, UWorld* World);
+	explicit FNiagaraParameterCollectionResolver(UNiagaraComponent* InNiagaraComponent, UNiagaraSystem* InNiagaraSystem, FNiagaraWorldManager* InNiagaraWorldManager);
+
+	void BindCollections(TConstArrayView<TObjectPtr<UNiagaraParameterCollection>> ParameterCollections, FNiagaraParameterStore& DestinationStore) const;
+	void UnbindCollections(TConstArrayView<TObjectPtr<UNiagaraParameterCollection>> ParameterCollections, FNiagaraParameterStore& DestinationStore) const;
+
+	void BindCollections(UNiagaraScript* Script, FNiagaraParameterStore& DestinationStore) const;
+	void UnbindCollections(UNiagaraScript* Script, FNiagaraParameterStore& DestinationStore) const;
+
+	UNiagaraParameterCollectionInstance* ResolveInstance(UNiagaraParameterCollection* Collection) const;
+
+private:
+	UNiagaraComponent*		NiagaraComponent = nullptr;
+	UNiagaraSystem*			NiagaraSystem = nullptr;
+	FNiagaraWorldManager*	NiagaraWorldManager = nullptr;
+};
diff --git a/Engine/Plugins/FX/Niagara/Source/Niagara/Private/NiagaraSystemInstance.cpp b/Engine/Plugins/FX/Niagara/Source/Niagara/Private/NiagaraSystemInstance.cpp
--- a/Engine/Plugins/FX/Niagara/Source/Niagara/Private/NiagaraSystemInstance.cpp
+++ b/Engine/Plugins/FX/Niagara/Source/Niagara/Private/NiagaraSystemInstance.cpp
@@ -14,6 +14,7 @@
 #include "NiagaraSceneComponentUtils.h"
 #include "NiagaraStats.h"
 #include "NiagaraParameterCollection.h"
+#include "NiagaraParameterCollectionResolver.h"
 #include "NiagaraSystemSimulation.h"
 #include "NiagaraWorldManager.h"
 #include "NiagaraComponent.h"
@@ -486,8 +487,7 @@ void FNiagaraSystemInstance::SetSolo(bool bInSolo)
 		if (bInSolo)
 		{
 			TSharedPtr<FNiagaraSystemSimulation, ESPMode::ThreadSafe> NewSoloSim = MakeShared<FNiagaraSystemSimulation, ESPMode::ThreadSafe>();
-			NewSoloSim->Init(System, World, true, TG_MAX);
-
+			NewSoloSim->Init(System, World, true, TG_MAX, FNiagaraParameterCollectionResolver(Cast<UNiagaraComponent>(GetAttachComponent()), System, World));
 			NewSoloSim->TransferInstance(this);
 		}
 		else
@@ -978,9 +978,9 @@ void FNiagaraSystemInstance::ResetInternal(bool bResetSimulations)
 #endif
 }
 
-UNiagaraParameterCollectionInstance* FNiagaraSystemInstance::GetParameterCollectionInstance(UNiagaraParameterCollection* Collection)
+const FNiagaraParameterCollectionResolver FNiagaraSystemInstance::CreateParameterCollectionResolver() const
 {
-	return SystemSimulation->GetParameterCollectionInstance(Collection);
+	return FNiagaraParameterCollectionResolver(Cast<UNiagaraComponent>(GetAttachComponent()), System, World);
 }
 
 void FNiagaraSystemInstance::AdvanceSimulation(int32 TickCountToSimulate, float TickDeltaSeconds)
@@ -1091,7 +1091,7 @@ void FNiagaraSystemInstance::ReInitInternal()
 		if (!SystemSimulation.IsValid())
 		{
 			SystemSimulation = MakeShared<FNiagaraSystemSimulation, ESPMode::ThreadSafe>();
-			SystemSimulation->Init(System, World, true, TG_MAX);
+			SystemSimulation->Init(System, World, true, TG_MAX, FNiagaraParameterCollectionResolver(Cast<UNiagaraComponent>(GetAttachComponent()), System, World));
 		}
 	}
 	else
@@ -1186,22 +1186,6 @@ void FNiagaraSystemInstance::Cleanup()
 	ClearEventDataSets();
 }
 
-//Unsure on usage of this atm. Possibly useful in future.
-// void FNiagaraSystemInstance::RebindParameterCollection(UNiagaraParameterCollectionInstance* OldInstance, UNiagaraParameterCollectionInstance* NewInstance)
-// {
-// 	OldInstance->GetParameterStore().Unbind(&InstanceParameters);
-// 	NewInstance->GetParameterStore().Bind(&InstanceParameters);
-//
-// 	for (FNiagaraEmitterInstanceRef& Simulation : Emitters)
-// 	{
-// 		Simulation->RebindParameterCollection(OldInstance, NewInstance);
-// 	}
-//
-// 	//Have to re init the instance data for data interfaces.
-// 	//This is actually lots more work than absolutely needed in some cases so we can improve it a fair bit.
-// 	InitDataInterfaces();
-// }
-
 void FNiagaraSystemInstance::BindParameters()
 {
 	if (OverrideParameters != nullptr)
@@ -1249,6 +1233,40 @@ void FNiagaraSystemInstance::UnbindParameters(bool bFromComplete)
 	}
 }
 
+void FNiagaraSystemInstance::RebindParameterCollection(UNiagaraParameterCollectionInstance* OldCollection, UNiagaraParameterCollectionInstance* NewCollection)
+{
+	check(OldCollection && NewCollection);
+	check(OldCollection->GetParent() == NewCollection->GetParent());
+
+	WaitForConcurrentTickDoNotFinalize(true);
+
+	// Rebind the system simulation parameter collection
+	if (ensure(IsSolo()))
+	{
+		UNiagaraParameterCollection* CollectionParent = OldCollection->GetParent();
+		for (FNiagaraScriptExecutionContextBase* ScriptExecutionContext : { SystemSimulation->GetSpawnExecutionContext() , SystemSimulation->GetUpdateExecutionContext() })
+		{
+			if (ScriptExecutionContext->Script->GetCachedParameterCollectionReferences().Contains(CollectionParent))
+			{
+				OldCollection->GetParameterStore().Unbind(&ScriptExecutionContext->Parameters);
+				NewCollection->GetParameterStore().Bind(&ScriptExecutionContext->Parameters);
+			}
+		}
+	}
+
+ 	for (FNiagaraEmitterInstanceRef& Simulation : Emitters)
+ 	{
+ 		Simulation->RebindParameterCollection(OldCollection, NewCollection);
+ 	}
+
+	// If we have data interfaces inside the collections we may need to initialize them again
+	//-OPT: We could look to see if any of them are actually bound, this plays is safe
+	if (NewCollection->GetParameterStore().GetDataInterfaces().Num() > 0)
+	{
+		InitDataInterfaces();
+	}
+}
+
 void FNiagaraSystemInstance::BindToParameterStore(FNiagaraParameterStore& ParameterStore)
 {
 	if (ParameterStore.IsEmpty())
diff --git a/Engine/Plugins/FX/Niagara/Source/Niagara/Private/NiagaraSystemSimulation.cpp b/Engine/Plugins/FX/Niagara/Source/Niagara/Private/NiagaraSystemSimulation.cpp
--- a/Engine/Plugins/FX/Niagara/Source/Niagara/Private/NiagaraSystemSimulation.cpp
+++ b/Engine/Plugins/FX/Niagara/Source/Niagara/Private/NiagaraSystemSimulation.cpp
@@ -6,6 +6,7 @@
 #include "NiagaraModule.h"
 #include "NiagaraTypes.h"
 #include "NiagaraParameterCollection.h"
+#include "NiagaraParameterCollectionResolver.h"
 #include "NiagaraSimCache.h"
 #include "NiagaraSimulationTaskPriority.h"
 #include "NiagaraSystemEmitterState.h"
@@ -588,7 +589,7 @@ FNiagaraSystemSimulation::~FNiagaraSystemSimulation()
 	Destroy();
 }
 
-bool FNiagaraSystemSimulation::Init(UNiagaraSystem* InSystem, UWorld* InWorld, bool bInIsSolo, ETickingGroup InTickGroup)
+bool FNiagaraSystemSimulation::Init(UNiagaraSystem* InSystem, UWorld* InWorld, bool bInIsSolo, ETickingGroup InTickGroup, const FNiagaraParameterCollectionResolver& ParameterCollectionResolver)
 {
 	SCOPE_CYCLE_COUNTER(STAT_NiagaraSystemSim_Init);
 	System = InSystem;
@@ -652,36 +653,8 @@ bool FNiagaraSystemSimulation::Init(UNiagaraSystem* InSystem, UWorld* InWorld, b
 		{
 			//SCOPE_CYCLE_COUNTER(STAT_NiagaraSystemSim_Init_BindParams);
 
-			//Bind parameter collections.
-			auto BindParameterColleciton = [&](UNiagaraParameterCollection* Collection, FNiagaraParameterStore& DestStore)
-			{
-				if (Collection)
-				{
-					if (UNiagaraParameterCollectionInstance* CollectionInst = GetParameterCollectionInstance(Collection))
-					{
-						CollectionInst->GetParameterStore().Bind(&DestStore);
-					}
-					else
-					{
-						UE_LOG(LogNiagara, Error, TEXT("Attempting to bind system simulation to a null parameter collection instance | Collection: %s | System: %s |"), *Collection->GetPathName(), *System->GetPathName());
-					}
-				}
-				else
-				{
-					UE_LOG(LogNiagara, Error, TEXT("Attempting to bind system simulation to a null parameter collection | System: %s |"), *System->GetPathName());
-				}
-			};
-
-			//Bind parameter collections.
-			for (UNiagaraParameterCollection* Collection : SpawnScript->GetCachedParameterCollectionReferences())
-			{
-				BindParameterColleciton(Collection, SpawnExecContext->Parameters);
-			}
-			for (UNiagaraParameterCollection* Collection : UpdateScript->GetCachedParameterCollectionReferences())
-			{
-				BindParameterColleciton(Collection, UpdateExecContext->Parameters);
-			}
-
+			ParameterCollectionResolver.BindCollections(SpawnScript, SpawnExecContext->Parameters);
+			ParameterCollectionResolver.BindCollections(UpdateScript, UpdateExecContext->Parameters);
 			SpawnScript->RapidIterationParameters.Bind(&SpawnExecContext->Parameters);
 			UpdateScript->RapidIterationParameters.Bind(&UpdateExecContext->Parameters);
 		}
@@ -737,27 +710,6 @@ void FNiagaraSystemSimulation::Destroy()
 	World = nullptr;
 }
 
-UNiagaraParameterCollectionInstance* FNiagaraSystemSimulation::GetParameterCollectionInstance(UNiagaraParameterCollection* Collection)
-{
-	UNiagaraParameterCollectionInstance* Ret = nullptr;
-
-	if (System)
-	{
-		Ret = System->GetParameterCollectionOverride(Collection);
-	}
-
-	//If no explicit override from the system, just get the current instance set on the world.
-	if (!Ret)
-	{
-		if (FNiagaraWorldManager* WorldMan = FNiagaraWorldManager::Get(World))
-		{
-			Ret = WorldMan->GetParameterCollection(Collection);
-		}
-	}
-
-	return Ret;
-}
-
 void FNiagaraSystemSimulation::TransferInstance(FNiagaraSystemInstance* Instance)
 {
 	check(Instance);
@@ -2173,7 +2125,7 @@ void FNiagaraSystemSimulation::UpdateSystemInstances(FNiagaraSystemSimulationTic
 		{
 			UE_LOG(LogNiagara, Log, TEXT("=== Updated %d Systems ==="), NumInstances);
 			DestinationData.Dump(0, NumInstances, TEXT("System Data - Post Update"));
-			UpdateInstanceParameterDataSet.GetCurrentDataChecked().Dump(0, NumInstances, TEXT("Update Instance Paramter Data"));
+			UpdateInstanceParameterDataSet.GetCurrentDataChecked().Dump(0, NumInstances, TEXT("Update Instance Parameter Data"));
 		}
 
 		//Also run the update script on the newly spawned systems too.
@@ -2201,7 +2153,7 @@ void FNiagaraSystemSimulation::UpdateSystemInstances(FNiagaraSystemSimulationTic
 			{
 				UE_LOG(LogNiagara, Log, TEXT("=== Spawn Updated %d Systems ==="), SpawnNum);
 				DestinationData.Dump(OrigNum, SpawnNum, TEXT("System Data - Post Update (new systems)"));
-				UpdateInstanceParameterDataSet.GetCurrentDataChecked().Dump(OrigNum, SpawnNum, TEXT("Update Instance Paramter Data (new systems)"));
+				UpdateInstanceParameterDataSet.GetCurrentDataChecked().Dump(OrigNum, SpawnNum, TEXT("Update Instance Parameter Data (new systems)"));
 			}
 		}
 
diff --git a/Engine/Plugins/FX/Niagara/Source/Niagara/Private/NiagaraWorldManager.cpp b/Engine/Plugins/FX/Niagara/Source/Niagara/Private/NiagaraWorldManager.cpp
--- a/Engine/Plugins/FX/Niagara/Source/Niagara/Private/NiagaraWorldManager.cpp
+++ b/Engine/Plugins/FX/Niagara/Source/Niagara/Private/NiagaraWorldManager.cpp
@@ -9,6 +9,7 @@
 #include "NiagaraSystemSimulation.h"
 #include "NiagaraTypes.h"
 #include "NiagaraParameterCollection.h"
+#include "NiagaraParameterCollectionResolver.h"
 #include "NiagaraSystemInstance.h"
 #include "GameFramework/PlayerController.h"
 #include "Engine/GameViewportClient.h"
@@ -731,7 +732,7 @@ FNiagaraSystemSimulationRef FNiagaraWorldManager::GetSystemSimulation(ETickingGr
 	
 	FNiagaraSystemSimulationRef Sim = MakeShared<FNiagaraSystemSimulation, ESPMode::ThreadSafe>();
 	SystemSimulations[ActualTickGroup][int(TickPass)].Add(System, Sim);
-	Sim->Init(System, World, false, TickGroup);
+	Sim->Init(System, World, false, TickGroup, FNiagaraParameterCollectionResolver(nullptr, System, this));
 
 	// Old path which will block simulation at tick boundaries
 	if (bUseTickPasses == false)
diff --git a/Engine/Plugins/FX/Niagara/Source/Niagara/Private/Stateless/NiagaraStatelessEmitterInstance.cpp b/Engine/Plugins/FX/Niagara/Source/Niagara/Private/Stateless/NiagaraStatelessEmitterInstance.cpp
--- a/Engine/Plugins/FX/Niagara/Source/Niagara/Private/Stateless/NiagaraStatelessEmitterInstance.cpp
+++ b/Engine/Plugins/FX/Niagara/Source/Niagara/Private/Stateless/NiagaraStatelessEmitterInstance.cpp
@@ -7,6 +7,7 @@
 #include "Stateless/NiagaraStatelessExpression.h"
 #include "Stateless/NiagaraStatelessParticleSimContext.h"
 #include "NiagaraGpuComputeDispatchInterface.h"
+#include "NiagaraParameterCollectionResolver.h"
 #include "NiagaraSystem.h"
 #include "NiagaraSystemInstance.h"
 
@@ -294,13 +295,8 @@ void FNiagaraStatelessEmitterInstance::BindParameters(bool bExternalOnly)
 		{
 			ParentSystemInstance->BindToParameterStore(RendererBindings);
 
-			for (UNiagaraParameterCollection* ParameterCollection : EmitterData->BoundParameterCollections)
-			{
-				if (UNiagaraParameterCollectionInstance* ParameterCollectionInstance = ParentSystemInstance->GetParameterCollectionInstance(ParameterCollection))
-				{
-					ParameterCollectionInstance->GetParameterStore().Bind(&RendererBindings);
-				}
-			}
+			const FNiagaraParameterCollectionResolver ParameterCollectionResolver = ParentSystemInstance->CreateParameterCollectionResolver();
+			ParameterCollectionResolver.BindCollections(MakeArrayView(EmitterData->BoundParameterCollections), RendererBindings);
 		}
 	}
 }
@@ -313,17 +309,28 @@ void FNiagaraStatelessEmitterInstance::UnbindParameters(bool bExternalOnly)
 		{
 			ParentSystemInstance->UnbindFromParameterStore(RendererBindings);
 
-			for (UNiagaraParameterCollection* ParameterCollection : EmitterData->BoundParameterCollections)
-			{
-				if (UNiagaraParameterCollectionInstance* ParameterCollectionInstance = ParentSystemInstance->GetParameterCollectionInstance(ParameterCollection))
-				{
-					ParameterCollectionInstance->GetParameterStore().Unbind(&RendererBindings);
-				}
-			}
+			const FNiagaraParameterCollectionResolver ParameterCollectionResolver = ParentSystemInstance->CreateParameterCollectionResolver();
+			ParameterCollectionResolver.UnbindCollections(MakeArrayView(EmitterData->BoundParameterCollections), RendererBindings);
 		}
 	}
 }
 
+void FNiagaraStatelessEmitterInstance::RebindParameterCollection(UNiagaraParameterCollectionInstance* OldCollection, UNiagaraParameterCollectionInstance* NewCollection)
+{
+	if (RendererBindings.IsEmpty() || !ParentSystemInstance)
+	{
+		return;
+	}
+
+	if (!EmitterData->BoundParameterCollections.Contains(OldCollection->GetParent()))
+	{
+		return;
+	}
+
+	OldCollection->GetParameterStore().Unbind(&RendererBindings);
+	NewCollection->GetParameterStore().Bind(&RendererBindings);
+}
+
 bool FNiagaraStatelessEmitterInstance::ShouldTick() const
 {
 	return InternalExecutionState <= ENiagaraExecutionState::Inactive;
diff --git a/Engine/Plugins/FX/Niagara/Source/Niagara/Public/NiagaraComponent.h b/Engine/Plugins/FX/Niagara/Source/Niagara/Public/NiagaraComponent.h
--- a/Engine/Plugins/FX/Niagara/Source/Niagara/Public/NiagaraComponent.h
+++ b/Engine/Plugins/FX/Niagara/Source/Niagara/Public/NiagaraComponent.h
@@ -603,6 +603,25 @@ public:
 	UFUNCTION(BlueprintCallable, Category = Niagara, meta = (DisplayName = "Get Niagara Variable (Matrix)", Keywords = "user parameter variable matrix"))
 	[[nodiscard]] NIAGARA_API FMatrix GetVariableMatrix(FName InVariableName, bool& bIsValid) const;
 
+	/**
+	Override the parameter collection to use.
+	The order for resolving collections is Component -> System -> World.
+	Note that when using a collection override the simulation may be forced in solo mode which is less performant.
+	*/
+	UFUNCTION(Experimental, BlueprintCallable, Category = Niagara)
+	NIAGARA_API void SetParameterCollectionOverride(UNiagaraParameterCollectionInstance* ParameterCollection);
+
+	/** Remove the override collection, this will revert back to the default instance. */
+	UFUNCTION(Experimental, BlueprintCallable, Category = Niagara)
+	NIAGARA_API void RemoveParameterCollectionOverride(UNiagaraParameterCollectionInstance* OverrideCollection);
+
+	/** Clear all parameter collection overrides. */
+	UFUNCTION(Experimental, BlueprintCallable, Category = Niagara)
+	NIAGARA_API void ClearParameterCollectionOverrides();
+
+	/** Get the current parameter collection override if one exists. */
+	UNiagaraParameterCollectionInstance* GetParameterCollectionOverride(UNiagaraParameterCollection* ParameterCollection) const;
+
 	/** Resets the System to it's initial pre-simulated state. */
 	UFUNCTION(BlueprintCallable, Category = Niagara, meta = (DisplayName = "Reset System"))
 	NIAGARA_API void ResetSystem();
@@ -913,7 +932,6 @@ private:
 	/** Restore relative transform from auto attachment and optionally detach from parent (regardless of whether it was an auto attachment). */
 	NIAGARA_API void CancelAutoAttachment(bool bDetachFromParent);
 
-
 	/** Saved relative transform before auto attachment. Used during detachment to restore the transform if we had automatically attached. */
 	FVector SavedAutoAttachRelativeLocation;
 	FRotator SavedAutoAttachRelativeRotation;
@@ -952,6 +970,9 @@ private:
 	UPROPERTY(transient)
 	TObjectPtr<class UNiagaraCullProxyComponent> CullProxy;
 
+	UPROPERTY(transient)
+	TArray<TObjectPtr<UNiagaraParameterCollectionInstance>> ParameterCollectionOverrides;
+
 	NIAGARA_API void CreateCullProxy(bool bForce = false);
 	NIAGARA_API void DestroyCullProxy();
 
diff --git a/Engine/Plugins/FX/Niagara/Source/Niagara/Public/NiagaraFunctionLibrary.h b/Engine/Plugins/FX/Niagara/Source/Niagara/Public/NiagaraFunctionLibrary.h
--- a/Engine/Plugins/FX/Niagara/Source/Niagara/Public/NiagaraFunctionLibrary.h
+++ b/Engine/Plugins/FX/Niagara/Source/Niagara/Public/NiagaraFunctionLibrary.h
@@ -170,9 +170,16 @@ public:
 // 	UFUNCTION(BlueprintCallable, Category = Niagara, meta = (Keywords = "niagara System", UnsafeDuringActorConstruction = "true"))
 // 	static void SetUpdateScriptConstant(UNiagaraComponent* Component, FName EmitterName, FName ConstantName, FVector Value);
 
-	UFUNCTION(BlueprintCallable, Category = Niagara, meta = (Keywords = "niagara parameter collection", WorldContext = "WorldContextObject"))
+	UFUNCTION(BlueprintCallable, Category = Niagara, meta = (Keywords = "niagara parameter collection", WorldContext = "WorldContextObject", UnsafeDuringActorConstruction = "true"))
 	static NIAGARA_API UNiagaraParameterCollectionInstance* GetNiagaraParameterCollection(UObject* WorldContextObject, UNiagaraParameterCollection* Collection);
 
+	/**
+	Create a parameter collection instance for the provided parameter collection.
+	This can return null if not in a game world or your passed in a null collection.
+	*/
+	UFUNCTION(BlueprintCallable, Category = Niagara, meta = (Keywords = "niagara parameter collection", WorldContext = "WorldContextObject", UnsafeDuringActorConstruction = "true"))
+	static NIAGARA_API UNiagaraParameterCollectionInstance* CreateNiagaraParameterCollectionInstance(UObject* WorldContextObject, UNiagaraParameterCollection* ParentCollection);
+
 	static NIAGARA_API const TArray<FNiagaraFunctionSignature>& GetVectorVMFastPathOps(bool bIgnoreConsoleVariable = false);
 	static NIAGARA_API bool DefineFunctionHLSL(const FNiagaraFunctionSignature& FunctionSignature, FString& HlslOutput);
 
diff --git a/Engine/Plugins/FX/Niagara/Source/Niagara/Public/NiagaraSystemInstance.h b/Engine/Plugins/FX/Niagara/Source/Niagara/Public/NiagaraSystemInstance.h
--- a/Engine/Plugins/FX/Niagara/Source/Niagara/Public/NiagaraSystemInstance.h
+++ b/Engine/Plugins/FX/Niagara/Source/Niagara/Public/NiagaraSystemInstance.h
@@ -18,6 +18,7 @@ class FNiagaraSystemInstance;
 class FNiagaraSystemSimulation;
 class FNiagaraGpuComputeDispatchInterface;
 class FNiagaraGPUSystemTick;
+struct FNiagaraParameterCollectionResolver;
 class UNiagaraSimCache;
 class FNiagaraSystemGpuComputeProxy;
 struct FNiagaraParameterStore;
@@ -149,9 +150,9 @@ public:
 
 	NIAGARA_API UActorComponent* GetPrereqComponent() const;
 
-	//void RebindParameterCollection(UNiagaraParameterCollectionInstance* OldInstance, UNiagaraParameterCollectionInstance* NewInstance);
 	NIAGARA_API void BindParameters();
 	NIAGARA_API void UnbindParameters(bool bFromComplete = false);
+	void RebindParameterCollection(UNiagaraParameterCollectionInstance* OldCollection, UNiagaraParameterCollectionInstance* NewCollection);
 
 	// Bindings Override Parameters / Instance Parameters / System Simulation Parameters to the provded parameter store
 	// I.e. binds all relevant parameters from us and parent, does not do children (i.e. emitters)
@@ -386,7 +387,7 @@ public:
 
 	NIAGARA_API bool IsReadyToRun() const;
 
-	NIAGARA_API UNiagaraParameterCollectionInstance* GetParameterCollectionInstance(UNiagaraParameterCollection* Collection);
+	const FNiagaraParameterCollectionResolver CreateParameterCollectionResolver() const;
 
 	/** 
 	Manually advances this system's simulation by the specified number of ticks and tick delta. 
diff --git a/Engine/Plugins/FX/Niagara/Source/Niagara/Public/NiagaraSystemInstanceController.h b/Engine/Plugins/FX/Niagara/Source/Niagara/Public/NiagaraSystemInstanceController.h
--- a/Engine/Plugins/FX/Niagara/Source/Niagara/Public/NiagaraSystemInstanceController.h
+++ b/Engine/Plugins/FX/Niagara/Source/Niagara/Public/NiagaraSystemInstanceController.h
@@ -153,6 +153,8 @@ public:
 	NIAGARA_SYSTEM_INSTANCE_CONTROLLER_SHIM(WaitForConcurrentTickAndFinalize,)
 	NIAGARA_SYSTEM_INSTANCE_CONTROLLER_SHIM(OnSimulationDestroyed,)
 	NIAGARA_SYSTEM_INSTANCE_CONTROLLER_SHIM(GetAttachComponent,)
+	NIAGARA_SYSTEM_INSTANCE_CONTROLLER_SHIM(CreateParameterCollectionResolver, )
+	NIAGARA_SYSTEM_INSTANCE_CONTROLLER_SHIM(RebindParameterCollection,)
 	#if WITH_EDITOR
 		NIAGARA_SYSTEM_INSTANCE_CONTROLLER_SHIM(HandleNeedsUIResync,)
 	#endif
diff --git a/Engine/Plugins/FX/Niagara/Source/Niagara/Public/NiagaraSystemSimulation.h b/Engine/Plugins/FX/Niagara/Source/Niagara/Public/NiagaraSystemSimulation.h
--- a/Engine/Plugins/FX/Niagara/Source/Niagara/Public/NiagaraSystemSimulation.h
+++ b/Engine/Plugins/FX/Niagara/Source/Niagara/Public/NiagaraSystemSimulation.h
@@ -10,10 +10,9 @@
 
 class UNiagaraEffectType;
 class UWorld;
-class UNiagaraParameterCollection;
-class UNiagaraParameterCollectionInstance;
 class FNiagaraSystemSimulation;
 class FNiagaraGpuComputeDispatchInterface;
+struct FNiagaraParameterCollectionResolver;
 
 using FNiagaraSystemSimulationPtr = TSharedPtr<FNiagaraSystemSimulation, ESPMode::ThreadSafe>;
 
@@ -254,7 +253,7 @@ class FNiagaraSystemSimulation : public TSharedFromThis<FNiagaraSystemSimulation
 public:
 	FNiagaraSystemSimulation();
 	~FNiagaraSystemSimulation();
-	bool Init(UNiagaraSystem* InSystem, UWorld* InWorld, bool bInIsSolo, ETickingGroup TickGroup);
+	bool Init(UNiagaraSystem* InSystem, UWorld* InWorld, bool bInIsSolo, ETickingGroup InTickGroup, const FNiagaraParameterCollectionResolver& ParameterCollectionResolver);
 	void Destroy();
 
 	bool IsValid() const { return bCanExecute && World != nullptr && ::IsValid(System); }
@@ -292,8 +291,6 @@ public:
 
 	inline UNiagaraSystem* GetSystem() const { return System; }
 
-	UNiagaraParameterCollectionInstance* GetParameterCollectionInstance(UNiagaraParameterCollection* Collection);
-
 	/** Transfers a system instance from the current simulation into this one. */
 	void TransferInstance(FNiagaraSystemInstance* SystemInst);
 
