Niagara - Fix an issue with the translator BuildMissingDefaults where it would generate errors or invalid code when generating defaults in spawn for later stages with different parameter maps.

#Synced-CL 48948473

diff --git a/Engine/Plugins/FX/Niagara/Source/NiagaraEditor/Private/NiagaraHlslTranslator.cpp b/Engine/Plugins/FX/Niagara/Source/NiagaraEditor/Private/NiagaraHlslTranslator.cpp
--- a/Engine/Plugins/FX/Niagara/Source/NiagaraEditor/Private/NiagaraHlslTranslator.cpp
+++ b/Engine/Plugins/FX/Niagara/Source/NiagaraEditor/Private/NiagaraHlslTranslator.cpp
@@ -81,6 +81,14 @@ DECLARE_CYCLE_STAT(TEXT("Niagara - HlslTranslator - GenerateFunctionSignature_Fi
 
 #define NIAGARA_SCOPE_CYCLE_COUNTER(x) //SCOPE_CYCLE_COUNTER(x)
 
+static int GNiagaraCompilationUseNewMissingDefaultTraversal = 1;
+static FAutoConsoleVariableRef CVarNiagaraCompilationUseNewMissingDefaultTraversal(
+	TEXT("fx.Niagara.Compilation.UseNewMissingDefaultTraversal"),
+	GNiagaraCompilationUseNewMissingDefaultTraversal,
+	TEXT("Whether or not to use the new missing default traversal when compiling scripts."),
+	ECVF_Default
+);
+
 // not pretty. TODO: refactor
 // this will be called via a delegate from UNiagaraScript's cache for cook function,
 // because editor tickables aren't ticked during cooking
@@ -1027,7 +1035,27 @@ void TNiagaraHlslTranslator<GraphBridge>::BuildMissingDefaults()
 				}
 			}
 
-			HandleParameterRead(ActiveStageIdx, Var, DefaultPin, DefaultPinNode, OutputChunkId, DefaultMode, DefaultBinding, !bWriteToParamMapEntries, true);
+			bool bIsKnownMissingDefault = ParamMapHistories[ActiveStageIdx].Variables.Contains(Var);
+			if (bIsKnownMissingDefault ||
+				GNiagaraCompilationUseNewMissingDefaultTraversal == 0)
+			{
+				HandleParameterRead(ActiveStageIdx, Var, DefaultPin, DefaultPinNode, OutputChunkId, DefaultMode, DefaultBinding, !bWriteToParamMapEntries, true);
+			}
+			else
+			{
+				// Missing defaults that are unknown to the current parameter map need to be handled differently since their parameter history is unrelated to the
+				// current stage.  Calling parameter map read in this case can result in calling compile on a pin with mismatched parameter map history and graph
+				// pins which can result in compile errors or incorrect code generation.  If the unknown default can not be resolved we use the struct default.
+				int32 MissingDefaultIdx = HandleReadUnknownMissingDefault(Var, DefaultPin);
+				if (MissingDefaultIdx == INDEX_NONE)
+				{
+					FNiagaraVariable DefaultVar = Var;
+					FNiagaraEditorUtilities::ResetVariableToDefaultValue(DefaultVar);
+					MissingDefaultIdx = GetConstant(DefaultVar);
+				}
+				FString ParameterMapInstanceName = GetParameterMapInstanceName(ActiveStageIdx);
+				AddBodyChunk(ParameterMapInstanceName + TEXT(".") + BuildSanitizedSymbolName(Var.GetName()), TEXT("{0}"), Var.GetType(), MissingDefaultIdx, false);
+			}
 		}
 
 		DeferredVariablesMissingDefault.Empty();
@@ -1054,6 +1082,92 @@ void TNiagaraHlslTranslator<GraphBridge>::BuildMissingDefaults()
 	AddBodyComment(TEXT("// End HandleMissingDefaultValues\n\n"));
 }
 
+template<typename GraphBridge>
+int32 TNiagaraHlslTranslator<GraphBridge>::HandleReadUnknownMissingDefault(const FNiagaraVariable& DefaultVar, const FInputPin* DefaultPin)
+{
+	// Validate inputs
+	if (DefaultPin == nullptr)
+	{
+		return INDEX_NONE;
+	}
+	const FNode* DefaultPinNode = GraphBridge::GetOwningNode(DefaultPin);
+	if (DefaultPinNode == nullptr)
+	{
+		return INDEX_NONE;
+	}
+	const FGraph* DefaultPinGraph = GraphBridge::GetOwningGraph(DefaultPinNode);
+	if (DefaultPinGraph == nullptr)
+	{
+		return INDEX_NONE;
+	}
+
+	// Get the default mode
+	const FNiagaraVariableBase* AliasedVar = UniqueVarToAliasedVar.Find(DefaultVar);
+	FNiagaraScriptVariableBinding DefaultBinding;
+	TOptional<ENiagaraDefaultMode> DefaultMode = GraphBridge::GetGraphDefaultMode(DefaultPinGraph, AliasedVar ? *AliasedVar : FNiagaraVariableBase(DefaultVar), DefaultBinding);
+
+	int32 DefaultChunkIndex = INDEX_NONE;
+	if (DefaultMode == ENiagaraDefaultMode::Binding && DefaultBinding.IsValid())
+	{
+		// Bound defaults that can be handled are either local to the script being translated which will be in the UniqueVarToChunk map, or
+		// are in an external constant namespace which can be handled safely with parameter map read.
+		FNiagaraVariable BoundVar = FNiagaraVariable(DefaultVar.GetType(), DefaultBinding.GetName());
+		int32* ExistingDefault = UniqueVarToChunk.Find(BoundVar);
+		if (ExistingDefault != nullptr)
+		{
+			DefaultChunkIndex = *ExistingDefault;
+		}
+		else if (FNiagaraParameterUtilities::IsExternalConstantNamespace(BoundVar, CompileOptions.TargetUsage, CompileOptions.GetTargetUsageBitmask()))
+		{
+			HandleParameterRead(ActiveStageIdx, BoundVar, nullptr, nullptr, DefaultChunkIndex);
+		}
+	}
+	else if (DefaultMode == ENiagaraDefaultMode::Custom)
+	{
+		// For custom defaults we only support linked parameter map gets.  First we handle the linked variable as if it was a
+		// default binding, and if that doesn't work we recursively evaluate the linked default pin if it exists.
+		const FOutputPin* LinkedOutputPin = GraphBridge::GetLinkedOutputPin(DefaultPin);
+		if (LinkedOutputPin != nullptr)
+		{
+			const FNode* LinkedNode = GraphBridge::GetOwningNode(LinkedOutputPin);
+			if (LinkedNode != nullptr)
+			{
+				const FParamMapGetNode* LinkedGetNode = GraphBridge::AsParamMapGetNode(LinkedNode);
+				if (LinkedGetNode != nullptr)
+				{
+					FNiagaraVariable LinkedVar = GraphBridge::GetPinVariable(LinkedOutputPin, true, ENiagaraStructConversion::Simulation);
+					int32* ExistingDefault = UniqueVarToChunk.Find(LinkedVar);
+					if (ExistingDefault != nullptr)
+					{
+						DefaultChunkIndex = *ExistingDefault;
+					}
+					else if (FNiagaraParameterUtilities::IsExternalConstantNamespace(LinkedVar, CompileOptions.TargetUsage, CompileOptions.GetTargetUsageBitmask()))
+					{
+						HandleParameterRead(ActiveStageIdx, LinkedVar, nullptr, nullptr, DefaultChunkIndex);
+					}
+					else
+					{
+						const FInputPin* LinkedDefaultPin = GraphBridge::GetDefaultPin(LinkedGetNode, LinkedOutputPin);
+						if (LinkedDefaultPin != nullptr)
+						{
+							DefaultChunkIndex = HandleReadUnknownMissingDefault(LinkedVar, LinkedDefaultPin);
+						}
+					}
+				}
+			}
+		}
+	}
+
+	if (DefaultChunkIndex == INDEX_NONE)
+	{
+		// If a default was not found through a binding or through custom traversal, check if the pin itself has a default.
+		FNiagaraVariable PinVariable = GraphBridge::GetPinVariable(DefaultPin, true, ENiagaraStructConversion::Simulation);
+		DefaultChunkIndex = PinVariable.IsDataAllocated() ? GetConstant(PinVariable) : INDEX_NONE;
+	}
+
+	return DefaultChunkIndex;
+}
+
 template<typename GraphBridge>
 FString TNiagaraHlslTranslator<GraphBridge>::BuildParameterMapHlslDefinitions(TArray<FNiagaraVariable>& PrimaryDataSetOutputEntries)
 {
diff --git a/Engine/Plugins/FX/Niagara/Source/NiagaraEditor/Private/NiagaraHlslTranslator.h b/Engine/Plugins/FX/Niagara/Source/NiagaraEditor/Private/NiagaraHlslTranslator.h
--- a/Engine/Plugins/FX/Niagara/Source/NiagaraEditor/Private/NiagaraHlslTranslator.h
+++ b/Engine/Plugins/FX/Niagara/Source/NiagaraEditor/Private/NiagaraHlslTranslator.h
@@ -722,6 +722,7 @@ protected:
 	void HandleParameterRead(int32 ParamMapHistoryIdx, const FNiagaraVariable& Var, const FInputPin* DefaultPin, const FNode* ErrorNode, int32& OutputChunkId, TOptional<ENiagaraDefaultMode> DefaultMode = TOptional<ENiagaraDefaultMode>(), TOptional<FNiagaraScriptVariableBinding> DefaultBinding = TOptional<FNiagaraScriptVariableBinding>(), bool bTreatAsUnknownParameterMap = false, bool bIgnoreDefaultSetFirst = false);
 	FString BuildParameterMapHlslDefinitions(TArray<FNiagaraVariable>& PrimaryDataSetOutputEntries);
 	void BuildMissingDefaults();
+	int32 HandleReadUnknownMissingDefault(const FNiagaraVariable& DefaultVar, const FInputPin* DefaultPin);
 	void FinalResolveNamespacedTokens(const FString& ParameterMapInstanceNamespace, TArray<FString>& Tokens, TArray<FString>& ValidChildNamespaces, FParamMapHistoryBuilder& Builder, TArray<FNiagaraVariable>& UniqueParameterMapEntriesAliasesIntact, TArray<FNiagaraVariable>& UniqueParameterMapEntries, int32 ParamMapHistoryIdx, const FNode* InNodeForErrorReporting);
 	void FindConstantValue(int32 InputCompileResult, const FNiagaraTypeDefinition& TypeDef, FString& Value, FNiagaraVariable& Variable);
 
