- NPC override test

#virtualized
#Synced-CL 49705804

diff --git a/Engine/Plugins/FX/Niagara/Source/Niagara/Classes/NiagaraParameterCollection.h b/Engine/Plugins/FX/Niagara/Source/Niagara/Classes/NiagaraParameterCollection.h
--- a/Engine/Plugins/FX/Niagara/Source/Niagara/Classes/NiagaraParameterCollection.h
+++ b/Engine/Plugins/FX/Niagara/Source/Niagara/Classes/NiagaraParameterCollection.h
@@ -15,7 +15,7 @@ class UNiagaraParameterCollection;
  * Can be used to override selected parameters from a Niagara parameter collection with another value.
  * The values in the parameter collection instance can be set from Blueprint or C++, same as the regular parameter collection. 
  */
-UCLASS(MinimalAPI)
+UCLASS(BlueprintType, MinimalAPI)
 class UNiagaraParameterCollectionInstance : public UObject
 {
 	GENERATED_UCLASS_BODY()
