- Change default sim cache tree view to use namespaces
- This makes it easier to collapse parameter from modules / initial / previous for example

#Synced-CL 49733395

diff --git a/Engine/Plugins/FX/Niagara/Source/NiagaraEditor/Private/ViewModels/NiagaraSimCacheViewModel.cpp b/Engine/Plugins/FX/Niagara/Source/NiagaraEditor/Private/ViewModels/NiagaraSimCacheViewModel.cpp
--- a/Engine/Plugins/FX/Niagara/Source/NiagaraEditor/Private/ViewModels/NiagaraSimCacheViewModel.cpp
+++ b/Engine/Plugins/FX/Niagara/Source/NiagaraEditor/Private/ViewModels/NiagaraSimCacheViewModel.cpp
@@ -19,6 +19,14 @@
 
 namespace NiagaraSimCacheViewModelPrivate
 {
+	bool bNamespaceGroupingEnabled = true;
+	static FAutoConsoleVariableRef CVarNamespaceGroupingEnabled(
+		TEXT("fx.Niagara.SimCacheViewer.NamespaceGroupingEnabled"),
+		bNamespaceGroupingEnabled,
+		TEXT("When enabled namespaces will be part of the attribute tree view."),
+		ECVF_Default
+	);
+
 	TOptional<bool> FindComponentFilteredState(FName ComponentName, TArray<FNiagaraSimCacheViewModel::FComponentInfo> Components)
 	{
 		for (const FNiagaraSimCacheViewModel::FComponentInfo& Component : Components)
@@ -30,7 +38,156 @@ namespace NiagaraSimCacheViewModelPrivate
 		}
 		return {};
 	}
-}
+
+	void TraverseStructItems(const FString& StructPath, const FNiagaraTypeDefinition& StructTypeDef, const TFunction<void(const FString&, const FNiagaraTypeDefinition&)>& ItemCallback)
+	{
+		UScriptStruct* Struct = StructTypeDef.GetScriptStruct();
+		if (Struct == nullptr)
+		{
+			return;
+		}
+
+		for (TFieldIterator<FProperty> PropertyIt(Struct, EFieldIteratorFlags::IncludeSuper); PropertyIt; ++PropertyIt)
+		{
+			const FProperty* Property = *PropertyIt;
+			const FString PropertyPath = StructPath + "." + Property->GetName();
+
+			if (const FStructProperty* StructProperty = CastField<FStructProperty>(Property))
+			{
+				if (UScriptStruct* FriendlyStruct = FNiagaraTypeHelper::FindNiagaraFriendlyTopLevelStruct(StructProperty->Struct, ENiagaraStructConversion::Simulation))
+				{
+					TraverseStructItems(PropertyPath, FNiagaraTypeDefinition(FriendlyStruct), ItemCallback);
+				}
+			}
+			if (Property->IsA(FIntProperty::StaticClass()))
+			{
+				ItemCallback(PropertyPath, FNiagaraTypeDefinition::GetIntDef());
+			}
+			else if (Property->IsA(FFloatProperty::StaticClass()))
+			{
+				ItemCallback(PropertyPath, FNiagaraTypeDefinition::GetFloatDef());
+			}
+			else if (Property->IsA(FBoolProperty::StaticClass()))
+			{
+				ItemCallback(PropertyPath, FNiagaraTypeDefinition::GetBoolDef());
+			}
+		}
+	}
+
+	namespace NamespaceGroupedTree
+	{
+		void AddItemToTreeInternal(TWeakPtr<SNiagaraSimCacheTreeView> WeakOwnerTreeView, TArray<TSharedRef<FNiagaraSimCacheTreeItem>>& ComponentTreeItems, FString ComponentPath, const FNiagaraTypeDefinition ComponentType)
+		{
+			TSharedPtr<FNiagaraSimCacheTreeItem> ParentItem = ComponentTreeItems[0];
+
+			FString CurrentPath;
+
+			TArray<FString> ComponentPathParts;
+			ComponentPath.ParseIntoArray(ComponentPathParts, TEXT("."));
+
+			for (int32 i = 0; i < ComponentPathParts.Num(); ++i)
+			{
+				const FString PathPart = ComponentPathParts[i];
+				if (CurrentPath.Len() > 0)
+				{
+					CurrentPath.AppendChar('.');
+				}
+				CurrentPath.Append(PathPart);
+
+				TSharedRef<FNiagaraSimCacheTreeItem>* CurrentPathItem = ComponentTreeItems.FindByPredicate(
+					[&CurrentPath](const TSharedPtr<FNiagaraSimCacheTreeItem>& Item)
+					{
+						return Item->GetFilterName() == CurrentPath;
+					}
+				);
+
+				if (CurrentPathItem == nullptr)
+				{
+					TSharedRef<FNiagaraSimCacheComponentTreeItem> NewItem = MakeShared<FNiagaraSimCacheComponentTreeItem>(WeakOwnerTreeView);
+					NewItem->SetDisplayName(FText::FromString(PathPart));
+					NewItem->SetFilterName(CurrentPath);
+
+					const bool bIsLastPart = i == ComponentPathParts.Num() - 1;
+					if (bIsLastPart)
+					{
+						NewItem->TypeDef = ComponentType;
+					}
+
+					ComponentTreeItems.Emplace(NewItem);
+					CurrentPathItem = &ComponentTreeItems.Last();
+				}
+				ParentItem->AddChild(*CurrentPathItem);
+				ParentItem = *CurrentPathItem;
+			}
+		}
+
+		void AddItemToTree(TWeakPtr<SNiagaraSimCacheTreeView> WeakOwnerTreeView, TArray<TSharedRef<FNiagaraSimCacheTreeItem>>& ComponentTreeItems, const FNiagaraSimCacheVariable& Variable)
+		{
+			const FNiagaraTypeDefinition VariableTypeDef = Variable.Variable.GetType();
+			const FString VariableName = Variable.Variable.GetName().ToString();
+
+			AddItemToTreeInternal(WeakOwnerTreeView, ComponentTreeItems, VariableName, VariableTypeDef);
+			if (!VariableTypeDef.IsEnum() && !FNiagaraTypeDefinition::IsScalarDefinition(VariableTypeDef))
+			{
+				TraverseStructItems(
+					VariableName,
+					VariableTypeDef,
+					[&WeakOwnerTreeView, &ComponentTreeItems](const FString& StructPath, const FNiagaraTypeDefinition& StructTypeDef)
+					{
+						AddItemToTreeInternal(WeakOwnerTreeView, ComponentTreeItems, StructPath, StructTypeDef);
+					}
+				);
+			}
+		}
+	} //namespace NamespaceGroupedTree
+
+	namespace TypeGroupedTree
+	{
+		TSharedRef<FNiagaraSimCacheTreeItem> AddItemToTreeInternal(TWeakPtr<SNiagaraSimCacheTreeView> WeakOwnerTreeView, const TArray<TSharedRef<FNiagaraSimCacheTreeItem>>& ComponentTreeItems, FString ComponentPath, const FNiagaraTypeDefinition ComponentType)
+		{
+			TSharedPtr<FNiagaraSimCacheTreeItem> ParentItem = ComponentTreeItems.Last();
+
+			TSharedRef<FNiagaraSimCacheComponentTreeItem> NewItem = MakeShared<FNiagaraSimCacheComponentTreeItem>(WeakOwnerTreeView);
+			NewItem->SetFilterName(ComponentPath);
+			NewItem->TypeDef = ComponentType;
+
+			if (ComponentTreeItems.Num() > 1)
+			{
+				const FString DisplayPath = ComponentPath.Mid(ParentItem->GetFilterName().Len() + 1);
+				NewItem->SetDisplayName(FText::FromString(DisplayPath));
+			}
+			else
+			{
+				NewItem->SetDisplayName(FText::FromString(ComponentPath));
+			}
+
+			ParentItem->AddChild(NewItem);
+			return NewItem;
+		}
+
+		void AddItemToTree(TWeakPtr<SNiagaraSimCacheTreeView> WeakOwnerTreeView, TArray<TSharedRef<FNiagaraSimCacheTreeItem>>& ComponentTreeItems, const FNiagaraSimCacheVariable& Variable)
+		{
+			const FNiagaraTypeDefinition VariableTypeDef = Variable.Variable.GetType();
+			const FString VariableName = Variable.Variable.GetName().ToString();
+
+			TSharedRef<FNiagaraSimCacheTreeItem> ParentItem = AddItemToTreeInternal(WeakOwnerTreeView, ComponentTreeItems, VariableName, VariableTypeDef);
+			if (!VariableTypeDef.IsEnum() && !FNiagaraTypeDefinition::IsScalarDefinition(VariableTypeDef))
+			{
+				ComponentTreeItems.Add(ParentItem);
+				TraverseStructItems(
+					VariableName,
+					VariableTypeDef,
+					[&WeakOwnerTreeView, &ComponentTreeItems, &ParentItem](const FString& StructPath, const FNiagaraTypeDefinition& StructTypeDef)
+					{
+						AddItemToTreeInternal(WeakOwnerTreeView, ComponentTreeItems, StructPath, StructTypeDef);
+					}
+				);
+				ComponentTreeItems.Pop();
+			}
+		}
+	} //namespace TypeGroupedTree
+
+} //namespace NiagaraSimCacheViewModelPrivate
 
 FNiagaraSimCacheViewModel::FNiagaraSimCacheViewModel()
 {
@@ -609,17 +766,15 @@ void FNiagaraSimCacheViewModel::UpdateComponentInfos()
 	}
 }
 
-void FNiagaraSimCacheViewModel::BuildTreeItemChildren(TSharedPtr<FNiagaraSimCacheTreeItem> InTreeItem, TWeakPtr<SNiagaraSimCacheTreeView> OwningTreeView)
+void FNiagaraSimCacheViewModel::BuildTreeItemChildren(TSharedPtr<FNiagaraSimCacheTreeItem> TreeItem, TWeakPtr<SNiagaraSimCacheTreeView> OwningTreeView)
 {
-	FNiagaraSimCacheTreeItem* TreeItem = InTreeItem.Get();
-
 	if (!TreeItem || !SimCache)
 	{
 		return;
 	}
 
 	TOptional<int32> CacheEmitterIndex;
-	switch (InTreeItem->GetType())
+	switch (TreeItem->GetType())
 	{
 		case ENiagaraSimCacheOverviewItemType::System:
 			CacheEmitterIndex = INDEX_NONE;
@@ -639,25 +794,20 @@ void FNiagaraSimCacheViewModel::BuildTreeItemChildren(TSharedPtr<FNiagaraSimCach
 	
 	if (CacheEmitterIndex.IsSet())
 	{
+		TArray<TSharedRef<FNiagaraSimCacheTreeItem>> ComponentTreeItems;
+		ComponentTreeItems.Emplace(TreeItem.ToSharedRef());
+
 		SimCache->ForEachEmitterAttribute(
 			CacheEmitterIndex.GetValue(),
-			[&](const FNiagaraSimCacheVariable& Variable)
+			[OwningTreeView, &ComponentTreeItems](const FNiagaraSimCacheVariable& Variable)
 			{
-				FNiagaraTypeDefinition TypeDef = Variable.Variable.GetType();
-
-				TSharedRef<FNiagaraSimCacheComponentTreeItem> CurrentItem = MakeShared<FNiagaraSimCacheComponentTreeItem>(OwningTreeView);
-				
-				CurrentItem->SetDisplayName(FText::FromName(Variable.Variable.GetName()));
-				CurrentItem->SetFilterName(Variable.Variable.GetName().ToString());
-				CurrentItem->TypeDef = TypeDef;
-				CurrentItem->EmitterName = TreeItem->GetEmitterName();
-				//CurrentItem->RootItem = TreeItem;
-
-				TreeItem->AddChild(CurrentItem);
-				
-				if(!TypeDef.IsEnum() && !FNiagaraTypeDefinition::IsScalarDefinition(TypeDef))
+				if (NiagaraSimCacheViewModelPrivate::bNamespaceGroupingEnabled)
+				{
+					NiagaraSimCacheViewModelPrivate::NamespaceGroupedTree::AddItemToTree(OwningTreeView, ComponentTreeItems, Variable);
+				}
+				else
 				{
-					RecursiveBuildTreeItemChildren(TreeItem, CurrentItem, TypeDef, OwningTreeView);
+					NiagaraSimCacheViewModelPrivate::TypeGroupedTree::AddItemToTree(OwningTreeView, ComponentTreeItems, Variable);
 				}
 				return true;
 			}
@@ -665,51 +815,6 @@ void FNiagaraSimCacheViewModel::BuildTreeItemChildren(TSharedPtr<FNiagaraSimCach
 	}
 }
 
-void FNiagaraSimCacheViewModel::RecursiveBuildTreeItemChildren(FNiagaraSimCacheTreeItem* Root,
-	TSharedRef<FNiagaraSimCacheComponentTreeItem> Parent, FNiagaraTypeDefinition TypeDefinition, TWeakPtr<SNiagaraSimCacheTreeView> OwningTreeView)
-{
-	UScriptStruct* Struct = TypeDefinition.GetScriptStruct();
-
-	for (TFieldIterator<FProperty> PropertyIt(Struct, EFieldIteratorFlags::IncludeSuper); PropertyIt; ++PropertyIt)
-	{
-		FProperty* Property = *PropertyIt;
-		TSharedRef<FNiagaraSimCacheComponentTreeItem> CurrentItem = MakeShared<FNiagaraSimCacheComponentTreeItem> (OwningTreeView);
-		
-		FString PropertyName = Property->GetName();
-		
-		CurrentItem->SetDisplayName(FText::FromString(PropertyName));
-		CurrentItem->SetFilterName(Parent->GetFilterName().Append(".").Append(PropertyName));
-		CurrentItem->SetEmitterName(Root->GetEmitterName());
-		//CurrentItem->RootItem = Root;
-
-		Parent->AddChild(CurrentItem);
-
-		if(Property->IsA(FStructProperty::StaticClass()))
-		{
-			const FStructProperty* StructProperty = CastFieldChecked<FStructProperty>(Property);
-			UScriptStruct* FriendlyStruct = FNiagaraTypeHelper::FindNiagaraFriendlyTopLevelStruct(StructProperty->Struct, ENiagaraStructConversion::Simulation);
-			FNiagaraTypeDefinition StructTypeDef = FNiagaraTypeDefinition(FriendlyStruct);
-			CurrentItem->TypeDef = StructTypeDef;
-			RecursiveBuildTreeItemChildren(Root, CurrentItem, StructTypeDef, OwningTreeView);
-		}
-		else if (Property->IsA(FNumericProperty::StaticClass()))
-		{
-			if(Property->IsA(FIntProperty::StaticClass()))
-			{
-				CurrentItem->TypeDef = FNiagaraTypeDefinition::GetIntDef();
-			}
-			else if (Property->IsA(FFloatProperty::StaticClass()))
-			{
-				CurrentItem->TypeDef = FNiagaraTypeDefinition::GetFloatDef();
-			}
-		}
-		else if (Property->IsA(FBoolProperty::StaticClass()))
-		{
-			CurrentItem->TypeDef = FNiagaraTypeDefinition::GetBoolDef();
-		}
-	}
-}
-
 void FNiagaraSimCacheViewModel::BuildEntries(TWeakPtr<SNiagaraSimCacheTreeView> OwningTreeView)
 {
 	RootEntries.Empty();
diff --git a/Engine/Plugins/FX/Niagara/Source/NiagaraEditor/Private/Widgets/SNiagaraSimCacheTreeView.h b/Engine/Plugins/FX/Niagara/Source/NiagaraEditor/Private/Widgets/SNiagaraSimCacheTreeView.h
--- a/Engine/Plugins/FX/Niagara/Source/NiagaraEditor/Private/Widgets/SNiagaraSimCacheTreeView.h
+++ b/Engine/Plugins/FX/Niagara/Source/NiagaraEditor/Private/Widgets/SNiagaraSimCacheTreeView.h
@@ -202,10 +202,10 @@ struct FNiagaraSimCacheComponentTreeItem : FNiagaraSimCacheTreeItem
 	{
 		TSharedRef<SHorizontalBox> Contents = SNew(SHorizontalBox);
 
-		if(TypeDef.IsSet())
 		{
-			FSlateBrush const* IconBrush = TypeDef.GetValue().IsStatic() ? FNiagaraEditorStyle::Get().GetBrush(TEXT("NiagaraEditor.StaticIcon")) : FAppStyle::GetBrush(TEXT("Kismet.AllClasses.VariableIcon"));
-			const FLinearColor TypeColor = UEdGraphSchema_Niagara::GetTypeColor(TypeDef.GetValue());
+			FSlateBrush const* IconBrush = TypeDef.IsSet() && TypeDef.GetValue().IsStatic() ? FNiagaraEditorStyle::Get().GetBrush(TEXT("NiagaraEditor.StaticIcon")) : FAppStyle::GetBrush(TEXT("Kismet.AllClasses.VariableIcon"));
+			const FLinearColor TypeColor = TypeDef.IsSet() ? UEdGraphSchema_Niagara::GetTypeColor(TypeDef.GetValue()) : FLinearColor::Black;
+			const FText ToolTipText = TypeDef.IsSet() ? TypeDef.GetValue().GetNameText() : FText();
 			
 			Contents->AddSlot()
 			.HAlign(HAlign_Left)
@@ -214,7 +214,7 @@ struct FNiagaraSimCacheComponentTreeItem : FNiagaraSimCacheTreeItem
 				SNew(SScaleBox)
 				[
 					SNew(SNiagaraIconWidget)
-					.IconToolTip(TypeDef.GetValue().GetNameText())
+					.IconToolTip(ToolTipText)
 					.IconBrush(IconBrush)
 					.IconColor(TypeColor)
 					.SecondaryIconBrush(FAppStyle::GetBrush(TEXT("NoBrush")))
diff --git a/Engine/Plugins/FX/Niagara/Source/NiagaraEditor/Public/ViewModels/NiagaraSimCacheViewModel.h b/Engine/Plugins/FX/Niagara/Source/NiagaraEditor/Public/ViewModels/NiagaraSimCacheViewModel.h
--- a/Engine/Plugins/FX/Niagara/Source/NiagaraEditor/Public/ViewModels/NiagaraSimCacheViewModel.h
+++ b/Engine/Plugins/FX/Niagara/Source/NiagaraEditor/Public/ViewModels/NiagaraSimCacheViewModel.h
@@ -126,7 +126,6 @@ private:
 	void UpdateCurrentEntries();
 
 	void BuildTreeItemChildren(TSharedPtr<FNiagaraSimCacheTreeItem> InTreeItem, TWeakPtr<SNiagaraSimCacheTreeView> OwningTreeView);
-	void RecursiveBuildTreeItemChildren(FNiagaraSimCacheTreeItem* Root, TSharedRef<FNiagaraSimCacheComponentTreeItem> Parent, FNiagaraTypeDefinition TypeDefinition, TWeakPtr<SNiagaraSimCacheTreeView> OwningTreeView);
 
 	// Build out component infos for each buffer in this cache
 	void UpdateComponentInfos();
