- PR #14231: [Niagara][Performance] Destroy all niagara particle systems when they cannot be rendered (i.e. dedicated server) instead of only ones marked for auto destroy
- Not enabled by default can be enabled with fx.Niagara.Rollback.DestroyAllComponentsThatCantRender

#Synced-CL 49733395

diff --git a/Engine/Plugins/FX/Niagara/Source/Niagara/Private/NiagaraComponent.cpp b/Engine/Plugins/FX/Niagara/Source/Niagara/Private/NiagaraComponent.cpp
--- a/Engine/Plugins/FX/Niagara/Source/Niagara/Private/NiagaraComponent.cpp
+++ b/Engine/Plugins/FX/Niagara/Source/Niagara/Private/NiagaraComponent.cpp
@@ -58,6 +58,14 @@ DEFINE_RENDER_COMMAND_PIPE(NiagaraDynamicData, ERenderCommandPipeFlags::None);
 
 namespace NiagaraComponentPrivate
 {
+	bool bDestroyAllComponentsThatCantRender = false;
+	static FAutoConsoleVariableRef CVarRollback_DestroyAllComponentsThatCantRender(
+		TEXT("fx.Niagara.Rollback.DestroyAllComponentsThatCantRender"),
+		bDestroyAllComponentsThatCantRender,
+		TEXT("When true any component that can not render will be destroyed not just ones that are bAutoDestroy."),
+		ECVF_Default
+	);
+
 	static bool GParallelGDMEEnabled = true;
 	static FAutoConsoleVariableRef CVarNiagaraParallelGDME(
 		TEXT("fx.Niagara.ParallelGDME"),
@@ -1297,10 +1305,13 @@ void UNiagaraComponent::ActivateInternal(bool bReset /* = false */, bool bIsScal
 	}
 
 	UWorld* World = GetWorld();
+
 	// If the particle system can't ever render (ie on dedicated server or in a commandlet) than do not activate...
 	if (!FApp::CanEverRender() || !World || World->IsNetMode(NM_DedicatedServer) || World->bIsTearingDown)
 	{
-		if ( bAutoDestroy )
+		// For bAutoDestroy components kill them immediately otherwise they will never complete
+		// For none bAutoDestroy we can destroy them only if they are registered, killing other types can result in killing a component that may activate / be used later on
+		if (bAutoDestroy || (NiagaraComponentPrivate::bDestroyAllComponentsThatCantRender && IsRegistered()))
 		{
 			DestroyComponent();
 		}
