- NPC parameter tooltip will attempt to resolve the collection path to disambiguate where they come from

#Synced-CL 49733395

diff --git a/Engine/Plugins/FX/Niagara/Source/NiagaraEditor/Private/SNiagaraParameterName.cpp b/Engine/Plugins/FX/Niagara/Source/NiagaraEditor/Private/SNiagaraParameterName.cpp
--- a/Engine/Plugins/FX/Niagara/Source/NiagaraEditor/Private/SNiagaraParameterName.cpp
+++ b/Engine/Plugins/FX/Niagara/Source/NiagaraEditor/Private/SNiagaraParameterName.cpp
@@ -6,11 +6,31 @@
 #include "NiagaraNodeParameterMapBase.h"
 #include "NiagaraEditorUtilities.h"
 
+#include "UObject/UObjectIterator.h"
+
 #include "Widgets/SBoxPanel.h"
 #include "Widgets/Layout/SBorder.h"
 #include "Widgets/Text/SInlineEditableTextBlock.h"
 #include "Widgets/Text/STextBlock.h"
 
+namespace SNiagaraParameterNamePrivate
+{
+	void TryResolveParameterCollectionPath(FName CollectionName, FNiagaraNamespaceMetadata& NamespaceMetadata)
+	{
+		for (TObjectIterator<UNiagaraParameterCollection> NPCIt; NPCIt; ++NPCIt)
+		{
+			UNiagaraParameterCollection* NPC = *NPCIt;
+			if (::IsValid(NPC) && NPC->IsAsset() && (NPC->GetNamespace() == CollectionName))
+			{
+				CollectionName = FName(*NPC->GetPathName());
+				break;
+			}
+		}
+
+		NamespaceMetadata.SetDescription(FText::Format(NSLOCTEXT("NiagaraParameterName", "NPCAssetToolTipFormat", "References parameter collection '{0}'."), FText::FromName(CollectionName)));
+	}
+}
+
 void SNiagaraParameterName::Construct(const FArguments& InArgs)
 {
 	EditableTextStyle = InArgs._EditableTextStyle;
@@ -190,6 +210,11 @@ void SNiagaraParameterName::UpdateContent(FName InDisplayedParameterName, int32
 				{
 					NamespaceModifierMetadata = DefaultNamespaceModifierMetadata;
 					NamespaceModifierDisplayName = FText::FromString(FName::NameToDisplayString(NamespaceModifier.ToString(), false).ToUpper());
+
+					if (NamePartIndex == 1 && NameParts[0] == FNiagaraConstants::ParameterCollectionNamespace)
+					{
+						SNiagaraParameterNamePrivate::TryResolveParameterCollectionPath(NameParts[1], NamespaceModifierMetadata);
+					}
 				}
 
 				TSharedRef<SBorder> NamespaceModifierBorder = CreateNamespaceWidget(
