- Warn the user when changing an NPC namespace that assets will loose references
- Iterating over all the NPC references to fix up will be done later

#Synced-CL 49733395

diff --git a/Engine/Plugins/FX/Niagara/Source/NiagaraEditor/Private/Toolkits/NiagaraParameterCollectionToolkit.cpp b/Engine/Plugins/FX/Niagara/Source/NiagaraEditor/Private/Toolkits/NiagaraParameterCollectionToolkit.cpp
--- a/Engine/Plugins/FX/Niagara/Source/NiagaraEditor/Private/Toolkits/NiagaraParameterCollectionToolkit.cpp
+++ b/Engine/Plugins/FX/Niagara/Source/NiagaraEditor/Private/Toolkits/NiagaraParameterCollectionToolkit.cpp
@@ -26,11 +26,139 @@
 #include "ViewModels/NiagaraParameterCollectionAssetViewModel.h"
 
 #include "PropertyEditorModule.h"
+#include "IDetailCustomization.h"
 #include "IDetailsView.h"
+#include "DetailLayoutBuilder.h"
+#include "DetailCategoryBuilder.h"
+#include "DetailWidgetRow.h"
+#include "Misc/MessageDialog.h"
 #include "Widgets/Input/SSearchBox.h"
 
 #define LOCTEXT_NAMESPACE "NiagaraParameterCollectionEditor"
 
+namespace NiagaraParameterCollectionToolkitPrivate
+{
+	class FNPCDetailCustomization : public IDetailCustomization
+	{
+	public:
+		static TSharedRef<IDetailCustomization> MakeInstance()
+		{
+			return MakeShared<FNPCDetailCustomization>();
+		}
+
+		// Begin: IDetailCustomization
+		virtual void CustomizeDetails(IDetailLayoutBuilder& DetailBuilder) override
+		{
+			static const FName NAME_ParameterColection("Parameter Collection");
+			static const FName NAME_Namespace("Namespace");
+
+			IDetailCategoryBuilder& DetailCategory = DetailBuilder.EditCategory(NAME_ParameterColection);
+
+			TArray<TSharedRef<IPropertyHandle>> CategoryProperties;
+			DetailCategory.GetDefaultProperties(CategoryProperties);
+
+			for (TSharedRef<IPropertyHandle>& PropertyHandle : CategoryProperties)
+			{
+				const FName PropertyName = PropertyHandle->GetProperty() ? PropertyHandle->GetProperty()->GetFName() : NAME_None;
+				if (PropertyName == NAME_Namespace)
+				{
+					TOptional<FAssetData> OwningAsset;
+					{
+						TArray<UObject*> OwningObjects;
+						PropertyHandle->GetOuterObjects(OwningObjects);
+
+						if (OwningObjects.Num() == 1)
+						{
+							OwningAsset.Emplace(OwningObjects[0]->GetPackage());
+						}
+					}
+
+
+					IDetailPropertyRow& PropertyRow = DetailCategory.AddProperty(PropertyHandle);
+					PropertyRow.CustomWidget()
+					.NameContent()
+					[
+						PropertyHandle->CreatePropertyNameWidget()
+					]
+					.ValueContent()
+					[
+						SNew(SEditableTextBox)
+						.Text_Lambda(
+							[PropertyHandle]()
+							{
+								FText Val;
+								PropertyHandle->GetValueAsFormattedText(Val);
+								return Val;
+							}
+						)
+						.OnTextCommitted_Lambda(
+							[PropertyHandle, OwningAsset](const FText& NewText, ETextCommit::Type CommitType)
+							{
+								FText ExistingText;
+								PropertyHandle->GetValueAsFormattedText(ExistingText);
+								if (NewText.EqualTo(ExistingText))
+								{
+									return;
+								}
+
+								// Intentionally set the value as the dialog will cause reenterance into the function and we will early out
+								PropertyHandle->SetValueFromFormattedString(NewText.ToString());
+
+								// Look to see if any assets would be impacted by this change
+								// If so allow the user to opt out of the change as committed would break the asset
+								if (OwningAsset.IsSet())
+								{
+									TArray<FAssetData> FoundReferencers;
+									const int32 ReferencedAssetCount =
+										FNiagaraEditorUtilities::GetReferencedAssetCount(
+										OwningAsset.GetValue(),
+										[&FoundReferencers](const FAssetData& AssetToCheck)
+										{
+											// Ignore any references from NPC / NPCIs but we need to include all other referencers
+											if (AssetToCheck.GetClass() == UNiagaraParameterCollection::StaticClass() || AssetToCheck.GetClass() == UNiagaraParameterCollectionInstance::StaticClass())
+											{
+												return FNiagaraEditorUtilities::ETrackAssetResult::Ignore;
+											}
+											FoundReferencers.Add(AssetToCheck);
+											return FNiagaraEditorUtilities::ETrackAssetResult::Count;
+										}
+									);
+
+									if (ReferencedAssetCount > 0)
+									{
+										const EAppReturnType::Type DialogReturnValue =
+											FMessageDialog::Open(
+												EAppMsgType::OkCancel,
+												FText::Format(LOCTEXT("CommitNamespaceChangeMessage", "Changing the namespace will impact {0} known asset(s), you will need to manually fix up these references.  See the log for a full list."), ReferencedAssetCount),
+												LOCTEXT("CommitNamespaceChangeTitle", "Commit namespace change?")
+											);
+
+										UE_LOG(LogNiagaraEditor, Log, TEXT("Known NPC referencers..."));
+										for (const FAssetData& FoundReferencer : FoundReferencers)
+										{
+											UE_LOG(LogNiagaraEditor, Log, TEXT("\t%s"), *FoundReferencer.GetFullName());
+										}
+
+										if (DialogReturnValue == EAppReturnType::Cancel)
+										{
+											PropertyHandle->SetValueFromFormattedString(ExistingText.ToString());
+										}
+									}
+								}
+							}
+						)
+					];
+				}
+				else
+				{
+					DetailCategory.AddProperty(PropertyHandle);
+				}
+			}
+		}
+		// End: IDetailCustomization
+	};
+} // namespace NiagaraParameterCollectionToolkitPrivate
+
 const FName FNiagaraParameterCollectionToolkit::MainTabID(TEXT("NiagaraParameterCollectionEditor_Main"));
 
 void FNiagaraParameterCollectionToolkit::RegisterTabSpawners(const TSharedRef<class FTabManager>& InTabManager)
@@ -148,7 +276,6 @@ FLinearColor FNiagaraParameterCollectionToolkit::GetWorldCentricTabColorScale()
 	return FNiagaraEditorModule::WorldCentricTabColorScale;
 }
 
-
 TSharedRef<SDockTab> FNiagaraParameterCollectionToolkit::SpawnTab_Main(const FSpawnTabArgs& Args)
 {
 	check(Args.GetTabId().TabType == MainTabID);
@@ -165,7 +292,8 @@ TSharedRef<SDockTab> FNiagaraParameterCollectionToolkit::SpawnTab_Main(const FSp
 	DetailsViewArgs.bAllowSearch = false;
 	
 	TSharedRef<IDetailsView> DetailsView = PropertyEditorModule.CreateDetailView(DetailsViewArgs);
-	
+	DetailsView->RegisterInstancedCustomPropertyLayout(UNiagaraParameterCollection::StaticClass(), FOnGetDetailCustomizationInstance::CreateStatic(&NiagaraParameterCollectionToolkitPrivate::FNPCDetailCustomization::MakeInstance));
+
 	if (Instance->IsDefaultInstance())
 	{
 		DetailsView->SetObject(Collection);
