- LWE: Parameters panel support
- LWE particle attributes are now visible in the panel
- Read / writes to attributes will now show up in the stack for LW emitters

#Synced-CL 49733395

diff --git a/Engine/Plugins/FX/Niagara/Source/Niagara/Internal/Stateless/NiagaraStatelessEmitter.h b/Engine/Plugins/FX/Niagara/Source/Niagara/Internal/Stateless/NiagaraStatelessEmitter.h
--- a/Engine/Plugins/FX/Niagara/Source/Niagara/Internal/Stateless/NiagaraStatelessEmitter.h
+++ b/Engine/Plugins/FX/Niagara/Source/Niagara/Internal/Stateless/NiagaraStatelessEmitter.h
@@ -58,6 +58,11 @@ public:
 	//~End UNiagaraEmitterBase Interface
 
 #if WITH_EDITOR
+	// Callback is owner name, variable, read, write
+	NIAGARA_API void GatherSimulationAttributeReferences(const FNiagaraStatelessSpawnInfo& SpawnInfo, TFunction<void(FName, const FNiagaraVariableBase&, bool, bool)> ReferenceFound) const;
+	NIAGARA_API void GatherSimulationAttributeReferences(const UNiagaraStatelessModule* Module, TFunction<void(FName, const FNiagaraVariableBase&, bool, bool)> ReferenceFound) const;
+	NIAGARA_API void GatherSimulationAttributeReferences(TFunction<void(FName, const FNiagaraVariableBase&, bool, bool)> ReferenceFound) const;
+
 	FOnTemplateChanged& GetOnTemplateChanged() { return OnTemplateChanged; }
 
 protected:
diff --git a/Engine/Plugins/FX/Niagara/Source/Niagara/Private/Stateless/NiagaraStatelessEmitter.cpp b/Engine/Plugins/FX/Niagara/Source/Niagara/Private/Stateless/NiagaraStatelessEmitter.cpp
--- a/Engine/Plugins/FX/Niagara/Source/Niagara/Private/Stateless/NiagaraStatelessEmitter.cpp
+++ b/Engine/Plugins/FX/Niagara/Source/Niagara/Private/Stateless/NiagaraStatelessEmitter.cpp
@@ -250,6 +250,71 @@ void UNiagaraStatelessEmitter::HandleVariableRemoved(const FNiagaraVariable& InO
 	//	FNiagaraSystemUpdateContext UpdateCtx(FVersionedNiagaraEmitter(this, EmitterVersion), true);
 	//}
 }
+
+void UNiagaraStatelessEmitter::GatherSimulationAttributeReferences(const FNiagaraStatelessSpawnInfo& SpawnInfo, TFunction<void(FName, const FNiagaraVariableBase&, bool, bool)> ReferenceFound) const
+{
+	FNiagaraDistributionBase::ForEachParameterBinding(
+		FNiagaraStatelessSpawnInfo::StaticStruct(),
+		&SpawnInfo,
+		[&ReferenceFound, &SpawnInfo](const FNiagaraVariableBase& Variable)
+		{
+			const UEnum* SpawnTypeEnum = StaticEnum<ENiagaraStatelessSpawnInfoType>();
+			const FName SpawnTypeName(*SpawnTypeEnum->GetDisplayNameTextByValue(int64(SpawnInfo.Type)).ToString());
+			ReferenceFound(SpawnTypeName, Variable, true, false);
+		}
+	);
+}
+
+void UNiagaraStatelessEmitter::GatherSimulationAttributeReferences(const UNiagaraStatelessModule* Module, TFunction<void(FName, const FNiagaraVariableBase&, bool, bool)> ReferenceFound) const
+{
+	if (StatelessEmitterData == nullptr || !StatelessEmitterData->ParticleDataSetCompiledData.IsValid() || !Module || !Module->IsModuleEnabled())
+	{
+		return;
+	}
+
+	FNiagaraDistributionBase::ForEachParameterBinding(
+		Module,
+		[this, Module, &ReferenceFound](const FNiagaraVariableBase& Variable)
+		{
+			const FText ModuleDisplayText = Module->GetClass()->GetDisplayNameText();
+			ReferenceFound(FName(*ModuleDisplayText.ToString()), Variable, true, false);
+		}
+	);
+
+	// Gather particle data read / writes
+	TArray<FNiagaraVariableBase> ModuleOutputVariables;
+	Module->GetOutputVariables(ModuleOutputVariables, UNiagaraStatelessModule::EVariableFilter::Used);
+
+	for (const FNiagaraVariableBase& DataSetVariable : StatelessEmitterData->ParticleDataSetCompiledData->Variables)
+	{
+		if (!ModuleOutputVariables.Contains(DataSetVariable))
+		{
+			continue;
+		}
+
+		FNameBuilder VariableNameBuilder;
+		FNiagaraConstants::ParticleAttributeNamespace.ToString(VariableNameBuilder);
+		VariableNameBuilder.AppendChar('.');
+		DataSetVariable.GetName().AppendString(VariableNameBuilder);
+
+		const FNiagaraVariableBase ParticleVariable(DataSetVariable.GetType(), FName(VariableNameBuilder.ToString()));
+		const FText ModuleDisplayText = Module->GetClass()->GetDisplayNameText();
+		ReferenceFound(FName(*ModuleDisplayText.ToString()), ParticleVariable, false, true);
+	}
+}
+
+void UNiagaraStatelessEmitter::GatherSimulationAttributeReferences(TFunction<void(FName, const FNiagaraVariableBase&, bool, bool)> ReferenceFound) const
+{
+	for (const FNiagaraStatelessSpawnInfo& SpawnInfo : SpawnInfos)
+	{
+		GatherSimulationAttributeReferences(SpawnInfo, ReferenceFound);
+	}
+
+	for (const UNiagaraStatelessModule* Module : Modules)
+	{
+		GatherSimulationAttributeReferences(Module, ReferenceFound);
+	}
+}
 #endif
 
 #if WITH_EDITOR
diff --git a/Engine/Plugins/FX/Niagara/Source/NiagaraEditor/Private/ViewModels/NiagaraParameterPanelViewModel.cpp b/Engine/Plugins/FX/Niagara/Source/NiagaraEditor/Private/ViewModels/NiagaraParameterPanelViewModel.cpp
--- a/Engine/Plugins/FX/Niagara/Source/NiagaraEditor/Private/ViewModels/NiagaraParameterPanelViewModel.cpp
+++ b/Engine/Plugins/FX/Niagara/Source/NiagaraEditor/Private/ViewModels/NiagaraParameterPanelViewModel.cpp
@@ -10,6 +10,7 @@
 #include "NiagaraEditorModule.h"
 #include "NiagaraEditorUtilities.h"
 #include "NiagaraEmitter.h"
+#include "Stateless/NiagaraStatelessEmitter.h"
 #include "NiagaraEmitterHandle.h"
 #include "NiagaraGraph.h"
 #include "NiagaraNodeAssignment.h"
@@ -825,7 +826,6 @@ void FNiagaraSystemToolkitParameterPanelViewModel::Cleanup()
 	for (TSharedRef<FNiagaraEmitterHandleViewModel> EmitterHandleViewModel : SystemViewModel->GetEmitterHandleViewModels())
 	{
 		FNiagaraEmitterHandle* EmitterHandle = EmitterHandleViewModel->GetEmitterHandle();
-		//-TODO:Stateless: Add Stateless support
 		FVersionedNiagaraEmitterData* EmitterData = EmitterHandle ? EmitterHandle->GetEmitterData() : nullptr;
 		if (EmitterData)
 		{
@@ -929,7 +929,6 @@ void FNiagaraSystemToolkitParameterPanelViewModel::Init(const FSystemToolkitUICo
 	{
 		if (FNiagaraEmitterHandle* EmitterHandle = EmitterHandleViewModel->GetEmitterHandle())
 		{
-			//-TODO:Stateless: Add Stateless support?
 			if (EmitterHandle->GetEmitterData() == nullptr)
 			{
 				continue;
@@ -1890,7 +1889,9 @@ TArray<FNiagaraParameterPanelItem> FNiagaraSystemToolkitParameterPanelViewModel:
 	TArray<UNiagaraGraph*> AllGraphs = GetAllGraphsConst();
 	
 	TArray<UNiagaraGraph*> GraphsToCheck;
-	if(SystemViewModel->GetSelectionViewModel()->GetSystemIsSelected() == false && SystemViewModel->GetSelectionViewModel()->GetSelectedEmitterHandleIds().Num() == 0)
+	
+	UNiagaraSystemSelectionViewModel* SelectionViewModel = SystemViewModel->GetSelectionViewModel();
+	if(SelectionViewModel->GetSystemIsSelected() == false && SelectionViewModel->GetSelectedEmitterHandleIds().Num() == 0)
 	{	
 		GraphsToCheck = AllGraphs;
 	}
@@ -2035,7 +2036,7 @@ TArray<FNiagaraParameterPanelItem> FNiagaraSystemToolkitParameterPanelViewModel:
 				auto AddReadReference = [GraphToCheck](const FNiagaraParameterMapHistory::FReadHistory& History, FNiagaraParameterPanelItem& Item)
 				{
 					FNiagaraParameterReferencePath Source;
-					Source.SourceGraph = GraphToCheck;
+					Source.SourceObject = GraphToCheck;
 					Source.ModuleName = History.ReadPin.ModuleName;
 					if (Source.ModuleName == NAME_None && History.ReadPin.Pin->LinkedTo.Num() > 0)
 					{
@@ -2084,10 +2085,7 @@ TArray<FNiagaraParameterPanelItem> FNiagaraSystemToolkitParameterPanelViewModel:
 						// ignore writes from map get nodes, as it's just the default value pin
 						return;
 					}
-					FNiagaraParameterReferencePath Source;
-					Source.SourceGraph = GraphToCheck;
-					Source.ModuleName = History.ModuleName;
-					Item.AddToWriteCount(Source);
+					Item.AddToWriteCount(FNiagaraParameterReferencePath(GraphToCheck, History.ModuleName));
 				};
 
 				const TArray<FName>& CustomIterationSourceNamespaces = Builder.Histories[0].IterationNamespaceOverridesEncountered;
@@ -2186,40 +2184,101 @@ TArray<FNiagaraParameterPanelItem> FNiagaraSystemToolkitParameterPanelViewModel:
 			}
 		}
 
+		// Stateless emitter usage
+		for (const TSharedRef<FNiagaraEmitterHandleViewModel>& EmitterViewModel : SystemViewModel->GetEmitterHandleViewModels())
+		{
+			FNiagaraEmitterHandle* EmitterHandle = EmitterViewModel->GetEmitterHandle();
+			UNiagaraStatelessEmitter* StatelessEmitter = EmitterHandle && EmitterHandle->GetEmitterMode() == ENiagaraEmitterMode::Stateless ? EmitterHandle->GetStatelessEmitter() : nullptr;
+			if (StatelessEmitter == nullptr)
+			{
+				continue;
+			}
+
+			if (SelectionViewModel->GetSelectedEmitterHandleIds().Num() != 0)
+			{
+				if (!SelectionViewModel->GetSelectedEmitterHandleIds().Contains(EmitterHandle->GetId()))
+				{
+					continue;
+				}
+			}
+
+			StatelessEmitter->GatherSimulationAttributeReferences(
+				[&](FName OwnerName, const FNiagaraVariableBase& Variable, bool bRead, bool bWrite)
+				{
+					if (VisitedInvalidParameters.Contains(Variable))
+					{
+						return;
+					}
+
+					FNiagaraParameterPanelItem* ItemPtr = VisitedParameterToItemMap.Find(Variable);
+					if (ItemPtr == nullptr)
+					{
+						UNiagaraScriptVariable* ScriptVariable = NewObject<UNiagaraScriptVariable>(GetTransientPackage());
+						ScriptVariable->AddToRoot();
+						ScriptVariable->Init(Variable, FNiagaraVariableMetaData());
+						TransientParameterToScriptVarMap.Add(Variable, ScriptVariable);
+
+						FNiagaraNamespaceMetadata CandidateNamespaceMetaData = FNiagaraEditorUtilities::GetNamespaceMetaDataForVariableName(Variable.GetName());
+						if (CachedCurrentCategories.Contains(FNiagaraParameterPanelCategory(CandidateNamespaceMetaData)) == false)
+						{
+							VisitedInvalidParameters.Add(Variable);
+							return;
+						}
+
+						FNiagaraParameterPanelItem Item;
+						Item.ScriptVariable = ScriptVariable;
+						Item.NamespaceMetaData = CandidateNamespaceMetaData;
+						ItemPtr = &VisitedParameterToItemMap.Add(Variable, Item);
+					}
+
+					if (bRead)
+					{
+						ItemPtr->AddToReadCount(FNiagaraParameterReferencePath(StatelessEmitter, OwnerName));
+					}
+					if (bWrite)
+					{
+						ItemPtr->AddToWriteCount(FNiagaraParameterReferencePath(StatelessEmitter, OwnerName));
+					}
+				}
+			);
+		}
+
 		// Add active renderers usage variables to the counts
 		for (TSharedRef<FNiagaraEmitterHandleViewModel> EmitterVMS : SystemViewModel->GetEmitterHandleViewModels())
 		{
 			if (EmitterVMS.Get().IsValid() && EmitterVMS.Get().GetIsEnabled() && EmitterVMS.Get().GetEmitterHandle())
 			{
-				if (FVersionedNiagaraEmitterData* ED = EmitterVMS.Get().GetEmitterHandle()->GetEmitterData())
+				FNiagaraEmitterHandle* EmitterHandle = EmitterVMS.Get().GetEmitterHandle();
+				UObject* SourceObject = EmitterHandle->GetEmitterBase();
+				if (const FVersionedNiagaraEmitterData* EmitterData = EmitterVMS.Get().GetEmitterHandle()->GetEmitterData())
 				{
-					FNiagaraAliasContext ResolveAliasesContext(FNiagaraAliasContext::ERapidIterationParameterMode::EmitterOrParticleScript);
-					ResolveAliasesContext.ChangeEmitterNameToEmitter(EmitterVMS.Get().GetEmitterHandle()->GetUniqueInstanceName());
-					
-					ED->ForEachEnabledRenderer(
-						[&](UNiagaraRendererProperties* RenderProperties)
+					SourceObject = Cast<UNiagaraScriptSource>(EmitterData->GraphSource)->NodeGraph;
+				}
+
+				FNiagaraAliasContext ResolveAliasesContext(FNiagaraAliasContext::ERapidIterationParameterMode::EmitterOrParticleScript);
+				ResolveAliasesContext.ChangeEmitterNameToEmitter(EmitterHandle->GetUniqueInstanceName());
+
+				EmitterHandle->ForEachEnabledRendererWithIndex(
+					[&](const UNiagaraRendererProperties* RenderProperties, int32 RendererIndex)
+					{
+						for (const FNiagaraVariableAttributeBinding* Binding : RenderProperties->GetAttributeBindings())
 						{
-							for (const FNiagaraVariableAttributeBinding* Binding : RenderProperties->GetAttributeBindings())
+							FNiagaraVariable BoundAttribute = RenderProperties->GetBoundAttribute(Binding);
+							if (BoundAttribute.IsValid())
 							{
-								FNiagaraVariable BoundAttribute = RenderProperties->GetBoundAttribute(Binding);
-								if (BoundAttribute.IsValid())
-								{
-									BoundAttribute = FNiagaraUtilities::ResolveAliases(BoundAttribute, ResolveAliasesContext);
+								BoundAttribute = FNiagaraUtilities::ResolveAliases(BoundAttribute, ResolveAliasesContext);
 
-									if (FNiagaraParameterPanelItem* ItemPtr = VisitedParameterToItemMap.Find(BoundAttribute))
-									{
-										// This variable has already been registered, increment the reference count. Otherwise, it is 
-										// not a live binding and we can skip.
-										FNiagaraParameterReferencePath Source;
-										Source.SourceGraph = Cast<UNiagaraScriptSource>(ED->GraphSource)->NodeGraph;
-										Source.ModuleName = FName(RenderProperties->GetWidgetDisplayName().ToString() + " (" + Binding->GetDataSetBindableVariable().GetName().ToString() + " Binding)");
-										ItemPtr->AddToReadCount(Source);
-									}
+								if (FNiagaraParameterPanelItem* ItemPtr = VisitedParameterToItemMap.Find(BoundAttribute))
+								{
+									// This variable has already been registered, increment the reference count. Otherwise, it is 
+									// not a live binding and we can skip.
+									const FName SourceModuleName(RenderProperties->GetWidgetDisplayName().ToString() + " (" + Binding->GetDataSetBindableVariable().GetName().ToString() + " Binding)");
+									ItemPtr->AddToReadCount(FNiagaraParameterReferencePath(SourceObject, SourceModuleName));
 								}
 							}
 						}
-					);
-				}
+					}
+				);
 			}
 		}
 
@@ -2529,7 +2588,6 @@ void FNiagaraSystemToolkitParameterPanelViewModel::ReconcileOnGraphChangedBindin
 		{
 			TArray<UNiagaraScript*> EmitterScripts;
 			const bool bCompilableOnly = false;
-			//-TODO:Stateless: Add Stateless support?
 			if (FVersionedNiagaraEmitterData* EmitterData = EmitterHandle->GetEmitterData())
 			{
 				EmitterData->GetScripts(EmitterScripts, bCompilableOnly);
diff --git a/Engine/Plugins/FX/Niagara/Source/NiagaraEditor/Private/ViewModels/Stack/NiagaraStackStatelessEmitterSimulateGroup.cpp b/Engine/Plugins/FX/Niagara/Source/NiagaraEditor/Private/ViewModels/Stack/NiagaraStackStatelessEmitterSimulateGroup.cpp
--- a/Engine/Plugins/FX/Niagara/Source/NiagaraEditor/Private/ViewModels/Stack/NiagaraStackStatelessEmitterSimulateGroup.cpp
+++ b/Engine/Plugins/FX/Niagara/Source/NiagaraEditor/Private/ViewModels/Stack/NiagaraStackStatelessEmitterSimulateGroup.cpp
@@ -522,16 +522,18 @@ const UNiagaraStackEntry::FCollectedUsageData& UNiagaraStackStatelessModuleItem:
 		CachedCollectedUsageData = FCollectedUsageData();
 
 		INiagaraParameterPanelViewModel* ParamVM = GetSystemViewModel()->GetParameterPanelViewModel();
-		UNiagaraStatelessModule* StatelessModule = StatelessModuleWeak.Get();
-		if ( ParamVM && StatelessModule )
+		const UNiagaraStatelessModule* StatelessModule = StatelessModuleWeak.Get();
+		const UNiagaraStatelessEmitter* StatelessEmitter = StatelessModule ? StatelessModule->GetTypedOuter<const UNiagaraStatelessEmitter>() : nullptr;
+		if ( ParamVM && StatelessModule && StatelessEmitter )
 		{
-			FNiagaraDistributionBase::ForEachParameterBinding(
+			StatelessEmitter->GatherSimulationAttributeReferences(
 				StatelessModule,
-				[this, &ParamVM](const FNiagaraVariableBase& Variable)
+				[this, &ParamVM](FName WhereForm, const FNiagaraVariableBase& Variable, bool bRead, bool bWrite)
 				{
 					if (ParamVM->IsVariableSelected(Variable))
 					{
-						CachedCollectedUsageData.GetValue().bHasReferencedParameterRead = true;
+						CachedCollectedUsageData.GetValue().bHasReferencedParameterRead |= bRead;
+						CachedCollectedUsageData.GetValue().bHasReferencedParameterWrite |= bWrite;
 					}
 				}
 			);
diff --git a/Engine/Plugins/FX/Niagara/Source/NiagaraEditor/Private/ViewModels/Stack/NiagaraStackStatelessEmitterSpawnGroup.cpp b/Engine/Plugins/FX/Niagara/Source/NiagaraEditor/Private/ViewModels/Stack/NiagaraStackStatelessEmitterSpawnGroup.cpp
--- a/Engine/Plugins/FX/Niagara/Source/NiagaraEditor/Private/ViewModels/Stack/NiagaraStackStatelessEmitterSpawnGroup.cpp
+++ b/Engine/Plugins/FX/Niagara/Source/NiagaraEditor/Private/ViewModels/Stack/NiagaraStackStatelessEmitterSpawnGroup.cpp
@@ -392,16 +392,17 @@ const UNiagaraStackEntry::FCollectedUsageData& UNiagaraStackStatelessEmitterSpaw
 
 		INiagaraParameterPanelViewModel* ParamVM = GetSystemViewModel()->GetParameterPanelViewModel();
 		FNiagaraStatelessSpawnInfo* SpawnInfo = GetSpawnInfo();
-		if (SpawnInfo && ParamVM)
+		UNiagaraStatelessEmitter* StatelessEmitter = GetStatelessEmitter();
+		if (StatelessEmitter && SpawnInfo && ParamVM)
 		{
-			FNiagaraDistributionBase::ForEachParameterBinding(
-				FNiagaraStatelessSpawnInfo::StaticStruct(),
-				SpawnInfo,
-				[this, &ParamVM](const FNiagaraVariableBase& Variable)
+			StatelessEmitter->GatherSimulationAttributeReferences(
+				*SpawnInfo,
+				[this, &ParamVM](FName WhereForm, const FNiagaraVariableBase& Variable, bool bRead, bool bWrite)
 				{
 					if (ParamVM->IsVariableSelected(Variable))
 					{
-						CachedCollectedUsageData.GetValue().bHasReferencedParameterRead = true;
+						CachedCollectedUsageData.GetValue().bHasReferencedParameterRead |= bRead;
+						CachedCollectedUsageData.GetValue().bHasReferencedParameterWrite |= bWrite;
 					}
 				}
 			);
diff --git a/Engine/Plugins/FX/Niagara/Source/NiagaraEditor/Private/Widgets/SNiagaraParameterPanel.cpp b/Engine/Plugins/FX/Niagara/Source/NiagaraEditor/Private/Widgets/SNiagaraParameterPanel.cpp
--- a/Engine/Plugins/FX/Niagara/Source/NiagaraEditor/Private/Widgets/SNiagaraParameterPanel.cpp
+++ b/Engine/Plugins/FX/Niagara/Source/NiagaraEditor/Private/Widgets/SNiagaraParameterPanel.cpp
@@ -460,12 +460,13 @@ TSharedRef<SWidget> SNiagaraParameterPanel::OnGenerateWidgetForItem(const FNiaga
 		TMap<FName, TArray<FNiagaraParameterReferencePath>> PathByParent;
 		for (const FNiagaraParameterReferencePath& Path : Item.ReferencePaths)
 		{
-			if (const UNiagaraGraph* NiagaraGraph = Path.SourceGraph.Get())
+			if (const UObject* SourceObject = Path.SourceObject.Get())
 			{
-				FName ParentName = NiagaraGraph->GetFName();
-				if (NiagaraGraph->GetOwningEmitter().Emitter)
+				FName ParentName = SourceObject->GetFName();
+				const UNiagaraGraph* SourceGraph = Cast<const UNiagaraGraph>(SourceObject);
+				if (SourceGraph && SourceGraph->GetOwningEmitter().Emitter)
 				{
-					ParentName = NiagaraGraph->GetOwningEmitter().Emitter.GetFName();
+					ParentName = SourceGraph->GetOwningEmitter().Emitter.GetFName();
 				}
 				PathByParent.FindOrAdd(ParentName).Add(Path);
 			}
diff --git a/Engine/Plugins/FX/Niagara/Source/NiagaraEditor/Public/NiagaraParameterPanelTypes.h b/Engine/Plugins/FX/Niagara/Source/NiagaraEditor/Public/NiagaraParameterPanelTypes.h
--- a/Engine/Plugins/FX/Niagara/Source/NiagaraEditor/Public/NiagaraParameterPanelTypes.h
+++ b/Engine/Plugins/FX/Niagara/Source/NiagaraEditor/Public/NiagaraParameterPanelTypes.h
@@ -23,7 +23,8 @@ public:
 	)
 		: ScriptVariable(InScriptVariable)
 		, NamespaceMetaData(InNamespaceMetaData)
-	{};
+	{
+	}
 
 	/* Equality operator to support AddUnique when gathering via parameter panel view models. */
 	bool operator== (const FNiagaraParameterPanelItemBase & Other) const { return GetVariable() == Other.GetVariable(); };
@@ -39,14 +40,21 @@ public:
 
 struct FNiagaraParameterReferencePath
 {
-	TWeakObjectPtr<const UNiagaraGraph> SourceGraph;
+	FNiagaraParameterReferencePath() = default;
+	FNiagaraParameterReferencePath(const UObject* InSourceObject, FName InModuleName)
+		: SourceObject(InSourceObject)
+		, ModuleName(InModuleName)
+	{
+	}
+
+	TWeakObjectPtr<const UObject> SourceObject;
 	FName ModuleName = NAME_None;
 	bool bRead = false;
 	bool bWrite = false;
 
 	bool operator==(const FNiagaraParameterReferencePath& Other) const
 	{
-		return Other.SourceGraph == SourceGraph && Other.ModuleName == ModuleName;
+		return Other.SourceObject == SourceObject && Other.ModuleName == ModuleName;
 	};
 };
 
