- Add capture space when exporting Niagara mesh data

#preflght

#Synced-CL 49775968

diff --git a/Engine/Plugins/FX/Niagara/Shaders/Private/NiagaraRibbonVertexFactory.ush b/Engine/Plugins/FX/Niagara/Shaders/Private/NiagaraRibbonVertexFactory.ush
--- a/Engine/Plugins/FX/Niagara/Shaders/Private/NiagaraRibbonVertexFactory.ush
+++ b/Engine/Plugins/FX/Niagara/Shaders/Private/NiagaraRibbonVertexFactory.ush
@@ -856,14 +856,18 @@ FVertexFactoryIntermediates GetVertexFactoryIntermediates(FVertexFactoryInput In
 	return Intermediates;
 }
 
-float3 ConstructWorldPosition(FVertexFactoryInput Input, FRibbonPositionIntermediates Intermediates, FLWCVector3 PreViewTranslation)
+FLWCVector3 ConstructWorldPositionLWC(FVertexFactoryInput Input, FRibbonPositionIntermediates Intermediates)
 {
 	const float2 SliceVertexPositions = UnpackParticleVertexPosition(Input.InterpVtxId);
 	const float3 Offset = ((Intermediates.GeoRight * SliceVertexPositions.xxx) + (Intermediates.GeoUp * SliceVertexPositions.yyy)) * Intermediates.RibbonWidth.xxx;
-	const FLWCVector3 Pos = LWCAdd(LWCAdd(Intermediates.Position, LWCPromote(Offset)), PreViewTranslation);
-		
-	// Intermediates are already in world space, only need to account for the translation around view origin.
-	return LWCToFloat(Pos);
+	return LWCAdd(Intermediates.Position, LWCPromote(Offset));
+}
+
+float3 ConstructWorldPosition(FVertexFactoryInput Input, FRibbonPositionIntermediates Intermediates, FLWCVector3 PreViewTranslation)
+{
+	const FLWCVector3 WorldPosition = ConstructWorldPositionLWC(Input, Intermediates);
+	const FLWCVector3 TranslatedPosition = LWCAdd(WorldPosition, PreViewTranslation);
+	return LWCToFloat(TranslatedPosition);
 }
 
 float4 VertexFactoryGetWorldPosition(FVertexFactoryInput Input, FVertexFactoryIntermediates Intermediates)
@@ -1139,9 +1143,9 @@ uint NiagaraVertexFactoryExport_GetIndirectVertexCount()
 	return NiagaraRibbonVFLooseParameters.IndirectDrawOutput[NiagaraRibbonVFLooseParameters.IndirectDrawOutputOffset + INDEX_GEN_DRAW_INDIRECT_OFFSET + 0];
 }
 
-float3 NiagaraVertexFactoryExport_GetPosition(FVertexFactoryInput Input, FVertexFactoryIntermediates Intermediates)
+FLWCVector3 NiagaraVertexFactoryExport_GetPosition(FVertexFactoryInput Input, FVertexFactoryIntermediates Intermediates)
 {
-	return ConstructWorldPosition(Input, Intermediates.Current, LWCPromote(float3(0,0,0)));
+	return ConstructWorldPositionLWC(Input, Intermediates.Current);
 }
 
 float4 NiagaraVertexFactoryExport_GetColor(FVertexFactoryInput Input, FVertexFactoryIntermediates Intermediates)
diff --git a/Engine/Plugins/FX/Niagara/Shaders/Private/NiagaraSpriteVertexFactory.ush b/Engine/Plugins/FX/Niagara/Shaders/Private/NiagaraSpriteVertexFactory.ush
--- a/Engine/Plugins/FX/Niagara/Shaders/Private/NiagaraSpriteVertexFactory.ush
+++ b/Engine/Plugins/FX/Niagara/Shaders/Private/NiagaraSpriteVertexFactory.ush
@@ -1373,12 +1373,9 @@ uint NiagaraVertexFactoryExport_GetIndirectVertexCount()
 	return NumVertices * NumInstances;
 }
 
-float3 NiagaraVertexFactoryExport_GetPosition(FVertexFactoryInput Input, FVertexFactoryIntermediates Intermediates)
+FLWCVector3 NiagaraVertexFactoryExport_GetPosition(FVertexFactoryInput Input, FVertexFactoryIntermediates Intermediates)
 {
-	//FLWCVector3 Origin = LWCGetOrigin(DFToTileOffset(GetPrimitiveDataFromUniformBuffer().LocalToWorld));
-	//return LWCToFloat(LWCSubtract(Intermediates.VertexWorldPosition, Origin));
-
-	return LWCToFloat(Intermediates.VertexWorldPosition);
+	return Intermediates.VertexWorldPosition;
 }
 
 float4 NiagaraVertexFactoryExport_GetColor(FVertexFactoryInput Input, FVertexFactoryIntermediates Intermediates)
diff --git a/Engine/Plugins/FX/Niagara/Shaders/Private/NiagaraVertexFactoryExport.usf b/Engine/Plugins/FX/Niagara/Shaders/Private/NiagaraVertexFactoryExport.usf
--- a/Engine/Plugins/FX/Niagara/Shaders/Private/NiagaraVertexFactoryExport.usf
+++ b/Engine/Plugins/FX/Niagara/Shaders/Private/NiagaraVertexFactoryExport.usf
@@ -15,6 +15,7 @@ uint IsIndirectDraw;
 uint NumInstances;
 uint NumVerticesPerInstance;
 uint bApplyWPO;
+uint bWorldSpace;
 uint SectionInfoOutputOffset;
 
 uint VertexStride;
@@ -113,10 +114,22 @@ void VertexFactoryExportCS(uint3 GroupId : SV_GroupID, int GroupThreadIndex : SV
 		MaterialWPO = MakeFinite(GetMaterialWorldPositionOffset(VertexParameters));
 	}
 
-	//-TODO: LWC issues
-	OutputPosition(VertexIndex, NiagaraVertexFactoryExport_GetPosition(Input, VFIntermediates) + MaterialWPO);
-	OutputColor(VertexIndex, NiagaraVertexFactoryExport_GetColor(Input, VFIntermediates));
-	OutputTangentBasis(VertexIndex, NiagaraVertexFactoryExport_GetTangentBasis(Input, VFIntermediates));
+	const FLWCVector3 VertexPosition = NiagaraVertexFactoryExport_GetPosition(Input, VFIntermediates);
+	const float4 VertexColor = NiagaraVertexFactoryExport_GetColor(Input, VFIntermediates);
+	const float3x3 VertexTangentBasis = NiagaraVertexFactoryExport_GetTangentBasis(Input, VFIntermediates);
+	if (bWorldSpace)
+	{
+		OutputPosition(VertexIndex, LWCToFloat(VertexPosition) + MaterialWPO);
+		OutputTangentBasis(VertexIndex, VertexTangentBasis);
+	}
+	else
+	{
+		const FLWCInverseMatrix WorldToLocal = DFToTileOffset(GetPrimitiveDataFromUniformBuffer().WorldToLocal);
+		OutputPosition(VertexIndex, LWCMultiply(VertexPosition, WorldToLocal) + MaterialWPO);
+		OutputTangentBasis(VertexIndex, mul(VertexTangentBasis, (float3x3)WorldToLocal.M));
+	}
+	OutputColor(VertexIndex, VertexColor);
+	
 	for ( uint i=0; i < VertexTexCoordNum; ++i )
 	{
 		OutputTexCoord(VertexIndex, i, NiagaraVertexFactoryExport_GetTexCoord(Input, VFIntermediates, i));
diff --git a/Engine/Plugins/FX/Niagara/Source/Niagara/Internal/NiagaraRendererReadback.h b/Engine/Plugins/FX/Niagara/Source/Niagara/Internal/NiagaraRendererReadback.h
--- a/Engine/Plugins/FX/Niagara/Source/Niagara/Internal/NiagaraRendererReadback.h
+++ b/Engine/Plugins/FX/Niagara/Source/Niagara/Internal/NiagaraRendererReadback.h
@@ -18,11 +18,29 @@ class FSceneView;
 	#define WITH_NIAGARA_RENDERER_READBACK WITH_EDITOR
 #endif
 
+UENUM()
+enum class ENiagaraRendererReadbackCoordinateSpace : uint8
+{
+	/**
+	Data is captured in world space.
+	large world coordinates will be truncated.
+	*/
+	World,
+	/** 
+	Data is captured in local space to the primitive.
+	*/
+	Local,
+};
+
 USTRUCT(BlueprintType)
 struct FNiagaraRendererReadbackParameters
 {
 	GENERATED_BODY()
 
+	// What coordinate space to use when capturing positions / tangents / etc.
+	UPROPERTY(EditAnywhere, BlueprintReadWrite, Category = "Parameters")
+	ENiagaraRendererReadbackCoordinateSpace CoordinateSpace = ENiagaraRendererReadbackCoordinateSpace::World;
+
 	// When enabled vertex positions will be exported
 	UPROPERTY(EditAnywhere, BlueprintReadWrite, Category = "Parameters")
 	bool bExportPosition = true;
diff --git a/Engine/Plugins/FX/Niagara/Source/Niagara/Private/NiagaraRendererReadback.cpp b/Engine/Plugins/FX/Niagara/Source/Niagara/Private/NiagaraRendererReadback.cpp
--- a/Engine/Plugins/FX/Niagara/Source/Niagara/Private/NiagaraRendererReadback.cpp
+++ b/Engine/Plugins/FX/Niagara/Source/Niagara/Private/NiagaraRendererReadback.cpp
@@ -53,6 +53,7 @@ namespace NiagaraRendererReadback
 	{
 		bool													bExportMaterials = true;
 		bool													bApplyWPO = false;
+		bool													bWorldSpace = true;
 		bool													bCaptureAllPrimitives = false;
 		TOptional<int>											ViewIndexToCapture;
 		FNiagaraRendererReadbackResult							Result;
@@ -79,7 +80,7 @@ namespace NiagaraRendererReadback
 
 	//////////////////////////////////////////////////////////////////////////////////////////////////////////////////////////////////////////////////////////////////////////////
 
-	void ExportMeshBatch(FRHICommandList& RHICmdList, FScene* Scene, ERHIFeatureLevel::Type FeatureLevel, const FReadbackMeshBatch& ReadbackMeshBatch, bool bApplyWPO, const FNiagaraRendererReadbackResult& VertexLayout, FRDGBufferUAV* VertexBufferUAV)
+	void ExportMeshBatch(FRHICommandList& RHICmdList, FScene* Scene, ERHIFeatureLevel::Type FeatureLevel, const FReadbackMeshBatch& ReadbackMeshBatch, bool bApplyWPO, bool bWorldSpace, const FNiagaraRendererReadbackResult& VertexLayout, FRDGBufferUAV* VertexBufferUAV)
 	{
 		const FMeshBatch& MeshBatch = ReadbackMeshBatch.MeshBatch;
 
@@ -129,6 +130,7 @@ namespace NiagaraRendererReadback
 			SingleShaderBindings.Add(Shader->NumInstances,				ReadbackMeshBatch.NumInstances);
 			SingleShaderBindings.Add(Shader->NumVerticesPerInstance,	ReadbackMeshBatch.NumVerticesPerInstance);
 			SingleShaderBindings.Add(Shader->bApplyWPO,					bApplyWPO ? 1 : 0);
+			SingleShaderBindings.Add(Shader->bWorldSpace,				bWorldSpace ? 1 : 0);
 
 			SingleShaderBindings.Add(Shader->VertexStride,				VertexLayout.VertexStride);
 			SingleShaderBindings.Add(Shader->VertexPositionOffset,		VertexLayout.VertexPositionOffset);
@@ -262,7 +264,7 @@ namespace NiagaraRendererReadback
 
 							for (const FReadbackMeshBatch& ReadbackMeshBatch : Request->MeshBatches)
 							{
-								ExportMeshBatch(RHICmdList, Scene, FeatureLevel, ReadbackMeshBatch, Request->bApplyWPO, Request->Result, PassParameters->RWVertexBuffer);
+								ExportMeshBatch(RHICmdList, Scene, FeatureLevel, ReadbackMeshBatch, Request->bApplyWPO, Request->bWorldSpace, Request->Result, PassParameters->RWVertexBuffer);
 							}
 						}
 					);
@@ -417,6 +419,7 @@ namespace NiagaraRendererReadback
 		// Copy over parameters
 		ReadbackRequest->bExportMaterials	= Parameters.bExportMaterials;
 		ReadbackRequest->bApplyWPO			= Parameters.bApplyWPO;
+		ReadbackRequest->bWorldSpace		= Parameters.CoordinateSpace == ENiagaraRendererReadbackCoordinateSpace::World;
 		ReadbackRequest->ViewIndexToCapture	= Parameters.ViewIndexToCapture;
 
 		// Build Vertex Output
diff --git a/Engine/Plugins/FX/Niagara/Source/NiagaraEditor/Private/Widgets/SNiagaraDebuggerMeshCapture.cpp b/Engine/Plugins/FX/Niagara/Source/NiagaraEditor/Private/Widgets/SNiagaraDebuggerMeshCapture.cpp
--- a/Engine/Plugins/FX/Niagara/Source/NiagaraEditor/Private/Widgets/SNiagaraDebuggerMeshCapture.cpp
+++ b/Engine/Plugins/FX/Niagara/Source/NiagaraEditor/Private/Widgets/SNiagaraDebuggerMeshCapture.cpp
@@ -95,7 +95,7 @@ TSharedRef<SWidget> SNiagaraDebuggerMeshCapture::MakeToolbar()
 			FCanExecuteAction::CreateSP(this, &SNiagaraDebuggerMeshCapture::CanCaptureSelected)
 		),
 		NAME_None,
-		LOCTEXT("CaptureSelected", "Captured Selected"),
+		LOCTEXT("CaptureSelected", "Capture Selected"),
 		LOCTEXT("CaptureSelectedTooltip", "Captures the visible Niagara components that are children of the outliner selection.")
 	);
 
@@ -105,7 +105,7 @@ TSharedRef<SWidget> SNiagaraDebuggerMeshCapture::MakeToolbar()
 			FCanExecuteAction::CreateSP(this, &SNiagaraDebuggerMeshCapture::CanCaptureAll)
 		),
 		NAME_None,
-		LOCTEXT("CaptureAll", "Captured All"),
+		LOCTEXT("CaptureAll", "Capture All"),
 		LOCTEXT("CaptureAllTooltip", "Captures all the visible Niagara components that are next rendered to the scene.")
 	);
 
diff --git a/Engine/Plugins/FX/Niagara/Source/NiagaraVertexFactories/Internal/NiagaraVertexFactoryExport.h b/Engine/Plugins/FX/Niagara/Source/NiagaraVertexFactories/Internal/NiagaraVertexFactoryExport.h
--- a/Engine/Plugins/FX/Niagara/Source/NiagaraVertexFactories/Internal/NiagaraVertexFactoryExport.h
+++ b/Engine/Plugins/FX/Niagara/Source/NiagaraVertexFactories/Internal/NiagaraVertexFactoryExport.h
@@ -57,6 +57,7 @@ public:
 	LAYOUT_FIELD(FShaderParameter,			NumInstances);
 	LAYOUT_FIELD(FShaderParameter,			NumVerticesPerInstance);
 	LAYOUT_FIELD(FShaderParameter,			bApplyWPO);
+	LAYOUT_FIELD(FShaderParameter,			bWorldSpace);
 	LAYOUT_FIELD(FShaderParameter,			SectionInfoOutputOffset);
 
 	LAYOUT_FIELD(FShaderParameter,			VertexStride);
diff --git a/Engine/Plugins/FX/Niagara/Source/NiagaraVertexFactories/Private/NiagaraVertexFactoryExport.cpp b/Engine/Plugins/FX/Niagara/Source/NiagaraVertexFactories/Private/NiagaraVertexFactoryExport.cpp
--- a/Engine/Plugins/FX/Niagara/Source/NiagaraVertexFactories/Private/NiagaraVertexFactoryExport.cpp
+++ b/Engine/Plugins/FX/Niagara/Source/NiagaraVertexFactories/Private/NiagaraVertexFactoryExport.cpp
@@ -22,6 +22,7 @@ FNiagaraVertexFactoryExportCS::FNiagaraVertexFactoryExportCS(const FMeshMaterial
 	NumInstances.Bind(Initializer.ParameterMap, TEXT("NumInstances"));
 	NumVerticesPerInstance.Bind(Initializer.ParameterMap, TEXT("NumVerticesPerInstance"));
 	bApplyWPO.Bind(Initializer.ParameterMap, TEXT("bApplyWPO"));
+	bWorldSpace.Bind(Initializer.ParameterMap, TEXT("bWorldSpace"));
 	SectionInfoOutputOffset.Bind(Initializer.ParameterMap, TEXT("SectionInfoOutputOffset"));
 
 	VertexStride.Bind(Initializer.ParameterMap, TEXT("VertexStride"));
