Renamed delegate retrieval function due to overlap with other function name

Removed AddCategory HierarchyChanged payload as it was already being triggered via the AddChild function
> Caveat: It's called after creation, before initialization. This means things such as names might not be initialized so reaction code can't properly assess data in all cases
> Future: Defer change propagation by a frame, or do virtual Initialize function with TInstancedStruct Initilization payload

#Synced-CL 49840944

diff --git a/Engine/Plugins/FX/Niagara/Source/NiagaraEditor/Private/Customizations/UserParameters/NiagaraComponentUserParametersBuilder.cpp b/Engine/Plugins/FX/Niagara/Source/NiagaraEditor/Private/Customizations/UserParameters/NiagaraComponentUserParametersBuilder.cpp
--- a/Engine/Plugins/FX/Niagara/Source/NiagaraEditor/Private/Customizations/UserParameters/NiagaraComponentUserParametersBuilder.cpp
+++ b/Engine/Plugins/FX/Niagara/Source/NiagaraEditor/Private/Customizations/UserParameters/NiagaraComponentUserParametersBuilder.cpp
@@ -283,7 +283,7 @@ FNiagaraComponentUserParametersNodeBuilder::~FNiagaraComponentUserParametersNode
 			TSharedPtr<FNiagaraSystemViewModel> SystemViewModel = TNiagaraViewModelManager<UNiagaraSystem, FNiagaraSystemViewModel>::GetExistingViewModelForObject(System);
 			if(SystemViewModel.IsValid() && SystemViewModel->GetUserParametersHierarchyViewModel())
 			{
-				SystemViewModel->GetUserParametersHierarchyViewModel()->OnHierarchyChanged().RemoveAll(this);
+				SystemViewModel->GetUserParametersHierarchyViewModel()->GetOnHierarchyChangedDelegate().RemoveAll(this);
 			}
 		}
 	}
@@ -500,8 +500,8 @@ void FNiagaraComponentUserParametersNodeBuilder::RegisterRebuildOnHierarchyChang
 			TSharedPtr<FNiagaraSystemViewModel> SystemViewModel = TNiagaraViewModelManager<UNiagaraSystem, FNiagaraSystemViewModel>::GetExistingViewModelForObject(Asset);
 			if(SystemViewModel.IsValid())
 			{
-				SystemViewModel->GetUserParametersHierarchyViewModel()->OnHierarchyChanged().RemoveAll(this);
-				SystemViewModel->GetUserParametersHierarchyViewModel()->OnHierarchyChanged().Add(UDataHierarchyViewModelBase::FOnHierarchyChanged::FDelegate::CreateSP(this, &FNiagaraComponentUserParametersNodeBuilder::RebuildOnHierarchyChanged));
+				SystemViewModel->GetUserParametersHierarchyViewModel()->GetOnHierarchyChangedDelegate().RemoveAll(this);
+				SystemViewModel->GetUserParametersHierarchyViewModel()->GetOnHierarchyChangedDelegate().Add(UDataHierarchyViewModelBase::FOnHierarchyChanged::FDelegate::CreateSP(this, &FNiagaraComponentUserParametersNodeBuilder::RebuildOnHierarchyChanged));
 			}
 		}
 	}
diff --git a/Engine/Plugins/FX/Niagara/Source/NiagaraEditor/Private/Customizations/UserParameters/NiagaraSystemUserParametersBuilder.cpp b/Engine/Plugins/FX/Niagara/Source/NiagaraEditor/Private/Customizations/UserParameters/NiagaraSystemUserParametersBuilder.cpp
--- a/Engine/Plugins/FX/Niagara/Source/NiagaraEditor/Private/Customizations/UserParameters/NiagaraSystemUserParametersBuilder.cpp
+++ b/Engine/Plugins/FX/Niagara/Source/NiagaraEditor/Private/Customizations/UserParameters/NiagaraSystemUserParametersBuilder.cpp
@@ -63,7 +63,7 @@ FNiagaraSystemUserParameterBuilder::~FNiagaraSystemUserParameterBuilder()
 		{
 			if (UNiagaraUserParametersHierarchyViewModel* HierarchyViewModel = SystemViewModel.Pin()->GetUserParametersHierarchyViewModel())
 			{
-				HierarchyViewModel->OnHierarchyChanged().RemoveAll(this);
+				HierarchyViewModel->GetOnHierarchyChangedDelegate().RemoveAll(this);
 			}
 		}
 	}
@@ -79,8 +79,8 @@ void FNiagaraSystemUserParameterBuilder::GenerateChildContent(IDetailChildrenBui
 		{
 			System->GetExposedParameters().OnStructureChanged().Add(FNiagaraParameterStore::FOnStructureChanged::FDelegate::CreateSP(this, &FNiagaraSystemUserParameterBuilder::Rebuild));
 
-			SystemViewModel.Pin()->GetUserParametersHierarchyViewModel()->OnHierarchyChanged().RemoveAll(this);
-			SystemViewModel.Pin()->GetUserParametersHierarchyViewModel()->OnHierarchyChanged().Add(UDataHierarchyViewModelBase::FOnHierarchyChanged::FDelegate::CreateSP(this, &FNiagaraSystemUserParameterBuilder::RebuildOnHierarchyChanged));
+			SystemViewModel.Pin()->GetUserParametersHierarchyViewModel()->GetOnHierarchyChangedDelegate().RemoveAll(this);
+			SystemViewModel.Pin()->GetUserParametersHierarchyViewModel()->GetOnHierarchyChangedDelegate().Add(UDataHierarchyViewModelBase::FOnHierarchyChanged::FDelegate::CreateSP(this, &FNiagaraSystemUserParameterBuilder::RebuildOnHierarchyChanged));
 			
 			TSharedPtr<FNiagaraUserParameterPanelViewModel> UserParameterPanelViewModel = SystemViewModel.Pin()->GetUserParameterPanelViewModel();
 			if(ensure(UserParameterPanelViewModel.IsValid()))
diff --git a/Engine/Plugins/FX/Niagara/Source/NiagaraEditor/Private/Toolkits/NiagaraScriptToolkit.cpp b/Engine/Plugins/FX/Niagara/Source/NiagaraEditor/Private/Toolkits/NiagaraScriptToolkit.cpp
--- a/Engine/Plugins/FX/Niagara/Source/NiagaraEditor/Private/Toolkits/NiagaraScriptToolkit.cpp
+++ b/Engine/Plugins/FX/Niagara/Source/NiagaraEditor/Private/Toolkits/NiagaraScriptToolkit.cpp
@@ -212,7 +212,7 @@ void FNiagaraScriptToolkit::Initialize( const EToolkitMode::Type Mode, const TSh
 
 	ParametersHierarchyViewModel = NewObject<UNiagaraScriptParametersHierarchyViewModel>();
 	ParametersHierarchyViewModel->Initialize(ScriptViewModel.ToSharedRef());
-	ParametersHierarchyViewModel->OnHierarchyChanged().AddSP(this, &FNiagaraScriptToolkit::OnHierarchyChanged);
+	ParametersHierarchyViewModel->GetOnHierarchyChangedDelegate().AddSP(this, &FNiagaraScriptToolkit::OnHierarchyChanged);
 	
 	ParameterPanelViewModel = MakeShareable(new FNiagaraScriptToolkitParameterPanelViewModel(ScriptViewModel));
 	ParameterDefinitionsPanelViewModel = MakeShareable(new FNiagaraScriptToolkitParameterDefinitionsPanelViewModel(ScriptViewModel));
diff --git a/Engine/Plugins/FX/Niagara/Source/NiagaraEditor/Private/ViewModels/NiagaraEmitterViewModel.cpp b/Engine/Plugins/FX/Niagara/Source/NiagaraEditor/Private/ViewModels/NiagaraEmitterViewModel.cpp
--- a/Engine/Plugins/FX/Niagara/Source/NiagaraEditor/Private/ViewModels/NiagaraEmitterViewModel.cpp
+++ b/Engine/Plugins/FX/Niagara/Source/NiagaraEditor/Private/ViewModels/NiagaraEmitterViewModel.cpp
@@ -86,7 +86,7 @@ bool FNiagaraEmitterViewModel::Initialize(const FVersionedNiagaraEmitter& InEmit
 	SetEmitter(InEmitter);
 	SetSimulation(InSimulation);
 	SummaryViewHierarchyViewModel->Initialize(AsShared());
-	SummaryViewHierarchyViewModel->OnHierarchyChanged().AddSP(this, &FNiagaraEmitterViewModel::OnSummaryViewHierarchyChanged);
+	SummaryViewHierarchyViewModel->GetOnHierarchyChangedDelegate().AddSP(this, &FNiagaraEmitterViewModel::OnSummaryViewHierarchyChanged);
 	GetEditorData().OnSummaryViewStateChanged().AddSP(this, &FNiagaraEmitterViewModel::UpdateSelectionOnSummaryViewStateChanged);
 	return true;
 }
@@ -94,7 +94,7 @@ bool FNiagaraEmitterViewModel::Initialize(const FVersionedNiagaraEmitter& InEmit
 void FNiagaraEmitterViewModel::Reset()
 {
 	SetEmitter(FVersionedNiagaraEmitter());
-	SummaryViewHierarchyViewModel->OnHierarchyChanged().RemoveAll(this);
+	SummaryViewHierarchyViewModel->GetOnHierarchyChangedDelegate().RemoveAll(this);
 	SetSimulation(nullptr);
 }
 
diff --git a/Engine/Plugins/FX/Niagara/Source/NiagaraEditor/Private/ViewModels/NiagaraScratchPadScriptViewModel.cpp b/Engine/Plugins/FX/Niagara/Source/NiagaraEditor/Private/ViewModels/NiagaraScratchPadScriptViewModel.cpp
--- a/Engine/Plugins/FX/Niagara/Source/NiagaraEditor/Private/ViewModels/NiagaraScratchPadScriptViewModel.cpp
+++ b/Engine/Plugins/FX/Niagara/Source/NiagaraEditor/Private/ViewModels/NiagaraScratchPadScriptViewModel.cpp
@@ -84,7 +84,7 @@ void FNiagaraScratchPadScriptViewModel::Initialize(UNiagaraScript* Script, UNiag
 	ScriptHierarchyViewModel = NewObject<UNiagaraScriptParametersHierarchyViewModel>();
 	ScriptHierarchyViewModel->Initialize(AsShared());
 
-	ScriptHierarchyViewModel->OnHierarchyChanged().AddSP(this, &FNiagaraScratchPadScriptViewModel::OnHierarchyChanged);
+	ScriptHierarchyViewModel->GetOnHierarchyChangedDelegate().AddSP(this, &FNiagaraScratchPadScriptViewModel::OnHierarchyChanged);
 }
 
 TSharedPtr<IAssetReferenceFilter> FNiagaraScratchPadScriptViewModel::MakeAssetReferenceFilter() const
diff --git a/Engine/Plugins/FX/Niagara/Source/NiagaraEditor/Private/ViewModels/Stack/NiagaraStackSummaryViewInputCollection.cpp b/Engine/Plugins/FX/Niagara/Source/NiagaraEditor/Private/ViewModels/Stack/NiagaraStackSummaryViewInputCollection.cpp
--- a/Engine/Plugins/FX/Niagara/Source/NiagaraEditor/Private/ViewModels/Stack/NiagaraStackSummaryViewInputCollection.cpp
+++ b/Engine/Plugins/FX/Niagara/Source/NiagaraEditor/Private/ViewModels/Stack/NiagaraStackSummaryViewInputCollection.cpp
@@ -41,7 +41,7 @@ void UNiagaraStackSummaryViewCollection::Initialize(FRequiredEntryData InRequire
 
 	Emitter = InEmitter;
 	GetEmitterViewModel()->GetEditorData().OnSummaryViewStateChanged().AddUObject(this, &UNiagaraStackSummaryViewCollection::OnViewStateChanged);
-	GetEmitterViewModel()->GetSummaryHierarchyViewModel()->OnHierarchyChanged().AddUObject(this, &UNiagaraStackSummaryViewCollection::OnHierarchyChanged);
+	GetEmitterViewModel()->GetSummaryHierarchyViewModel()->GetOnHierarchyChangedDelegate().AddUObject(this, &UNiagaraStackSummaryViewCollection::OnHierarchyChanged);
 
 	AddChildFilter(FOnFilterChild::CreateUObject(this, &UNiagaraStackSummaryViewCollection::FilterOnlyModified));
 }
@@ -56,7 +56,7 @@ void UNiagaraStackSummaryViewCollection::FinalizeInternal()
 
 	if(EmitterViewModelPtr.IsValid() && EmitterViewModelPtr->GetSummaryHierarchyViewModel() != nullptr)
 	{
-		EmitterViewModelPtr->GetSummaryHierarchyViewModel()->OnHierarchyChanged().RemoveAll(this);
+		EmitterViewModelPtr->GetSummaryHierarchyViewModel()->GetOnHierarchyChangedDelegate().RemoveAll(this);
 	}
 
 	Super::FinalizeInternal();
diff --git a/Engine/Plugins/FX/Niagara/Source/NiagaraEditor/Private/Widgets/SNiagaraScriptInputPreviewPanel.cpp b/Engine/Plugins/FX/Niagara/Source/NiagaraEditor/Private/Widgets/SNiagaraScriptInputPreviewPanel.cpp
--- a/Engine/Plugins/FX/Niagara/Source/NiagaraEditor/Private/Widgets/SNiagaraScriptInputPreviewPanel.cpp
+++ b/Engine/Plugins/FX/Niagara/Source/NiagaraEditor/Private/Widgets/SNiagaraScriptInputPreviewPanel.cpp
@@ -340,7 +340,7 @@ void SNiagaraScriptInputPreviewPanel::SetupDelegates()
 {
 	Cast<UNiagaraScriptSource>(ScriptToolkit.Pin()->EditedNiagaraScript.GetScriptData()->GetSource())->NodeGraph->OnParametersChanged().AddSP(this, &SNiagaraScriptInputPreviewPanel::OnParametersChanged);
 
-	ScriptToolkit.Pin()->GetHierarchyViewModel()->OnHierarchyChanged().AddSP(this, &SNiagaraScriptInputPreviewPanel::RefreshOnHierarchyChanged);
+	ScriptToolkit.Pin()->GetHierarchyViewModel()->GetOnHierarchyChangedDelegate().AddSP(this, &SNiagaraScriptInputPreviewPanel::RefreshOnHierarchyChanged);
 }
 
 void SNiagaraScriptInputPreviewPanel::RemoveDelegates()
@@ -351,7 +351,7 @@ void SNiagaraScriptInputPreviewPanel::RemoveDelegates()
 
 		if(UNiagaraScriptParametersHierarchyViewModel* HierarchyViewModel = ScriptToolkit.Pin()->GetHierarchyViewModel())
 		{
-			HierarchyViewModel->OnHierarchyChanged().RemoveAll(this);
+			HierarchyViewModel->GetOnHierarchyChangedDelegate().RemoveAll(this);
 		}
 	}
 }
diff --git a/Engine/Source/Editor/DataHierarchyEditor/Private/DataHierarchyViewModelBase.cpp b/Engine/Source/Editor/DataHierarchyEditor/Private/DataHierarchyViewModelBase.cpp
--- a/Engine/Source/Editor/DataHierarchyEditor/Private/DataHierarchyViewModelBase.cpp
+++ b/Engine/Source/Editor/DataHierarchyEditor/Private/DataHierarchyViewModelBase.cpp
@@ -421,6 +421,8 @@ UDataHierarchyViewModelBase::~UDataHierarchyViewModelBase()
 
 void UDataHierarchyViewModelBase::Initialize()
 {		
+	bIsFinalized = false;
+	
 	HierarchyRoot = GetHierarchyRoot();
 	HierarchyRoot->SetFlags(RF_Transactional);
 
@@ -478,6 +480,7 @@ void UDataHierarchyViewModelBase::Finalize()
 	FinalizeInternal();
 	
 	bIsFinalized = true;
+	bIsInitialized = false;
 }
 
 TSharedPtr<FHierarchyElementViewModel> UDataHierarchyViewModelBase::CreateViewModelForElement(UHierarchyElement* Element, TSharedPtr<FHierarchyElementViewModel> Parent)
@@ -2001,10 +2004,6 @@ void UDataHierarchyViewModelBase::AddCategory(TSharedPtr<FHierarchyElementViewMo
 		}
 
 		CategoryViewModel->SetCategoryName(UE::DataHierarchyEditor::GetUniqueName(FName("New Category"), CategoryNames));
-		
-		TInstancedStruct<FHierarchyChangedPayload_ElementsAdded> Payload = TInstancedStruct<FHierarchyChangedPayload_ElementsAdded>::Make();
-		Payload.GetMutable<FHierarchyChangedPayload_ElementsAdded>().AddedElementViewModels.Add(CategoryViewModel);
-		OnHierarchyChanged(Payload);
 	}
 }
 
diff --git a/Engine/Source/Editor/DataHierarchyEditor/Public/DataHierarchyViewModelBase.h b/Engine/Source/Editor/DataHierarchyEditor/Public/DataHierarchyViewModelBase.h
--- a/Engine/Source/Editor/DataHierarchyEditor/Public/DataHierarchyViewModelBase.h
+++ b/Engine/Source/Editor/DataHierarchyEditor/Public/DataHierarchyViewModelBase.h
@@ -665,7 +665,7 @@ public:
 	FSimpleDelegate& OnRefreshSectionsView() { return RefreshSectionsViewDelegate; }
 
 	// Delegates for external systems
-	FOnHierarchyChanged& OnHierarchyChanged() { return OnHierarchyChangedDelegate; } 
+	FOnHierarchyChanged& GetOnHierarchyChangedDelegate() { return OnHierarchyChangedDelegate; } 
 	FOnRefreshViewRequested& OnRefreshViewRequested() { return RefreshAllViewsRequestedDelegate; }
 	FOnNavigateToElementIdentityInHierarchyRequested& OnNavigateToElementIdentityInHierarchyRequested() { return OnNavigateToElementIdentityInHierarchyRequestedDelegate; }
 	FOnNavigateToElementInHierarchyRequested& OnNavigateToElementInHierarchyRequested() { return OnNavigateToElementInHierarchyRequestedDelegate; }
