- Niagara type registry fixes
- Fixes a double add to the registry if type is added from multiple threads at the same time
- Fixes potential TLS cache entry with invalid type hash

#Synced-CL 50116758

diff --git a/Engine/Plugins/FX/Niagara/Source/Niagara/Private/NiagaraTypeRegistry.cpp b/Engine/Plugins/FX/Niagara/Source/Niagara/Private/NiagaraTypeRegistry.cpp
--- a/Engine/Plugins/FX/Niagara/Source/Niagara/Private/NiagaraTypeRegistry.cpp
+++ b/Engine/Plugins/FX/Niagara/Source/Niagara/Private/NiagaraTypeRegistry.cpp
@@ -287,12 +287,9 @@ int32 FNiagaraTypeRegistry::RegisterTypeInternal(const FNiagaraTypeDefinition& N
 		NiagaraTypeRegistryLocal::RegisteredTypeIndexCacheTLS.Reset();
 	}
 
-	uint32 TypeHash = 0;
-
+	const uint32 TypeHash = GetTypeHash(NewType);
 	if (NewType.RegisteredTypeDefIndex == INDEX_NONE)
 	{
-		TypeHash = GetTypeHash(NewType);
-
 		// first see if we're in the TLS cache
 		if (const int32* ExistingIndex = NiagaraTypeRegistryLocal::RegisteredTypeIndexCacheTLS.Find(TypeHash))
 		{
@@ -346,17 +343,25 @@ int32 FNiagaraTypeRegistry::RegisterTypeInternal(const FNiagaraTypeDefinition& N
 	}
 	else
 	{
-		// we need to create a new entry for a registered type
+		// We may need to create a new entry for a registered type
+		// Between the read and the write it could have been added so we need to do an addition search to make sure we don't double register
 		{
 			UE::TWriteScopeLock Lock(RegisteredTypesLock);
-			NewType.RegisteredTypeDefIndex = RegisteredTypes.Num();
-			RegisteredTypes.Add(MakeUnique<FNiagaraTypeDefinition>(NewType));
+			NewType.RegisteredTypeDefIndex = RegisteredTypes.IndexOfByPredicate([NewType](const TUniquePtr<FNiagaraTypeDefinition>& InDef) {return *InDef.Get() == NewType; });
+			if (NewType.RegisteredTypeDefIndex == INDEX_NONE)
+			{
+				NewType.RegisteredTypeDefIndex = RegisteredTypes.Num();
+				RegisteredTypes.Add(MakeUnique<FNiagaraTypeDefinition>(NewType));
+			}
 		}
 
 		{
 			FRWScopeLock Lock(RegistrationLock, SLT_Write);
-			RegisteredSoftObjectPaths.Emplace(NewType.ClassStructOrEnum);
-			RegisteredTypeIndexMap.Add(TypeHash, NewType.RegisteredTypeDefIndex);
+			if (!RegisteredTypeIndexMap.Contains(TypeHash))
+			{
+				RegisteredSoftObjectPaths.Emplace(NewType.ClassStructOrEnum);
+				RegisteredTypeIndexMap.Add(TypeHash, NewType.RegisteredTypeDefIndex);
+			}
 		}
 	}
 
