[HWRT] Fix RHI validation errors due to FNiagaraRibbonIndexBuffer not being in correct state.
- RayTracingDynamicGeometryUpdate requires index buffers to be in ERHIAccess::SRVMask state.

#Synced-CL 50182883

diff --git a/Engine/Plugins/FX/Niagara/Source/Niagara/Private/NiagaraRendererRibbons.cpp b/Engine/Plugins/FX/Niagara/Source/Niagara/Private/NiagaraRendererRibbons.cpp
--- a/Engine/Plugins/FX/Niagara/Source/Niagara/Private/NiagaraRendererRibbons.cpp
+++ b/Engine/Plugins/FX/Niagara/Source/Niagara/Private/NiagaraRendererRibbons.cpp
@@ -280,7 +280,7 @@ struct FNiagaraRibbonIndexBuffer final : FIndexBuffer
 		const FRHIBufferCreateDesc CreateDesc =
 			FRHIBufferCreateDesc::CreateIndex<uint32>(TEXT("NiagaraRibbonIndexBuffer"), NumElements)
 			.AddUsage(EBufferUsageFlags::Static | EBufferUsageFlags::UnorderedAccess | EBufferUsageFlags::ShaderResource)
-			.SetInitialState(ERHIAccess::VertexOrIndexBuffer);
+			.SetInitialState(ERHIAccess::SRVMask | ERHIAccess::VertexOrIndexBuffer);
 
 		IndexBufferRHI = RHICmdList.CreateBuffer(CreateDesc);
 	#if RHI_RAYTRACING
@@ -2383,9 +2383,9 @@ void FNiagaraRendererRibbons::InitializeViewIndexBuffersGPU(FRHICommandListImmed
 		Params.SubSegmentBitShift = RenderingViewResources->IndexGenerationSettings.SubSegmentBitShift;
 		Params.SubSegmentBitMask = RenderingViewResources->IndexGenerationSettings.SubSegmentBitMask;
 		
-		RHICmdList.Transition(FRHITransitionInfo(RenderingViewResources->IndexBuffer->UAV, ERHIAccess::VertexOrIndexBuffer, ERHIAccess::UAVCompute));
+		RHICmdList.Transition(FRHITransitionInfo(RenderingViewResources->IndexBuffer->UAV, ERHIAccess::SRVMask | ERHIAccess::VertexOrIndexBuffer, ERHIAccess::UAVCompute));
 		FComputeShaderUtils::DispatchIndirect(RHICmdList, ComputeShader, Params, RenderingViewResources->IndirectDrawBuffer->Buffer, IndirectDispatchArgsOffset);
-		RHICmdList.Transition(FRHITransitionInfo(RenderingViewResources->IndexBuffer->UAV, ERHIAccess::UAVCompute, ERHIAccess::VertexOrIndexBuffer));
+		RHICmdList.Transition(FRHITransitionInfo(RenderingViewResources->IndexBuffer->UAV, ERHIAccess::UAVCompute, ERHIAccess::SRVMask | ERHIAccess::VertexOrIndexBuffer));
 	}
 }
 
