Niagara Readiness Settings.

Moves PSO precaching controll for Niagara to be data driven. Can be set per project and FX Type currently.

#virtualized
#Synced-CL 50186389

#Synced-CL 50186389

diff --git a/Engine/Plugins/FX/Niagara/Config/DefaultNiagara.ini b/Engine/Plugins/FX/Niagara/Config/DefaultNiagara.ini
--- a/Engine/Plugins/FX/Niagara/Config/DefaultNiagara.ini
+++ b/Engine/Plugins/FX/Niagara/Config/DefaultNiagara.ini
@@ -73,6 +73,9 @@ BakeQualityLevel=Cinematic
 +QualityLevels=NSLOCTEXT("[/Script/Niagara]", "1E2BA6CD40ADCA8512BC70AC59AACAD7", "Cinematic")
 ComponentRendererWarningsPerClass=(("PostProcessComponent", NSLOCTEXT("[/Script/Niagara]", "022421AD4D0A1767041BEE86A524E12D", "The post process component has a separate \"Enabled\" flag in PostProcessVolume that you need to overwrite\n if you want to disable a component when a particle dies. Otherwise it will still affect the scene\n when it goes back into the component pool.")))
 
+ReadinessSettings=/Niagara/Settings/BaseReadinessSettings_Burst.BaseReadinessSettings_Burst
+DefaultReadinessProfile=(TagName="Niagara.ScalabilityProfile.Critical")
+
 [/Script/NiagaraEditor.NiagaraEditorSettings]
 DefaultSystem=None
 DefaultEmitter=/Niagara/DefaultAssets/DefaultEmitter.DefaultEmitter
diff --git a/Engine/Plugins/FX/Niagara/Source/Niagara/Classes/NiagaraEffectType.h b/Engine/Plugins/FX/Niagara/Source/Niagara/Classes/NiagaraEffectType.h
--- a/Engine/Plugins/FX/Niagara/Source/Niagara/Classes/NiagaraEffectType.h
+++ b/Engine/Plugins/FX/Niagara/Source/Niagara/Classes/NiagaraEffectType.h
@@ -9,6 +9,9 @@
 #include "NiagaraValidationRule.h"
 #include "NiagaraValidationRuleSet.h"
 #include "Particles/ParticlePerfStats.h"
+
+#include "NiagaraScalabilitySettings.h"
+
 #include "NiagaraEffectType.generated.h"
 
 #define DEBUG_SCALABILITY_STATE (!UE_BUILD_SHIPPING)
@@ -443,6 +446,14 @@ class UNiagaraEffectType : public UObject
 	
 	UPROPERTY(EditAnywhere, Category = "Scalability")
 	FNiagaraEmitterScalabilitySettingsArray EmitterScalabilitySettings;
+	
+	/** Readiness Settings used by systems using this Effect Type. If unset, the defaults in the project settings will be used. */
+	UPROPERTY(EditAnywhere, Category = "Scalability")
+	TObjectPtr<UNiagaraSystemReadinessSettings> ReadinessSettings;
+
+	/** How urgent it is that System's using this Effect Type are ready in time for activation and rendering. Settings controlling readiness and possible delays in rendering can be found in ReadinessSettings. */
+	UPROPERTY(EditAnywhere, Category = "Scalability", meta = (DisplayInScalabilityValuesBar, Categories = "Niagara.ScalabilityProfile"))
+	FGameplayTag ReadinessProfile = TAG_Niagara_ScalabilityProfile_Medium;
 
 #if WITH_EDITORONLY_DATA
 	/** A set of rules to apply when checking content. To create your own rules, write a custom class that extends UNiagaraValidationRule. */
@@ -482,6 +493,7 @@ class UNiagaraEffectType : public UObject
 	template<typename TAction>
 	void ForEachPlatformSet(TAction Func);
 
+	const FNiagaraReadinessSettingsProfile* FindReadinessSettingsProfile(FGameplayTag Profile)const;
 private:
 
 	/** Controls generation of performance baseline data for this effect type. */
diff --git a/Engine/Plugins/FX/Niagara/Source/Niagara/Classes/NiagaraEmitter.h b/Engine/Plugins/FX/Niagara/Source/Niagara/Classes/NiagaraEmitter.h
--- a/Engine/Plugins/FX/Niagara/Source/Niagara/Classes/NiagaraEmitter.h
+++ b/Engine/Plugins/FX/Niagara/Source/Niagara/Classes/NiagaraEmitter.h
@@ -334,6 +334,13 @@ struct FVersionedNiagaraEmitterData
 
 	UPROPERTY(EditAnywhere, BlueprintReadWrite, Category = "Emitter", meta = (SegmentedDisplay))
 	ENiagaraSimTarget SimTarget = ENiagaraSimTarget::CPUSim;
+	
+	// Determines how important this emitter is for the runtime and rendering of the system.
+	// An emitter with low priority does not block system execution or other emitters from rendering if its own materials are not yet done compiling.
+	// A critical emitter can block both activation and rendering of the system until it's ready.
+	// See also: readiness settings in the effect type.
+	UPROPERTY(EditAnywhere, BlueprintReadWrite, Category = "Emitter")
+	ENiagaraEmitterImportance Importance = ENiagaraEmitterImportance::Normal;
 
 	/**
 	How should we calculate bounds for the emitter.
diff --git a/Engine/Plugins/FX/Niagara/Source/Niagara/Classes/NiagaraEmitterInstance.h b/Engine/Plugins/FX/Niagara/Source/Niagara/Classes/NiagaraEmitterInstance.h
--- a/Engine/Plugins/FX/Niagara/Source/Niagara/Classes/NiagaraEmitterInstance.h
+++ b/Engine/Plugins/FX/Niagara/Source/Niagara/Classes/NiagaraEmitterInstance.h
@@ -82,6 +82,8 @@ public:
 	void SetFixedBounds(const FBox& InLocalBounds);
 	[[nodiscard]] FBox GetFixedBounds() const;
 	
+	[[nodiscard]] bool IsWaitingForPSOCaching() const { return bWaitingForPSO; }
+	
 #if WITH_EDITORONLY_DATA
 	uint32 GetTickTimeCycles() const { return TickTimeCycles; }
 	NIAGARA_API bool IsDisabledFromIsolation() const;
@@ -229,6 +231,7 @@ protected:
 	ENiagaraExecutionState				ExecutionState = ENiagaraExecutionState::Active;
 
 	bool								bCachedBoundsDynamic = false;
+	bool								bWaitingForPSO = false;
 	FBox								CachedBounds = FBox(ForceInit);
 	FBox								CachedSystemFixedBounds = FBox(ForceInit);
 
diff --git a/Engine/Plugins/FX/Niagara/Source/Niagara/Classes/NiagaraSystem.h b/Engine/Plugins/FX/Niagara/Source/Niagara/Classes/NiagaraSystem.h
--- a/Engine/Plugins/FX/Niagara/Source/Niagara/Classes/NiagaraSystem.h
+++ b/Engine/Plugins/FX/Niagara/Source/Niagara/Classes/NiagaraSystem.h
@@ -837,6 +837,9 @@ public:
 	NIAGARA_API const FNiagaraSystemScalabilityOverride& GetCurrentOverrideSettings() const;
 	inline bool NeedsSortedSignificanceCull()const{ return bNeedsSortedSignificanceCull; }
 
+	const FNiagaraReadinessSettings& GetReadinessSettings() const;
+	ENiagaraEmitterImportance GetHighestEmitterImportance() const { return MaxEmitterImportance; }
+
 	inline void GetMaxInstanceCounts(int32& OutSystemInstanceMax, int32& OutFXTypeInstanceMax, bool bBudgetAdjusted)const;
 	
 	FNiagaraPlatformSet& GetScalabilityPlatformSet() { return Platforms; }
@@ -943,7 +946,7 @@ private:
 	NIAGARA_API void UpdateHasGPUEmitters();
 
 	void PostLoadScriptRuntimeData();
-
+	
 protected:
 	/** An effect types defines settings shared between systems, for example scalability and validation rules.
 	 * Things like environment fx usually have a different effect type than gameplay relevant fx such as weapon impacts.
@@ -1110,6 +1113,9 @@ protected:
 	TObjectPtr<UNiagaraBakerSettings> BakerGeneratedSettings;
 #endif
 
+	UPROPERTY()
+	ENiagaraEmitterImportance MaxEmitterImportance = ENiagaraEmitterImportance::Low;
+	
 	/** Array of emitter indices sorted by execution priority. The emitters will be ticked in this order. Please note that some indices may have the top bit set (kStartNewOverlapGroupBit)
 	* to indicate synchronization points in parallel execution, so mask it out before using the values as indices in the emitters array.
 	*/
@@ -1172,7 +1178,7 @@ protected:
 
 	UPROPERTY()
 	uint32 bSystemStateFastPathEnabled : 1 = true;
-
+	
 #if WITH_EDITORONLY_DATA
 	/* In some cases we can not run validation as the data has been stripped.  This is a temporary workaround to skip validation in those cases. */
 	uint32 bAllowValidation : 1 = true;
@@ -1202,7 +1208,7 @@ protected:
 
 	UPROPERTY()
 	FInstancedStruct SystemStateDataStruct;
-
+	
 	FRenderCommandFence WaitRenderCommandsFence;
 };
 
diff --git a/Engine/Plugins/FX/Niagara/Source/Niagara/Internal/NiagaraEmitterInstanceImpl.h b/Engine/Plugins/FX/Niagara/Source/Niagara/Internal/NiagaraEmitterInstanceImpl.h
--- a/Engine/Plugins/FX/Niagara/Source/Niagara/Internal/NiagaraEmitterInstanceImpl.h
+++ b/Engine/Plugins/FX/Niagara/Source/Niagara/Internal/NiagaraEmitterInstanceImpl.h
@@ -114,6 +114,7 @@ public:
 
 private:
 	void CheckForErrors();
+	void UpdatePSOStatus();
 
 	void BuildConstantBufferTable(
 		const FNiagaraScriptExecutionContext& ExecContext,
diff --git a/Engine/Plugins/FX/Niagara/Source/Niagara/Private/NiagaraComponent.cpp b/Engine/Plugins/FX/Niagara/Source/Niagara/Private/NiagaraComponent.cpp
--- a/Engine/Plugins/FX/Niagara/Source/Niagara/Private/NiagaraComponent.cpp
+++ b/Engine/Plugins/FX/Niagara/Source/Niagara/Private/NiagaraComponent.cpp
@@ -1239,7 +1239,8 @@ bool UNiagaraComponent::InitializeSystem()
 			}
 		}
 
-		if (UPSOPrecacheSettingsManager::GetSettings().ParticlePrecachingTime > 1)
+		const FNiagaraReadinessSettings& ReadinessSettings = GetReadinessSettings();
+		if (ReadinessSettings.PSOPrecacheSettings.PrecachingTime == ENiagaraPSOPrecachingTime::OnInitialize)
 		{
 			PrecachePSOs();
 		}
@@ -1349,7 +1350,10 @@ void UNiagaraComponent::ActivateInternal(bool bReset /* = false */, bool bIsScal
 	}
 #endif
 	
-	if (!Asset->IsReadyToRun() || !IsWorldReadyToRun())
+	bool bIsReady = Asset->IsReadyToRun() && IsWorldReadyToRun();
+	bIsReady &= CheckPSOCacheIsReadyToRun();
+
+	if (!bIsReady)
 	{
 		bAwaitingActivationDueToNotReady = true;
 		bActivateShouldResetWhenReady = bReset;
@@ -1537,6 +1541,7 @@ void UNiagaraComponent::Deactivate()
 void UNiagaraComponent::DeactivateInternal(bool bIsScalabilityCull /* = false */)
 {
 	bool bWasCulledByScalabiltiy = bIsCulledByScalability;
+	bAwaitingActivationDueToNotReady = false;
 
 	if (bIsScalabilityCull)
 	{
@@ -1757,7 +1762,8 @@ void UNiagaraComponent::PostSystemTick_GameThread()
 
 	if (bRecachePSOs)
 	{
-		if (UPSOPrecacheSettingsManager::GetSettings().ParticlePrecachingTime > 1)
+		const FNiagaraReadinessSettings& ReadinessSettings = GetReadinessSettings();
+		if (ReadinessSettings.PSOPrecacheSettings.PrecachingTime != ENiagaraPSOPrecachingTime::None && UPSOPrecacheSettingsManager::GetSettings().ParticlePrecachingTime > 1)
 		{
 			PrecachePSOs();
 		}
@@ -2443,23 +2449,29 @@ FPrimitiveSceneProxy* UNiagaraComponent::CreateSceneProxy()
 	if (Asset != nullptr)
 	{
 		// If we are relying on proxy creation PSOs only then call precache PSOs here because asset doesn't have collected PSO precache data
-		if (UPSOPrecacheSettingsManager::GetSettings().ParticlePrecachingTime == 1)
+		const FNiagaraReadinessSettings& ReadinessSettings = GetReadinessSettings();
+		if (ReadinessSettings.PSOPrecacheSettings.PrecachingTime == ENiagaraPSOPrecachingTime::OnRenderStateCreation)
 		{
 			PrecachePSOs();
 		}
 
-		// If PSOs not precached yet then rely on PSOs reuqested by the shared asset itself - this will make sure
+		// If PSOs not precached yet then rely on PSOs requested by the shared asset itself - this will make sure
 		// those pending requests are boosted if still pending
 		if (!PSOPrecacheComponentData.IsPSOPrecacheCalled())
 		{
 			PrecacheAssetPSOs(Asset);
 		} 		
 
-		EPSOPrecachePriority PSOPrecachePriority = NiagaraComponentPrivate::GBoostPSOPrecacheHighest ? EPSOPrecachePriority::Highest : EPSOPrecachePriority::High;
-		if (CheckPSOPrecachingAndBoostPriority(PSOPrecachePriority) && GetPSOPrecacheProxyCreationStrategy() != EPSOPrecacheProxyCreationStrategy::AlwaysCreate)
+		//Settings can opt to not delay rendering and just render the default material. In most cases we'll delay here until the PSO is ready.
+		bool bCanDelayRender = ReadinessSettings.PSOPrecacheSettings.DelayMode == ENiagaraSystemPSOPrecachingDelayMode::DelayRendering;
+		EPSOPrecachePriority PSOPrecachePriority = ReadinessSettings.PSOPrecacheSettings.GetPriorityBoost();
+		if (bCanDelayRender)
 		{
-			UE_LOG(LogNiagara, Verbose, TEXT("Skipping CreateSceneProxy for UNiagaraComponent %s (UNiagaraSystem PSOs are still compiling)"), *GetFullName());
-			return nullptr;
+			if (CheckPSOPrecachingAndBoostPriority(PSOPrecachePriority) && GetPSOPrecacheProxyCreationStrategy() != EPSOPrecacheProxyCreationStrategy::AlwaysCreate)
+			{
+				UE_LOG(LogNiagara, Verbose, TEXT("Skipping CreateSceneProxy for UNiagaraComponent %s (UNiagaraSystem PSOs are still compiling)"), *GetFullName());
+				return nullptr;
+			}
 		}
 	}
 #endif // UE_WITH_PSO_PRECACHING
@@ -4649,6 +4661,98 @@ void UNiagaraComponent::UpdateInstanceSoloMode()
 	}
 }
 
+const FNiagaraReadinessSettings& UNiagaraComponent::GetReadinessSettings() const
+{
+	if (UNiagaraSystem* Sys = GetAsset())
+	{
+		return Sys->GetReadinessSettings();
+	}
+
+	static FNiagaraReadinessSettings Dummy;
+	return Dummy;
+}
+
+ENiagaraEmitterImportance UNiagaraComponent::GetHighestEmitterImportance() const
+{
+	if (UNiagaraSystem* Sys = GetAsset())
+	{
+		return Sys->GetHighestEmitterImportance();
+	}
+	return ENiagaraEmitterImportance::Low;
+}
+
+bool UNiagaraComponent::CheckEmitterPSOStillRunning() const
+{
+	if (UNiagaraSystem* Sys = GetAsset())
+	{
+		FMaterialInterfacePSOPrecacheParamsList PSOPrecacheList;
+		for (const FNiagaraEmitterHandle& EmitterHandle : Sys->GetEmitterHandles())
+		{
+			if (EmitterHandle.GetIsEnabled())
+			{
+				if (FVersionedNiagaraEmitterData* Data = EmitterHandle.GetEmitterData())
+				{
+					if (Data->Importance == ENiagaraEmitterImportance::Critical)
+					{
+						EmitterHandle.ForEachEnabledRendererWithIndex([&PSOPrecacheList](const UNiagaraRendererProperties* Properties, int32)
+						{
+							Properties->CollectPSOPrecacheData(nullptr, PSOPrecacheList);
+						});
+					}
+				}
+			}
+		}
+		for (FMaterialInterfacePSOPrecacheParams& PSOPrecacheParams : PSOPrecacheList)
+		{
+			if (Sys->IsLoadingPSOCache(PSOPrecacheParams.MaterialInterface))
+			{
+				return true;
+			}
+		}
+	}
+	return false;
+}
+
+bool UNiagaraComponent::CheckPSOCacheIsReadyToRun()
+{
+#if UE_WITH_PSO_PRECACHING
+	const FNiagaraReadinessSettings& ReadinessSettings = GetReadinessSettings();
+	ENiagaraEmitterImportance EmitterImportance = GetHighestEmitterImportance();
+	bool bDelayIfPSONotReady = ReadinessSettings.PSOPrecacheSettings.DelayMode == ENiagaraSystemPSOPrecachingDelayMode::DelayActivation;
+	bool bPrecachePSOs = bDelayIfPSONotReady || ReadinessSettings.PSOPrecacheSettings.PrecachingTime == ENiagaraPSOPrecachingTime::OnActivated || EmitterImportance == ENiagaraEmitterImportance::Critical;	
+	if (bPrecachePSOs)
+	{
+		PrecachePSOs();
+	
+		// check if the system as a whole needs to have all pso jobs finished before it can be activated
+		if (bDelayIfPSONotReady)
+		{
+			//If we're going to delay here and boost our PSO pri, ensure we also include the asset's PSO jobs.
+			PrecacheAssetPSOs(Asset);
+
+			EPSOPrecachePriority PSOPrecachePriorityBoost = ReadinessSettings.PSOPrecacheSettings.GetPriorityBoost();
+			if (GetPSOPrecacheProxyCreationStrategy() != EPSOPrecacheProxyCreationStrategy::AlwaysCreate && CheckPSOPrecachingAndBoostPriority(PSOPrecachePriorityBoost))
+			{
+				UE_LOG(LogNiagara, Verbose, TEXT("Skipping Niagara Activation for UNiagaraComponent %s (%s) (UNiagaraSystem PSOs are still compiling)"), *GetFullName(), *GetFullNameSafe(Asset));
+				return false;
+			}
+		}
+		// check if we have one or more critical emitters waiting to be activated
+		else if (EmitterImportance == ENiagaraEmitterImportance::Critical)
+		{
+			PrecacheAssetPSOs(Asset);
+			if (GetPSOPrecacheProxyCreationStrategy() != EPSOPrecacheProxyCreationStrategy::AlwaysCreate && CheckEmitterPSOStillRunning())
+			{
+				UE_LOG(LogNiagara, Verbose, TEXT("Skipping Niagara Activation for UNiagaraComponent %s (%s) (PSOs are still compiling for critical emitters)"), *GetFullName(), *GetFullNameSafe(Asset));
+				return false;
+			}
+		}
+	}
+#endif
+	
+	return true;
+}
+
 void UNiagaraComponent::SetGpuComputeDebug(bool bEnableDebug)
 {
 	if (bEnableGpuComputeDebug != bEnableDebug)
diff --git a/Engine/Plugins/FX/Niagara/Source/Niagara/Private/NiagaraEffectType.cpp b/Engine/Plugins/FX/Niagara/Source/Niagara/Private/NiagaraEffectType.cpp
--- a/Engine/Plugins/FX/Niagara/Source/Niagara/Private/NiagaraEffectType.cpp
+++ b/Engine/Plugins/FX/Niagara/Source/Niagara/Private/NiagaraEffectType.cpp
@@ -123,6 +123,19 @@ const FNiagaraEmitterScalabilitySettings& UNiagaraEffectType::GetActiveEmitterSc
 	return Dummy;
 }
 
+const FNiagaraReadinessSettingsProfile* UNiagaraEffectType::FindReadinessSettingsProfile(FGameplayTag Profile)const
+{
+	FGameplayTag ProfileToUse = Profile.IsValid() ? Profile : ReadinessProfile;
+	if (ReadinessSettings)
+	{
+		return ReadinessSettings->Find(ProfileToUse);
+	}
+	else
+	{
+		return GetDefault<UNiagaraSettings>()->FindReadinessSettingsProfile(ProfileToUse);
+	}
+}
+
 #if WITH_PER_FXTYPE_PARTICLE_PERF_STATS
 void UNiagaraEffectType::PopulatePerfStatNames()
 {
diff --git a/Engine/Plugins/FX/Niagara/Source/Niagara/Private/NiagaraEmitterInstanceImpl.cpp b/Engine/Plugins/FX/Niagara/Source/Niagara/Private/NiagaraEmitterInstanceImpl.cpp
--- a/Engine/Plugins/FX/Niagara/Source/Niagara/Private/NiagaraEmitterInstanceImpl.cpp
+++ b/Engine/Plugins/FX/Niagara/Source/Niagara/Private/NiagaraEmitterInstanceImpl.cpp
@@ -252,7 +252,7 @@ void FNiagaraEmitterInstanceImpl::Init(int32 InEmitterIndex)
 		ExecutionState = ENiagaraExecutionState::Disabled;
 		return;
 	}
-
+	
 	const TArray<TSharedRef<const FNiagaraEmitterCompiledData>>& EmitterCompiledDatas = ParentSystemInstance->GetSystem()->GetEmitterCompiledData();
 	if (EmitterCompiledDatas.IsValidIndex(EmitterIndex) == false)
 	{
@@ -440,7 +440,9 @@ void FNiagaraEmitterInstanceImpl::Init(int32 InEmitterIndex)
 				GPUExecContext->CombinedParamStore.Bind(&RendererBindings);
 			}
 		}
-	}	
+	}
+
+	UpdatePSOStatus();
 
 	MaxInstanceCount = EmitterData->GetMaxInstanceCount();
 	ParticleDataSet->SetMaxInstanceCount(MaxInstanceCount);
@@ -616,6 +618,34 @@ void FNiagaraEmitterInstanceImpl::CheckForErrors()
 	}
 }
 
+void FNiagaraEmitterInstanceImpl::UpdatePSOStatus()
+{
+	bWaitingForPSO = false;
+#if UE_WITH_PSO_PRECACHING
+	UNiagaraSystem* ParentSystem = ParentSystemInstance->GetSystem();
+	const FNiagaraEmitterHandle& EmitterHandle = GetEmitterHandle();
+	if (FVersionedNiagaraEmitterData* EmitterData = EmitterHandle.GetEmitterData())
+	{
+		FMaterialInterfacePSOPrecacheParamsList PSOPrecacheList;
+		EmitterHandle.ForEachEnabledRendererWithIndex([&PSOPrecacheList, this](const UNiagaraRendererProperties* Properties, int32)
+		{
+			Properties->CollectPSOPrecacheData(this, PSOPrecacheList);
+		});
+		for (FMaterialInterfacePSOPrecacheParams& PSOPrecacheParams : PSOPrecacheList)
+		{
+			if (ParentSystem->IsLoadingPSOCache(PSOPrecacheParams.MaterialInterface))
+			{
+				if (EmitterData->Importance != ENiagaraEmitterImportance::Normal || ParentSystem->GetReadinessSettings().PSOPrecacheSettings.DelayMode != ENiagaraSystemPSOPrecachingDelayMode::None)
+				{
+					bWaitingForPSO = true;
+					break;
+				}
+			}
+		}
+	}
+#endif
+}
+
 void FNiagaraEmitterInstanceImpl::DirtyDataInterfaces()
 {
 	if (IsDisabled())
@@ -967,6 +997,11 @@ void FNiagaraEmitterInstanceImpl::PostTick()
 				}
 			}
 		}
+		
+		if (bWaitingForPSO)
+		{
+			UpdatePSOStatus();
+		}
 
 #if STATS
 		EmitterData->GetStatData().AddStatCapture(TTuple<uint64, ENiagaraScriptUsage>((uint64)this, ENiagaraScriptUsage::ParticleSpawnScript), GetSpawnExecutionContext().ReportStats());
diff --git a/Engine/Plugins/FX/Niagara/Source/Niagara/Private/NiagaraPlatformSet.cpp b/Engine/Plugins/FX/Niagara/Source/Niagara/Private/NiagaraPlatformSet.cpp
--- a/Engine/Plugins/FX/Niagara/Source/Niagara/Private/NiagaraPlatformSet.cpp
+++ b/Engine/Plugins/FX/Niagara/Source/Niagara/Private/NiagaraPlatformSet.cpp
@@ -10,6 +10,7 @@
 #include "NiagaraSystemImpl.h"
 #include "NiagaraSettings.h"
 #include "UObject/UObjectIterator.h"
+#include "NiagaraScalabilitySettings.h"
 
 #include UE_INLINE_GENERATED_CPP_BY_NAME(NiagaraPlatformSet)
 
@@ -1010,6 +1011,14 @@ void FNiagaraPlatformSet::RefreshScalability()
 
 		//We've done a full scalability so clear these out as we wont need to do a targeted refresh.
 		ChangedCVars.Reset();
+
+		//Always refresh our scalability settings. Can we be more targeted here?
+		for (TObjectIterator<UNiagaraSettingsProfileSet> It; It; ++It)
+		{
+			UNiagaraSettingsProfileSet* Settings = *It;
+			check(Settings);
+			Settings->OnScalabilityUpdate();
+		}
 	}
 
 	bool bNeedsRefresh = ChangedCVars.Num() > 0;//Other reasons to refresh deferred like this?
@@ -1018,6 +1027,14 @@ void FNiagaraPlatformSet::RefreshScalability()
 	{
 		InvalidateCachedData();
 
+		//Always refresh our scalability settings. Can we be more targeted here?
+		for (TObjectIterator<UNiagaraSettingsProfileSet> It; It; ++It)
+		{
+			UNiagaraSettingsProfileSet* Settings = *It;
+			check(Settings);
+			Settings->OnScalabilityUpdate();
+		}
+
 		//Find any systems relying on a changed CVar and refresh it's scalability settings.
 		for (TObjectIterator<UNiagaraSystem> It; It; ++It)
 		{
diff --git a/Engine/Plugins/FX/Niagara/Source/Niagara/Private/NiagaraRenderer.cpp b/Engine/Plugins/FX/Niagara/Source/Niagara/Private/NiagaraRenderer.cpp
--- a/Engine/Plugins/FX/Niagara/Source/Niagara/Private/NiagaraRenderer.cpp
+++ b/Engine/Plugins/FX/Niagara/Source/Niagara/Private/NiagaraRenderer.cpp
@@ -464,6 +464,10 @@ bool FNiagaraRenderer::IsRendererEnabled(const UNiagaraRendererProperties* InPro
 			}
 		}
 	}
+	if (Emitter->IsWaitingForPSOCaching())
+	{
+		return false;
+	}
 	return true;
 }
 
diff --git a/Engine/Plugins/FX/Niagara/Source/Niagara/Private/NiagaraScalabilitySettings.cpp b/Engine/Plugins/FX/Niagara/Source/Niagara/Private/NiagaraScalabilitySettings.cpp
new file mode 100644
--- /dev/null
+++ b/Engine/Plugins/FX/Niagara/Source/Niagara/Private/NiagaraScalabilitySettings.cpp
@@ -0,0 +1,70 @@
+// Copyright Epic Games, Inc. All Rights Reserved.
+
+#include "NiagaraScalabilitySettings.h"
+#include "NiagaraSettings.h"
+
+#include UE_INLINE_GENERATED_CPP_BY_NAME(NiagaraScalabilitySettings)
+
+UE_DEFINE_GAMEPLAY_TAG_COMMENT(TAG_Niagara_ScalabilityProfile_Critical, "Niagara.ScalabilityProfile.Critical", "A critical element for Niagara FX that should have the absolute highest fidelity settings in a given context.")
+UE_DEFINE_GAMEPLAY_TAG_COMMENT(TAG_Niagara_ScalabilityProfile_High, "Niagara.ScalabilityProfile.High", "A high importance element for Niagara FX that should favour fidelity over performance in a given context.")
+UE_DEFINE_GAMEPLAY_TAG_COMMENT(TAG_Niagara_ScalabilityProfile_Medium, "Niagara.ScalabilityProfile.Medium", "A medium importance element for Niagara FX that should balance performance and fidelity in a given context.")
+UE_DEFINE_GAMEPLAY_TAG_COMMENT(TAG_Niagara_ScalabilityProfile_Low, "Niagara.ScalabilityProfile.Low", "A low importance element for Niagara FX that should favour performance over fidelity in a given context.")
+
+
+#if WITH_EDITOR
+
+void UNiagaraSystemReadinessSettings::PostEditChangeProperty(FPropertyChangedEvent& PropertyChangedEvent)
+{
+	Super::PostEditChangeProperty(PropertyChangedEvent);
+	if (PropertyChangedEvent.Property && PropertyChangedEvent.Property->GetFName() == GET_MEMBER_NAME_CHECKED(UNiagaraSystemReadinessSettings, Parent))
+	{
+		//Check parent chain for loops.
+		TSet<TObjectPtr<UNiagaraSystemReadinessSettings>> Seen;
+		
+		UNiagaraSystemReadinessSettings* Curr = this;
+		while(IsValid(Curr))
+		{
+			if (Seen.Contains(Curr))
+			{
+				UE_LOG(LogNiagara, Warning, TEXT("Cannot set Parent! Found circular chain of parent settings assets."));
+				Parent = nullptr;
+				break;
+			}
+			Seen.Add(Curr);
+			Curr = Curr->Parent;
+		}
+	}
+}
+
+#endif 
+
+void UNiagaraSystemReadinessSettings::OnScalabilityUpdate()
+{
+	for (auto& Pair : Profiles)
+	{
+		Pair.Value.OnScalabilityUpdate();
+	}
+}
+
+const FNiagaraReadinessSettingsProfile* UNiagaraSystemReadinessSettings::Find(FGameplayTag Profile)const
+{
+	if(const FNiagaraReadinessSettingsProfile* Found = Profiles.Find(Profile))
+	{
+		return Found;
+	}
+
+	if (IsValid(Parent))
+	{
+		// this recursive call would be problematic if users were to set a circular asset chain, but this is checked in UNiagaraSystemReadinessSettings::PostEditChangeProperty
+		return Parent->Find(Profile);
+	}
+
+	return nullptr;
+}
+
+EPSOPrecachePriority FNiagaraSystemPSOPrecacheSettings::GetPriorityBoost()const
+{
+	if (PreCachePriortyBoost == ENiagaraSystemPSOPrecachingPriorityBoost::Highest) return EPSOPrecachePriority::Highest;
+	else if (PreCachePriortyBoost == ENiagaraSystemPSOPrecachingPriorityBoost::High) return EPSOPrecachePriority::High;
+	else return EPSOPrecachePriority::Medium;
+}
\ No newline at end of file
diff --git a/Engine/Plugins/FX/Niagara/Source/Niagara/Private/NiagaraSettings.cpp b/Engine/Plugins/FX/Niagara/Source/Niagara/Private/NiagaraSettings.cpp
--- a/Engine/Plugins/FX/Niagara/Source/Niagara/Private/NiagaraSettings.cpp
+++ b/Engine/Plugins/FX/Niagara/Source/Niagara/Private/NiagaraSettings.cpp
@@ -74,3 +74,14 @@ void UNiagaraSettings::PostInitProperties()
 	}
 #endif
 }
+
+const FNiagaraReadinessSettingsProfile* UNiagaraSettings::FindReadinessSettingsProfile(FGameplayTag Profile) const
+{
+	if (UNiagaraSystemReadinessSettings* ReadinessSettingsPtr = Cast<UNiagaraSystemReadinessSettings>(ReadinessSettings.TryLoad()))
+	{
+		FGameplayTag ProfileToUse = Profile.IsValid() ? Profile : DefaultReadinessProfile;
+		return ReadinessSettingsPtr->Find(ProfileToUse);
+	}
+
+	return nullptr;
+}
diff --git a/Engine/Plugins/FX/Niagara/Source/Niagara/Private/NiagaraSystem.cpp b/Engine/Plugins/FX/Niagara/Source/Niagara/Private/NiagaraSystem.cpp
--- a/Engine/Plugins/FX/Niagara/Source/Niagara/Private/NiagaraSystem.cpp
+++ b/Engine/Plugins/FX/Niagara/Source/Niagara/Private/NiagaraSystem.cpp
@@ -1583,7 +1583,8 @@ void UNiagaraSystem::PostLoad()
 	}
 #endif // WITH_EDITORONLY_DATA
 
-	if (UPSOPrecacheSettingsManager::GetSettings().ParticlePrecachingTime >= 3)
+	const FNiagaraReadinessSettings& CurrentReadinessSettings = GetReadinessSettings();
+	if (CurrentReadinessSettings.PSOPrecacheSettings.bCacheOnLoad)
 	{
 		PrecachePSOs();
 	}
@@ -1806,9 +1807,13 @@ bool UNiagaraSystem::IsReadyToRunInternal() const
 
 bool UNiagaraSystem::IsReadyToRun() const
 {
-	if (PSOPrecacheCompletionEvent.IsValid() && !PSOPrecacheCompletionEvent->IsComplete())
+	const FNiagaraReadinessSettings& CurrentReadinessSettings = GetReadinessSettings();
+	if (CurrentReadinessSettings.PSOPrecacheSettings.DelayMode == ENiagaraSystemPSOPrecachingDelayMode::DelayActivation)
 	{
-		return false;
+		if (PSOPrecacheCompletionEvent.IsValid() && !PSOPrecacheCompletionEvent->IsComplete())
+		{
+			return false;
+		}
 	}
 	return FPlatformProperties::RequiresCookedData() ? bIsReadyToRunCached : IsReadyToRunInternal();
 }
@@ -4153,6 +4158,26 @@ const FNiagaraSystemScalabilityOverride& UNiagaraSystem::GetCurrentOverrideSetti
 	return Dummy;
 }
 
+const FNiagaraReadinessSettings& UNiagaraSystem::GetReadinessSettings()const
+{
+	const FNiagaraReadinessSettingsProfile* SettingsProfile = nullptr;
+	if (UNiagaraEffectType* FXType = GetEffectType())
+	{
+		SettingsProfile = FXType->FindReadinessSettingsProfile(FGameplayTag());
+	}
+	else
+	{
+		SettingsProfile = GetDefault<UNiagaraSettings>()->FindReadinessSettingsProfile(FGameplayTag());
+	}
+
+	if (SettingsProfile)
+	{
+		return SettingsProfile->Get();
+	}
+	static FNiagaraReadinessSettings Fallback;
+	return Fallback;
+}
+
 #if WITH_EDITOR
 void UNiagaraSystem::SetEffectType(UNiagaraEffectType* InEffectType)
 {
@@ -4296,6 +4321,19 @@ void UNiagaraSystem::ResolveScalabilitySettings()
 		};
 		ForEachScript(ScriptUsesSigIndex);
 	}
+	
+	//resolve emitter importance
+	MaxEmitterImportance = ENiagaraEmitterImportance::Low;
+	for (const FNiagaraEmitterHandle& EmitterHandle : GetEmitterHandles())
+	{
+		if (EmitterHandle.GetIsEnabled())
+		{
+			if (FVersionedNiagaraEmitterData* Data = EmitterHandle.GetEmitterData())
+			{
+				MaxEmitterImportance = MaxEmitterImportance < Data->Importance ? MaxEmitterImportance : Data->Importance;
+			}
+		}
+	}
 
 	FNiagaraWorldManager::RequestInvalidateCachedSystemScalabilityDataForAllWorlds();
 }
diff --git a/Engine/Plugins/FX/Niagara/Source/Niagara/Public/NiagaraComponent.h b/Engine/Plugins/FX/Niagara/Source/Niagara/Public/NiagaraComponent.h
--- a/Engine/Plugins/FX/Niagara/Source/Niagara/Public/NiagaraComponent.h
+++ b/Engine/Plugins/FX/Niagara/Source/Niagara/Public/NiagaraComponent.h
@@ -13,6 +13,9 @@
 #include "Particles/ParticlePerfStats.h"
 #include "Particles/ParticleSystemComponent.h"
 
+
+#include "NiagaraSettings.h"
+
 #include "NiagaraComponent.generated.h"
 
 class FMeshElementCollector;
@@ -309,6 +312,15 @@ public:
 private:
 	NIAGARA_API bool RequiresSoloMode() const;
 	NIAGARA_API void UpdateInstanceSoloMode();
+
+
+	const FNiagaraReadinessSettings& GetReadinessSettings() const;
+	ENiagaraEmitterImportance GetHighestEmitterImportance() const;
+	
+	// returns true if the pso job is still running
+	bool CheckEmitterPSOStillRunning() const;
+	bool CheckPSOCacheIsReadyToRun();
+
 public:
 
 	UFUNCTION(BlueprintCallable, Category = Niagara)
diff --git a/Engine/Plugins/FX/Niagara/Source/Niagara/Public/NiagaraScalabilitySettings.h b/Engine/Plugins/FX/Niagara/Source/Niagara/Public/NiagaraScalabilitySettings.h
new file mode 100644
--- /dev/null
+++ b/Engine/Plugins/FX/Niagara/Source/Niagara/Public/NiagaraScalabilitySettings.h
@@ -0,0 +1,207 @@
+// Copyright Epic Games, Inc. All Rights Reserved.
+
+#pragma once 
+
+#include "NiagaraPlatformSet.h"
+
+#include "NativeGameplayTags.h"
+#include "GameplayTagContainer.h"
+#include "Engine/DataAsset.h"
+
+#include "NiagaraScalabilitySettings.generated.h"
+
+enum class EPSOPrecachePriority : uint8;
+
+NIAGARA_API UE_DECLARE_GAMEPLAY_TAG_EXTERN(TAG_Niagara_ScalabilityProfile_Critical);
+NIAGARA_API UE_DECLARE_GAMEPLAY_TAG_EXTERN(TAG_Niagara_ScalabilityProfile_High);
+NIAGARA_API UE_DECLARE_GAMEPLAY_TAG_EXTERN(TAG_Niagara_ScalabilityProfile_Medium);
+NIAGARA_API UE_DECLARE_GAMEPLAY_TAG_EXTERN(TAG_Niagara_ScalabilityProfile_Low);
+
+UENUM()
+enum class ENiagaraSystemPSOPrecachingDelayMode : uint8
+{
+	/** The Niagara System is never delayed and will simulate and render immediately. Can cause incorrect rendering / default material if the PSO is not ready. */
+	None,
+	/** The Niagara System will delay its rendering until the PSO is ready. Simulation will still occur but will not be rendered until the PSO is ready. */
+	DelayRendering,
+	/** The Niagara System will delay its activation until the PSO is ready. No simulation or rendering will occur until the PSO is ready. */
+	DelayActivation,
+};
+
+UENUM()
+enum class ENiagaraEmitterImportance : uint8
+{
+	/** The emitter is critical for the system and blocks both simulation and rendering if it is not fully loaded. */
+	Critical = 0,
+	/** The emitter has default importance and uses the readiness settings of the system/effect type for rendering and simulation delay. */
+	Normal = 1,
+	/** The system is allowed to simulate and render even if this emitter is not yet ready. Can cause pop-ins if a renderer material has finished loading. */
+	Low = 2,
+};
+
+UENUM()
+enum class ENiagaraSystemPSOPrecachingPriorityBoost : uint8
+{
+	/** Niagara will boost the relevant PSO requests to highest priority */
+	Highest,
+	/** Niagara will boost the relevant PSO requests to high priority */
+	High,
+};
+
+UENUM()
+enum class ENiagaraPSOPrecachingTime : uint8
+{
+	/** Do not precache PSOs for this Niagara Component/Emitter. */
+	None,
+	/** Niagara will precache their PSOs when the Component/Emitter is initialized. */
+	OnInitialize,
+	/** Niagara will precache their PSOs when the Component/Emitter is activated. This is implied if the delay mode is set to DelayActivation. */
+	OnActivated,
+	/** Niagara will precache their PSOs when the Component/Emitter render state is created. */
+	OnRenderStateCreation,
+};
+
+USTRUCT()
+struct FNiagaraSettingsProfile
+{
+	GENERATED_BODY()
+};
+
+/** A base asset class for Niagara Settings Profile Sets. */
+UCLASS()
+class UNiagaraSettingsProfileSet : public UDataAsset
+{
+	GENERATED_BODY()
+public:
+
+	/** Handle any required changes to cached data after a scalability settings update. */
+	virtual void OnScalabilityUpdate() {}
+};
+
+
+/** Settings controlling how Niagara Systems can be delayed based on their PSO readiness. */
+USTRUCT()
+struct FNiagaraSystemPSOPrecacheSettings
+{
+	GENERATED_BODY()
+
+	/** How should this Niagara System be delayed to wait for its PSO, if at all. */
+	UPROPERTY(EditAnywhere, Category = "PSO Precache Settings")
+	ENiagaraSystemPSOPrecachingDelayMode DelayMode = ENiagaraSystemPSOPrecachingDelayMode::DelayRendering;
+
+	/** The PSO pre caching priority boost when the FX are used and potentially delayed. */
+	UPROPERTY(EditAnywhere, Category = "PSO Precache Settings")
+	ENiagaraSystemPSOPrecachingPriorityBoost PreCachePriortyBoost = ENiagaraSystemPSOPrecachingPriorityBoost::Highest;
+
+	/** Determines when a Niagara FX should precache its PSOs. */
+	UPROPERTY(EditAnywhere, Category = "PSO Precache Settings")
+	ENiagaraPSOPrecachingTime PrecachingTime = ENiagaraPSOPrecachingTime::OnInitialize;
+
+	/** If true, Niagara Systems will precache their PSOs on load. */
+	UPROPERTY(EditAnywhere, Category = "PSO Precache Settings")
+	uint8 bCacheOnLoad : 1 = false;
+
+	EPSOPrecachePriority GetPriorityBoost() const;
+};
+
+/** A settings profile containing Readiness settings. */
+USTRUCT()
+struct FNiagaraReadinessSettings : public FNiagaraSettingsProfile
+{
+	GENERATED_BODY()
+
+	UPROPERTY(EditAnywhere, Category = "Scalabiltiy")
+	FNiagaraPlatformSet PlatformSet;
+
+	UPROPERTY(EditAnywhere, Category = "Settings")
+	FNiagaraSystemPSOPrecacheSettings PSOPrecacheSettings;
+};
+
+/** Base class for settings objects that can be customized per Platform Set. */
+USTRUCT()
+struct FNiagaraScalabiltiySettingsProfile
+{
+	GENERATED_BODY()
+
+private:
+	mutable int32 CachedPlatformSetIndex = INDEX_NONE;
+
+public:
+	void OnScalabilityUpdate()
+	{
+		//Invalidate our cached index so we'll regenerate on next access.
+		CachedPlatformSetIndex = INDEX_NONE;
+	}
+
+protected:
+	template <typename T>
+	const T& GetInternal(const TArray<T>& Settings)const
+	{		
+		if (Settings.IsValidIndex(CachedPlatformSetIndex))
+		{
+			return Settings[CachedPlatformSetIndex];
+		}
+
+		CachedPlatformSetIndex = INDEX_NONE;
+		for (int32 i = 0; i < Settings.Num(); ++i)
+		{
+			const T& PlatformSettings = Settings[i];
+			if (PlatformSettings.PlatformSet.IsActive())
+			{
+				CachedPlatformSetIndex = i;
+				return PlatformSettings;
+			}
+		}
+		static T Fallback;
+		return Fallback;
+	}
+};
+
+/** Profile containing readiness setings. Customizable by platform. */
+USTRUCT()
+struct FNiagaraReadinessSettingsProfile : public FNiagaraScalabiltiySettingsProfile
+{
+	GENERATED_BODY()
+
+public:
+	
+	/** Get the Readiness settings for the current platform. */
+	const FNiagaraReadinessSettings& Get()const { return GetInternal(Settings); }
+
+protected:
+	/** Readiness separated by Platform Set. */
+	UPROPERTY(EditAnywhere, Category = "Settings", meta = (NiagaraOrthogonalPlatformSets))
+	TArray<FNiagaraReadinessSettings> Settings;
+};
+
+
+/** Profiles defining how Niagara Systems can be delayed when waiting for their prerequisite data to be ready. Also controls priorities for loading/building such data. */
+UCLASS(BlueprintType)
+class UNiagaraSystemReadinessSettings : public UNiagaraSettingsProfileSet
+{
+	GENERATED_BODY()
+
+public:
+
+#if WITH_EDITOR
+	// UObject interface
+	virtual void PostEditChangeProperty(FPropertyChangedEvent& PropertyChangedEvent) override;
+	// End of UObject interface
+#endif
+
+	const FNiagaraReadinessSettingsProfile* Find(FGameplayTag Profile)const;
+
+	/** Handle any required changes to cached data after a scalability settings update. */
+	virtual void OnScalabilityUpdate()override;
+
+protected:
+	
+	/** Parent profile settings. If a user cannot find a requested profile in these settings, we will look up into the parent. */
+	UPROPERTY(EditAnywhere, Category = "Settings")
+	TObjectPtr<UNiagaraSystemReadinessSettings> Parent;
+
+	/** Set of named profiles for settings controlling Readiness. */
+	UPROPERTY(EditAnywhere, Category = "Settings", meta = (Categories = "Niagara.ScalabilityProfile"))
+	TMap<FGameplayTag, FNiagaraReadinessSettingsProfile> Profiles;
+};
+
diff --git a/Engine/Plugins/FX/Niagara/Source/Niagara/Public/NiagaraSettings.h b/Engine/Plugins/FX/Niagara/Source/Niagara/Public/NiagaraSettings.h
--- a/Engine/Plugins/FX/Niagara/Source/Niagara/Public/NiagaraSettings.h
+++ b/Engine/Plugins/FX/Niagara/Source/Niagara/Public/NiagaraSettings.h
@@ -297,11 +297,20 @@ class UNiagaraSettings : public UDeveloperSettings
 	UPROPERTY(config, EditAnywhere, Category = Scalability)
 	TArray<FNiagaraPlatformSetRedirect> PlatformSetRedirects;
 
+	/** Settings profiles controlling how  Niagara Systems can be delayed to deal with latent required data such as PSOs. */
+	UPROPERTY(config, EditAnywhere, Category = Scalability, meta = (AllowedClasses = "/Script/Niagara.NiagaraSystemReadinessSettings"))
+	FSoftObjectPath ReadinessSettings;
+
+	UPROPERTY(config, EditAnywhere, Category = Scalability, meta = (Categories = "Niagara.ScalabilityProfile"))
+	FGameplayTag DefaultReadinessProfile = TAG_Niagara_ScalabilityProfile_Medium;
+
 #if WITH_EDITORONLY_DATA
 	UPROPERTY(config, EditAnywhere, Category = Niagara)
 	bool bGenerateMetaDataOnCompile = false;
 #endif // WITH_EDITORONLY_DATA
 
+	NIAGARA_API const FNiagaraReadinessSettingsProfile* FindReadinessSettingsProfile(FGameplayTag Profile) const;
+
 	// Begin UDeveloperSettings Interface
 	NIAGARA_API virtual FName GetCategoryName() const override;
 #if WITH_EDITOR
diff --git a/Engine/Plugins/FX/Niagara/Source/NiagaraEditor/Private/Customizations/NiagaraPlatformSetCustomization.cpp b/Engine/Plugins/FX/Niagara/Source/NiagaraEditor/Private/Customizations/NiagaraPlatformSetCustomization.cpp
--- a/Engine/Plugins/FX/Niagara/Source/NiagaraEditor/Private/Customizations/NiagaraPlatformSetCustomization.cpp
+++ b/Engine/Plugins/FX/Niagara/Source/NiagaraEditor/Private/Customizations/NiagaraPlatformSetCustomization.cpp
@@ -29,6 +29,35 @@
 
 #define LOCTEXT_NAMESPACE "FNiagaraPlatformSetCustomization"
 
+
+static TSharedPtr<IPropertyHandle> FindChildPlatformSet(TSharedPtr<IPropertyHandle> PropertyHandle)
+{
+	if (FStructProperty* Property = CastField<FStructProperty>(PropertyHandle->GetProperty()))
+	{
+		if (Property->Struct == FNiagaraPlatformSet::StaticStruct())
+		{
+			return PropertyHandle;
+		}
+	}
+
+	// recurse
+	uint32 NumChildren = 0;
+	PropertyHandle->GetNumChildren(NumChildren);
+
+	for (int32 Idx = 0; Idx < static_cast<int32>(NumChildren); ++Idx)
+	{
+		TSharedPtr<IPropertyHandle> Child = PropertyHandle->GetChildHandle(Idx);
+
+		TSharedPtr<IPropertyHandle> ChildResult = FindChildPlatformSet(Child);
+		if (ChildResult.IsValid())
+		{
+			return ChildResult;
+		}
+	}
+
+	return TSharedPtr<IPropertyHandle>();
+}
+
 uint8* GetBaseAddress(TSharedRef<IPropertyHandle> PropertyHandle)
 {
 	if (PropertyHandle->GetNumOuterObjects() > 0)
@@ -95,9 +124,16 @@ void FNiagaraPlatformSetCustomization::CustomizeHeader(TSharedRef<IPropertyHandl
 		int32 ThisIndex = CurrHandle->GetIndexInArray();
 		if (ThisIndex != INDEX_NONE)
 		{
-			PlatformSetArray = CurrHandle->GetParentHandle()->AsArray();
+			PlatformSetArrayProperty = CurrHandle->GetParentHandle();
+			PlatformSetArray = PlatformSetArrayProperty->AsArray();
 			PlatformSetArrayIndex = ThisIndex;
 		}
+		else
+		{
+			PlatformSetArrayProperty.Reset();
+			PlatformSetArray.Reset();
+			PlatformSetArrayIndex = INDEX_NONE;
+		}
 
 		if (FProperty* CurrProperty = CurrHandle->GetProperty())
 		{
@@ -124,6 +160,10 @@ void FNiagaraPlatformSetCustomization::CustomizeHeader(TSharedRef<IPropertyHandl
 					EmitterScalabilityOverrides = reinterpret_cast<FNiagaraEmitterScalabilityOverrides*>(ParentHandle->GetValueBaseAddress(BaseAddress));
 					break;
 				}
+				else if (PlatformSetArray.IsValid() && PlatformSetArrayProperty->HasMetaData("NiagaraOrthogonalPlatformSets"))
+				{
+					break;
+				}
 			}
 		}
 
@@ -304,37 +344,45 @@ void FNiagaraPlatformSetCustomization::UpdateCachedConflicts()
 			PlatformSets.Add(&Settings.Platforms);
 		}
 	}
-
-	CachedConflicts.Reset();
-	FNiagaraPlatformSet::GatherConflicts(PlatformSets, CachedConflicts);
-}
-
-static TSharedPtr<IPropertyHandle> FindChildPlatformSet(TSharedPtr<IPropertyHandle> PropertyHandle)
-{
-	if (FStructProperty* Property = CastField<FStructProperty>(PropertyHandle->GetProperty()))
+	else if (PlatformSetArray.IsValid())
 	{
-		if (Property->Struct == FNiagaraPlatformSet::StaticStruct())
-		{
-			return PropertyHandle;
-		}
-	}
-
-	// recurse
-	uint32 NumChildren = 0;
-	PropertyHandle->GetNumChildren(NumChildren);
+		//Try to handle the general case.
+		FName PlatformSetPropertyName = PropertyHandle->GetProperty()->GetFName();
 
-	for (int32 Idx = 0; Idx < static_cast<int32>(NumChildren); ++Idx)
-	{
-		TSharedPtr<IPropertyHandle> Child = PropertyHandle->GetChildHandle(Idx);
+		auto AddPlatformSetProperty = [&PlatformSets](TSharedPtr<IPropertyHandle> Handle)
+			{
+				FProperty* Prop = Handle->GetProperty();
+				if (const FStructProperty* StructProp = CastField<FStructProperty>(Prop))
+				{
+					if (StructProp->Struct == FNiagaraPlatformSet::StaticStruct())
+					{
+						TArray<void*> Raw;
+						Handle->AccessRawData(Raw);
+						if (Raw.Num() == 1 && Raw[0] != nullptr)
+						{
+							if (FNiagaraPlatformSet* FoundPlatformSet = reinterpret_cast<FNiagaraPlatformSet*>(Raw[0]))
+							{
+								PlatformSets.Add(FoundPlatformSet);
+							}
+						}
+					}
+				}
+			};
 
-		TSharedPtr<IPropertyHandle> ChildResult = FindChildPlatformSet(Child);
-		if (ChildResult.IsValid())
+		uint32 NumElements = 0;
+		PlatformSetArray->GetNumElements(NumElements);
+		for (uint32 i = 0; i < NumElements; ++i)
 		{
-			return ChildResult;
+			TSharedPtr<IPropertyHandle> ElemProperty = PlatformSetArray->GetElement(i);
+			if (TSharedPtr<IPropertyHandle> ElemPlatformSetPropertyHandle = FindChildPlatformSet(ElemProperty))
+			{
+				AddPlatformSetProperty(ElemPlatformSetPropertyHandle);
+			}
 		}
 	}
 
-	return TSharedPtr<IPropertyHandle>();
+	CachedConflicts.Reset();
+	FNiagaraPlatformSet::GatherConflicts(PlatformSets, CachedConflicts);
 }
 
 void FNiagaraPlatformSetCustomization::InvalidateSiblingConflicts() const
diff --git a/Engine/Plugins/FX/Niagara/Source/NiagaraEditor/Private/Customizations/NiagaraPlatformSetCustomization.h b/Engine/Plugins/FX/Niagara/Source/NiagaraEditor/Private/Customizations/NiagaraPlatformSetCustomization.h
--- a/Engine/Plugins/FX/Niagara/Source/NiagaraEditor/Private/Customizations/NiagaraPlatformSetCustomization.h
+++ b/Engine/Plugins/FX/Niagara/Source/NiagaraEditor/Private/Customizations/NiagaraPlatformSetCustomization.h
@@ -77,6 +77,7 @@ private:
 
 private:
 	TArray<FNiagaraPlatformSetConflictInfo> CachedConflicts;
+	TSharedPtr<IPropertyHandle> PlatformSetArrayProperty;
 	TSharedPtr<IPropertyHandleArray> PlatformSetArray;
 	int32 PlatformSetArrayIndex;
 
diff --git a/Engine/Source/Runtime/Engine/Classes/Particles/ParticleSystem.h b/Engine/Source/Runtime/Engine/Classes/Particles/ParticleSystem.h
--- a/Engine/Source/Runtime/Engine/Classes/Particles/ParticleSystem.h
+++ b/Engine/Source/Runtime/Engine/Classes/Particles/ParticleSystem.h
@@ -6,6 +6,7 @@
 #include "CoreMinimal.h"
 #include "UObject/ObjectMacros.h"
 #include "UObject/Object.h"
+#include "UObject/ObjectKey.h"
 #include "UObject/ScriptMacros.h"
 #include "Particles/ParticlePerfStats.h"
 #include "Async/TaskGraphInterfaces.h"
@@ -140,14 +141,16 @@ public:
 	const bool HasLaunchedPSOPrecaching() const { return PSOPrecachingLaunched; }
 	const FGraphEventRef& GetPrecachePSOsEvent() const { return PrecachePSOsEvent; }
 	const TArray<FMaterialPSOPrecacheRequestID>& GetMaterialPSOPrecacheRequestIDs() const { return MaterialPSOPrecacheRequestIDs; }
-
+	ENGINE_API bool IsLoadingPSOCache(UMaterialInterface* Material) const;
+	
 protected:
 	
-	ENGINE_API void LaunchPSOPrecaching(const FMaterialInterfacePSOPrecacheParamsList& VFsPerMaterials);
+	ENGINE_API void LaunchPSOPrecaching(const FMaterialInterfacePSOPrecacheParamsList& PSOPrecacheParamsList);
 
 	bool PSOPrecachingLaunched = false;
 	FGraphEventRef PrecachePSOsEvent;
 	TArray<FMaterialPSOPrecacheRequestID> MaterialPSOPrecacheRequestIDs;
+	TMap<FObjectKey, FGraphEventRef> MaterialPSORequests;
 };
 
 /**
diff --git a/Engine/Source/Runtime/Engine/Classes/Particles/ParticleSystemComponent.h b/Engine/Source/Runtime/Engine/Classes/Particles/ParticleSystemComponent.h
--- a/Engine/Source/Runtime/Engine/Classes/Particles/ParticleSystemComponent.h
+++ b/Engine/Source/Runtime/Engine/Classes/Particles/ParticleSystemComponent.h
@@ -477,6 +477,10 @@ public:
 
 protected:
 	ENGINE_API void PrecacheAssetPSOs(UFXSystemAsset* FXSystemAsset);
+
+private:
+	/** Helper flag to check if PSOs have been precached already */
+	uint8 bPSOPrecacheCalled : 1 = 0;
 };
 
 
diff --git a/Engine/Source/Runtime/Engine/Private/Particles/ParticleSystem.cpp b/Engine/Source/Runtime/Engine/Private/Particles/ParticleSystem.cpp
--- a/Engine/Source/Runtime/Engine/Private/Particles/ParticleSystem.cpp
+++ b/Engine/Source/Runtime/Engine/Private/Particles/ParticleSystem.cpp
@@ -36,6 +36,13 @@ void UFXSystemAsset::LaunchPSOPrecaching(const FMaterialInterfacePSOPrecachePara
 	if (IsComponentPSOPrecachingEnabled())
 	{
 		PrecacheMaterialPSOs(PSOPrecacheParamsList, MaterialPSOPrecacheRequestIDs, PrecachePSOsEvents);
+		for (const FMaterialInterfacePSOPrecacheParams& MaterialPSOPrecacheParams : PSOPrecacheParamsList)
+		{
+			if (MaterialPSOPrecacheParams.MaterialInterface)
+			{
+				MaterialPSORequests.Add(MaterialPSOPrecacheParams.MaterialInterface, PrecachePSOsEvents.Last());
+			}
+		}
 	} 
 
 	// Create task to signal that the PSO precache events are done by adding them as prerequisite to the task.
@@ -57,6 +64,7 @@ void UFXSystemAsset::LaunchPSOPrecaching(const FMaterialInterfacePSOPrecachePara
 				if (UFXSystemAsset* Asset = WeakOwnerAsset.Get())
 				{
 					Asset->PrecachePSOsEvent = nullptr;
+					Asset->MaterialPSORequests.Empty();
 				}
 			}
 
@@ -72,6 +80,15 @@ void UFXSystemAsset::LaunchPSOPrecaching(const FMaterialInterfacePSOPrecachePara
 	PSOPrecachingLaunched = true;
 }
 
+bool UFXSystemAsset::IsLoadingPSOCache(UMaterialInterface* Material) const
+{
+	if (!PSOPrecachingLaunched || MaterialPSORequests.IsEmpty() || !MaterialPSORequests.Contains(Material))
+	{
+		return false;
+	}
+	return MaterialPSORequests[Material]->IsCompleted() == false;
+}
+
 ////////////////////////////////////////////////////////////////////////////////////////////////////
 
 UParticleSystem::UParticleSystem(const FObjectInitializer& ObjectInitializer)
diff --git a/Engine/Source/Runtime/Engine/Private/Particles/ParticleSystemComponent.cpp b/Engine/Source/Runtime/Engine/Private/Particles/ParticleSystemComponent.cpp
--- a/Engine/Source/Runtime/Engine/Private/Particles/ParticleSystemComponent.cpp
+++ b/Engine/Source/Runtime/Engine/Private/Particles/ParticleSystemComponent.cpp
@@ -306,10 +306,11 @@ bool UFXSystemComponent::RequiresLWCTileRecache(const FVector3f CurrentTile, con
 void UFXSystemComponent::PrecacheAssetPSOs(UFXSystemAsset* FXSystemAsset)
 {
 #if UE_WITH_PSO_PRECACHING
-	if (!FApp::CanEverRender() || (!IsComponentPSOPrecachingEnabled()) || FXSystemAsset == nullptr)
+	if (!FApp::CanEverRender() || (!IsComponentPSOPrecachingEnabled()) || FXSystemAsset == nullptr || bPSOPrecacheCalled)
 	{
 		return;
 	}
+	bPSOPrecacheCalled = true;
 
 	FGraphEventRef GraphEvent = FXSystemAsset->GetPrecachePSOsEvent();
 
