- Add early check for invalid virtual shader paths
- This fixes an issue where the virtual path is invalid and the compilation is lost

#Synced-CL 49733395

diff --git a/Engine/Plugins/FX/Niagara/Source/NiagaraEditor/Private/NiagaraHlslTranslator.cpp b/Engine/Plugins/FX/Niagara/Source/NiagaraEditor/Private/NiagaraHlslTranslator.cpp
--- a/Engine/Plugins/FX/Niagara/Source/NiagaraEditor/Private/NiagaraHlslTranslator.cpp
+++ b/Engine/Plugins/FX/Niagara/Source/NiagaraEditor/Private/NiagaraHlslTranslator.cpp
@@ -9807,13 +9807,20 @@ FName FNiagaraHlslTranslator::GetDataInterfaceName(FName BaseName, const FString
 	return BaseName;
 }
 
-FString FNiagaraHlslTranslator::GetFunctionIncludeStatement(const FNiagaraCustomHlslInclude& Include) const
+FString FNiagaraHlslTranslator::GetFunctionIncludeStatement(const FNiagaraCustomHlslInclude& Include)
 {
 	TStringBuilder<128> IncludeStatement;
 
 	if (Include.bIsVirtual)
 	{
-		IncludeStatement.Appendf(TEXT("#include \"%s\"\n"), *Include.FilePath);
+		if ( !CheckVirtualShaderFilePath(Include.FilePath) )
+		{
+			Error(FText::FromString("Invalid virtual include " + Include.FilePath));
+		}
+		else
+		{
+			IncludeStatement.Appendf(TEXT("#include \"%s\"\n"), *Include.FilePath);
+		}
 	}
 	else if (FString FileContents; FFileHelper::LoadFileToString(FileContents, *Include.FilePath))
 	{
diff --git a/Engine/Plugins/FX/Niagara/Source/NiagaraEditor/Private/NiagaraHlslTranslator.h b/Engine/Plugins/FX/Niagara/Source/NiagaraEditor/Private/NiagaraHlslTranslator.h
--- a/Engine/Plugins/FX/Niagara/Source/NiagaraEditor/Private/NiagaraHlslTranslator.h
+++ b/Engine/Plugins/FX/Niagara/Source/NiagaraEditor/Private/NiagaraHlslTranslator.h
@@ -290,7 +290,7 @@ protected:
 	TArray<TArray<FName>> ActiveStageReadTargets;
 
 	void GenerateFunctionCall(ENiagaraScriptUsage ScriptUsage, FNiagaraFunctionSignature& FunctionSignature, TArrayView<const int32> Inputs, TArray<int32>& Outputs);
-	FString GetFunctionIncludeStatement(const FNiagaraCustomHlslInclude& Include) const;
+	FString GetFunctionIncludeStatement(const FNiagaraCustomHlslInclude& Include);
 	FString GetFunctionSignature(const FNiagaraFunctionSignature& Sig);
 
 	void WriteDataSetContextVars(TMap<FNiagaraDataSetID, TMap<int32, FDataSetAccessInfo>>& DataSetAccessInfo, bool bRead, FString& OutHLSLOutput);
