Fixed a possible perf regression caused by trying to load a soft obj ptr for every niagara system.
Also fixed a typo in a uproperty.

#rnx

#Synced-CL 50186389

diff --git a/Engine/Plugins/FX/Niagara/Config/DefaultNiagara.ini b/Engine/Plugins/FX/Niagara/Config/DefaultNiagara.ini
--- a/Engine/Plugins/FX/Niagara/Config/DefaultNiagara.ini
+++ b/Engine/Plugins/FX/Niagara/Config/DefaultNiagara.ini
@@ -72,9 +72,7 @@ BakeQualityLevel=Cinematic
 +QualityLevels=NSLOCTEXT("[/Script/Niagara]", "D56BFF7C4D755EFCB5B350BDE094625A", "Epic")
 +QualityLevels=NSLOCTEXT("[/Script/Niagara]", "1E2BA6CD40ADCA8512BC70AC59AACAD7", "Cinematic")
 ComponentRendererWarningsPerClass=(("PostProcessComponent", NSLOCTEXT("[/Script/Niagara]", "022421AD4D0A1767041BEE86A524E12D", "The post process component has a separate \"Enabled\" flag in PostProcessVolume that you need to overwrite\n if you want to disable a component when a particle dies. Otherwise it will still affect the scene\n when it goes back into the component pool.")))
-
-ReadinessSettings=/Niagara/Settings/BaseReadinessSettings_Burst.BaseReadinessSettings_Burst
-DefaultReadinessProfile=(TagName="Niagara.ScalabilityProfile.Critical")
+DefaultPSOPrecacheSettings=(DelayMode=DelayRendering,PreCachePriorityBoost=Highest,PrecachingTime=OnInitialize,bCacheOnLoad=False)
 
 [/Script/NiagaraEditor.NiagaraEditorSettings]
 DefaultSystem=None
diff --git a/Engine/Plugins/FX/Niagara/Source/Niagara/Classes/NiagaraEffectType.h b/Engine/Plugins/FX/Niagara/Source/Niagara/Classes/NiagaraEffectType.h
--- a/Engine/Plugins/FX/Niagara/Source/Niagara/Classes/NiagaraEffectType.h
+++ b/Engine/Plugins/FX/Niagara/Source/Niagara/Classes/NiagaraEffectType.h
@@ -493,7 +493,7 @@ class UNiagaraEffectType : public UObject
 	template<typename TAction>
 	void ForEachPlatformSet(TAction Func);
 
-	const FNiagaraReadinessSettingsProfile* FindReadinessSettingsProfile(FGameplayTag Profile)const;
+	const FNiagaraReadinessSettingsProfile* FindReadinessSettingsProfile(FGameplayTag Profile = FGameplayTag()) const;
 private:
 
 	/** Controls generation of performance baseline data for this effect type. */
diff --git a/Engine/Plugins/FX/Niagara/Source/Niagara/Private/NiagaraEffectType.cpp b/Engine/Plugins/FX/Niagara/Source/Niagara/Private/NiagaraEffectType.cpp
--- a/Engine/Plugins/FX/Niagara/Source/Niagara/Private/NiagaraEffectType.cpp
+++ b/Engine/Plugins/FX/Niagara/Source/Niagara/Private/NiagaraEffectType.cpp
@@ -130,10 +130,8 @@ const FNiagaraReadinessSettingsProfile* UNiagaraEffectType::FindReadinessSetting
 	{
 		return ReadinessSettings->Find(ProfileToUse);
 	}
-	else
-	{
-		return GetDefault<UNiagaraSettings>()->FindReadinessSettingsProfile(ProfileToUse);
-	}
+	
+	return nullptr;
 }
 
 #if WITH_PER_FXTYPE_PARTICLE_PERF_STATS
diff --git a/Engine/Plugins/FX/Niagara/Source/Niagara/Private/NiagaraScalabilitySettings.cpp b/Engine/Plugins/FX/Niagara/Source/Niagara/Private/NiagaraScalabilitySettings.cpp
--- a/Engine/Plugins/FX/Niagara/Source/Niagara/Private/NiagaraScalabilitySettings.cpp
+++ b/Engine/Plugins/FX/Niagara/Source/Niagara/Private/NiagaraScalabilitySettings.cpp
@@ -64,7 +64,7 @@ const FNiagaraReadinessSettingsProfile* UNiagaraSystemReadinessSettings::Find(FG
 
 EPSOPrecachePriority FNiagaraSystemPSOPrecacheSettings::GetPriorityBoost()const
 {
-	if (PreCachePriortyBoost == ENiagaraSystemPSOPrecachingPriorityBoost::Highest) return EPSOPrecachePriority::Highest;
-	else if (PreCachePriortyBoost == ENiagaraSystemPSOPrecachingPriorityBoost::High) return EPSOPrecachePriority::High;
+	if (PreCachePriorityBoost == ENiagaraSystemPSOPrecachingPriorityBoost::Highest) return EPSOPrecachePriority::Highest;
+	else if (PreCachePriorityBoost == ENiagaraSystemPSOPrecachingPriorityBoost::High) return EPSOPrecachePriority::High;
 	else return EPSOPrecachePriority::Medium;
 }
\ No newline at end of file
diff --git a/Engine/Plugins/FX/Niagara/Source/Niagara/Private/NiagaraSettings.cpp b/Engine/Plugins/FX/Niagara/Source/Niagara/Private/NiagaraSettings.cpp
--- a/Engine/Plugins/FX/Niagara/Source/Niagara/Private/NiagaraSettings.cpp
+++ b/Engine/Plugins/FX/Niagara/Source/Niagara/Private/NiagaraSettings.cpp
@@ -74,14 +74,3 @@ void UNiagaraSettings::PostInitProperties()
 	}
 #endif
 }
-
-const FNiagaraReadinessSettingsProfile* UNiagaraSettings::FindReadinessSettingsProfile(FGameplayTag Profile) const
-{
-	if (UNiagaraSystemReadinessSettings* ReadinessSettingsPtr = Cast<UNiagaraSystemReadinessSettings>(ReadinessSettings.TryLoad()))
-	{
-		FGameplayTag ProfileToUse = Profile.IsValid() ? Profile : DefaultReadinessProfile;
-		return ReadinessSettingsPtr->Find(ProfileToUse);
-	}
-
-	return nullptr;
-}
diff --git a/Engine/Plugins/FX/Niagara/Source/Niagara/Private/NiagaraSystem.cpp b/Engine/Plugins/FX/Niagara/Source/Niagara/Private/NiagaraSystem.cpp
--- a/Engine/Plugins/FX/Niagara/Source/Niagara/Private/NiagaraSystem.cpp
+++ b/Engine/Plugins/FX/Niagara/Source/Niagara/Private/NiagaraSystem.cpp
@@ -4158,23 +4158,18 @@ const FNiagaraSystemScalabilityOverride& UNiagaraSystem::GetCurrentOverrideSetti
 	return Dummy;
 }
 
-const FNiagaraReadinessSettings& UNiagaraSystem::GetReadinessSettings()const
+const FNiagaraReadinessSettings& UNiagaraSystem::GetReadinessSettings() const
 {
-	const FNiagaraReadinessSettingsProfile* SettingsProfile = nullptr;
 	if (UNiagaraEffectType* FXType = GetEffectType())
 	{
-		SettingsProfile = FXType->FindReadinessSettingsProfile(FGameplayTag());
-	}
-	else
-	{
-		SettingsProfile = GetDefault<UNiagaraSettings>()->FindReadinessSettingsProfile(FGameplayTag());
-	}
-
-	if (SettingsProfile)
-	{
-		return SettingsProfile->Get();
+		if (const FNiagaraReadinessSettingsProfile* SettingsProfile = FXType->FindReadinessSettingsProfile())
+		{
+			return SettingsProfile->Get();
+		}
 	}
+	
 	static FNiagaraReadinessSettings Fallback;
+	Fallback.PSOPrecacheSettings = GetDefault<UNiagaraSettings>()->DefaultPSOPrecacheSettings;
 	return Fallback;
 }
 
diff --git a/Engine/Plugins/FX/Niagara/Source/Niagara/Public/NiagaraScalabilitySettings.h b/Engine/Plugins/FX/Niagara/Source/Niagara/Public/NiagaraScalabilitySettings.h
--- a/Engine/Plugins/FX/Niagara/Source/Niagara/Public/NiagaraScalabilitySettings.h
+++ b/Engine/Plugins/FX/Niagara/Source/Niagara/Public/NiagaraScalabilitySettings.h
@@ -91,7 +91,7 @@ struct FNiagaraSystemPSOPrecacheSettings
 
 	/** The PSO pre caching priority boost when the FX are used and potentially delayed. */
 	UPROPERTY(EditAnywhere, Category = "PSO Precache Settings")
-	ENiagaraSystemPSOPrecachingPriorityBoost PreCachePriortyBoost = ENiagaraSystemPSOPrecachingPriorityBoost::Highest;
+	ENiagaraSystemPSOPrecachingPriorityBoost PreCachePriorityBoost = ENiagaraSystemPSOPrecachingPriorityBoost::Highest;
 
 	/** Determines when a Niagara FX should precache its PSOs. */
 	UPROPERTY(EditAnywhere, Category = "PSO Precache Settings")
diff --git a/Engine/Plugins/FX/Niagara/Source/Niagara/Public/NiagaraSettings.h b/Engine/Plugins/FX/Niagara/Source/Niagara/Public/NiagaraSettings.h
--- a/Engine/Plugins/FX/Niagara/Source/Niagara/Public/NiagaraSettings.h
+++ b/Engine/Plugins/FX/Niagara/Source/Niagara/Public/NiagaraSettings.h
@@ -296,21 +296,15 @@ class UNiagaraSettings : public UDeveloperSettings
 
 	UPROPERTY(config, EditAnywhere, Category = Scalability)
 	TArray<FNiagaraPlatformSetRedirect> PlatformSetRedirects;
-
-	/** Settings profiles controlling how  Niagara Systems can be delayed to deal with latent required data such as PSOs. */
-	UPROPERTY(config, EditAnywhere, Category = Scalability, meta = (AllowedClasses = "/Script/Niagara.NiagaraSystemReadinessSettings"))
-	FSoftObjectPath ReadinessSettings;
-
-	UPROPERTY(config, EditAnywhere, Category = Scalability, meta = (Categories = "Niagara.ScalabilityProfile"))
-	FGameplayTag DefaultReadinessProfile = TAG_Niagara_ScalabilityProfile_Medium;
+	
+	UPROPERTY(config, EditAnywhere, Category = Scalability)
+	FNiagaraSystemPSOPrecacheSettings DefaultPSOPrecacheSettings;
 
 #if WITH_EDITORONLY_DATA
 	UPROPERTY(config, EditAnywhere, Category = Niagara)
 	bool bGenerateMetaDataOnCompile = false;
 #endif // WITH_EDITORONLY_DATA
 
-	NIAGARA_API const FNiagaraReadinessSettingsProfile* FindReadinessSettingsProfile(FGameplayTag Profile) const;
-
 	// Begin UDeveloperSettings Interface
 	NIAGARA_API virtual FName GetCategoryName() const override;
 #if WITH_EDITOR
