PR #14336: Remove unused local variable 'NewCompileRequest'

#Synced-CL 50255171

diff --git a/Engine/Plugins/FX/Niagara/Source/Niagara/Private/NiagaraCompilationTypes.cpp b/Engine/Plugins/FX/Niagara/Source/Niagara/Private/NiagaraCompilationTypes.cpp
--- a/Engine/Plugins/FX/Niagara/Source/Niagara/Private/NiagaraCompilationTypes.cpp
+++ b/Engine/Plugins/FX/Niagara/Source/Niagara/Private/NiagaraCompilationTypes.cpp
@@ -183,7 +183,6 @@ public:
 		TaskStartTime = FPlatformTime::Seconds();
 		bForced = Options.bForced;
 
-		FNiagaraSystemAsyncCompileResults NewCompileRequest;
 		INiagaraModule& NiagaraModule = FModuleManager::GetModuleChecked<INiagaraModule>("Niagara");
 
 		CompileRequestHandle = NiagaraModule.RequestCompileSystem(Options.System, bForced, Options.TargetPlatform);
