- Fix shader compiler error when using all the vertex interpolator slots

#Synced-CL 49733395

diff --git a/Engine/Plugins/FX/Niagara/Shaders/Private/NiagaraRibbonVertexFactory.ush b/Engine/Plugins/FX/Niagara/Shaders/Private/NiagaraRibbonVertexFactory.ush
--- a/Engine/Plugins/FX/Niagara/Shaders/Private/NiagaraRibbonVertexFactory.ush
+++ b/Engine/Plugins/FX/Niagara/Shaders/Private/NiagaraRibbonVertexFactory.ush
@@ -222,12 +222,12 @@ struct FVertexFactoryInterpolantsVSToPS
 	float4 DynamicParameter3	: PARTICLE_DYNAMIC_PARAM3;
 #endif
 
-#if NUM_TEX_COORD_INTERPOLATORS
-	float4	TexCoords[(NUM_TEX_COORD_INTERPOLATORS + 1) / 2]	: TEXCOORD3;
+#if NEEDS_PARTICLE_COLOR
+	float4  Color : TEXCOORD0;
 #endif
 
-#if NEEDS_PARTICLE_COLOR
-	float4  Color				: TEXCOORD2;
+#if NUM_TEX_COORD_INTERPOLATORS
+	float4	TexCoords[(NUM_TEX_COORD_INTERPOLATORS + 1) / 2] : TEXCOORD1;
 #endif
 
 #if USE_PARTICLE_POSITION
