- SKM DI: Add child component source mode

#Synced-CL 49733395

diff --git a/Engine/Plugins/FX/Niagara/Source/Niagara/Classes/NiagaraDataInterfaceSkeletalMesh.h b/Engine/Plugins/FX/Niagara/Source/Niagara/Classes/NiagaraDataInterfaceSkeletalMesh.h
--- a/Engine/Plugins/FX/Niagara/Source/Niagara/Classes/NiagaraDataInterfaceSkeletalMesh.h
+++ b/Engine/Plugins/FX/Niagara/Source/Niagara/Classes/NiagaraDataInterfaceSkeletalMesh.h
@@ -329,6 +329,12 @@ enum class ENDISkeletalMesh_SourceMode : uint8
 
 	/** Only use the "Default Mesh" specified. */
 	DefaultMeshOnly,
+
+	/**
+	Search for the component as a child of the Niagara component.
+	It does not recurse the hierarchy, only looking for direct children.
+	*/
+	ChildrenOnly,
 };
 
 UENUM()
diff --git a/Engine/Plugins/FX/Niagara/Source/Niagara/Private/NiagaraDataInterfaceSkeletalMesh.cpp b/Engine/Plugins/FX/Niagara/Source/Niagara/Private/NiagaraDataInterfaceSkeletalMesh.cpp
--- a/Engine/Plugins/FX/Niagara/Source/Niagara/Private/NiagaraDataInterfaceSkeletalMesh.cpp
+++ b/Engine/Plugins/FX/Niagara/Source/Niagara/Private/NiagaraDataInterfaceSkeletalMesh.cpp
@@ -1675,6 +1675,7 @@ USkinnedAsset* UNiagaraDataInterfaceSkeletalMesh::GetSkinnedMesh(FNiagaraSystemI
 	const bool bTrySource = SourceMode == ENDISkeletalMesh_SourceMode::Default || SourceMode == ENDISkeletalMesh_SourceMode::Source;
 	const bool bTryAttachParent = SourceMode == ENDISkeletalMesh_SourceMode::Default || SourceMode == ENDISkeletalMesh_SourceMode::AttachParent;
 	const bool bTryDefaultMesh = SourceMode == ENDISkeletalMesh_SourceMode::Default || SourceMode == ENDISkeletalMesh_SourceMode::DefaultMeshOnly;
+	const bool bTryChildren = SourceMode == ENDISkeletalMesh_SourceMode::ChildrenOnly;
 
 	if (MeshUserParameter.Parameter.IsValid() && InstData && SystemInstance != nullptr)
 	{
@@ -1744,6 +1745,18 @@ USkinnedAsset* UNiagaraDataInterfaceSkeletalMesh::GetSkinnedMesh(FNiagaraSystemI
 			}
 		}
 	}
+	else if (bTryChildren && AttachComponent)
+	{
+		const int32 NumChildren = AttachComponent->GetNumChildrenComponents();
+		for (int32 iChild = 0; iChild < NumChildren; ++iChild)
+		{
+			if (USkinnedMeshComponent* SkinnedComponent = ConvertToValidComponent(AttachComponent->GetChildComponent(iChild)))
+			{
+				FoundSkinnedComponent = SkinnedComponent;
+				break;
+			}				
+		}
+	}
 
 	USkinnedAsset* Mesh = nullptr;
 	SceneComponent = nullptr;
