Very WIP niagara assistant plugin and supporting work base for collaboration.

Misc additions;

Generic Get/Set Variable functions for Component/System Instance controller and paraemter stores.
Allows easier operations in some cases that don't need/want a lookup table of strongly typed get/set functions.

GetPinTypeDefaultValueFromVariant() - Creates a BP pin default value but can handle variants and object refs etc.

Fixed a few locations where UObject types/values were not handled correctly.

Fixed default value string gen for FVector2 types.

#Synced-CL 50257544

diff --git a/Engine/Plugins/FX/Niagara/Source/Niagara/Private/NiagaraComponent.cpp b/Engine/Plugins/FX/Niagara/Source/Niagara/Private/NiagaraComponent.cpp
--- a/Engine/Plugins/FX/Niagara/Source/Niagara/Private/NiagaraComponent.cpp
+++ b/Engine/Plugins/FX/Niagara/Source/Niagara/Private/NiagaraComponent.cpp
@@ -2984,6 +2984,32 @@ void UNiagaraComponent::SetVariableTextureRenderTarget(FName InVariableName, UTe
 #endif
 }
 
+void UNiagaraComponent::SetVariable_InternalUseOnly(const FNiagaraVariableBase& Variable, const FNiagaraVariant& Variant)
+{
+	if (!Variant.IsValid())
+	{
+		return;
+	}
+
+	if (SystemInstanceController.IsValid())
+	{
+		SystemInstanceController->SetVariable_InternalUseOnly_Deferred(Variable, Variant);
+	}
+	else
+	{
+		OverrideParameters.SetParameter_InternalUseOnly(Variable, Variant, true);
+	}
+#if WITH_EDITOR
+	SetParameterOverride(Variable, Variant);
+#endif
+}
+
+void UNiagaraComponent::GetVariable_InternalUseOnly(const FNiagaraVariableBase& Variable, FNiagaraVariant& Variant) const
+{
+	Variant.Reset();
+	OverrideParameters.GetParameter_InternalUseOnly(Variable, Variant);
+}
+
 bool UNiagaraComponent::GetVariableBool(FName InVariableName, bool& bIsValid) const
 {
 	const FNiagaraVariableBase Variable(FNiagaraTypeDefinition::GetBoolDef(), InVariableName);
diff --git a/Engine/Plugins/FX/Niagara/Source/Niagara/Private/NiagaraParameterStore.cpp b/Engine/Plugins/FX/Niagara/Source/Niagara/Private/NiagaraParameterStore.cpp
--- a/Engine/Plugins/FX/Niagara/Source/Niagara/Private/NiagaraParameterStore.cpp
+++ b/Engine/Plugins/FX/Niagara/Source/Niagara/Private/NiagaraParameterStore.cpp
@@ -1453,6 +1453,91 @@ bool FNiagaraParameterStore::SetParameterData(const uint8* Data, FNiagaraVariabl
 	return false;
 }
 
+void FNiagaraParameterStore::SetParameter_InternalUseOnly(const FNiagaraVariableBase& Key, const FNiagaraVariant& Value, bool bAdd)
+{	
+	if (Key.IsDataInterface())
+	{
+		UE_LOG(LogNiagara, Error, TEXT("We do not support setting data interfaces in SetParameter_InternalUseOnly() currently."));
+		return;
+	}
+
+	bool bRebind = false;
+
+	if (bAdd)
+	{
+		bRebind = AddParameter(Key, false, false);
+	}
+
+	if (Key.IsUObject())
+	{
+		SetUObject(Value.GetUObject(), Key);
+	}
+	else if (Key.GetType() == FNiagaraTypeDefinition::GetPositionDef())
+	{
+		FVector WorldPos = FVector::ZeroVector;
+		if (Value.IsValid(FNiagaraTypeHelper::GetVectorDef()))
+		{			
+			WorldPos = Value.GetBytesValue<FVector>(FNiagaraTypeHelper::GetVectorDef());
+		}
+		else if(Value.IsValid(FNiagaraTypeDefinition::GetVec3Def()))
+		{
+			WorldPos = FVector(Value.GetBytesValue<FVector3f>(FNiagaraTypeDefinition::GetVec3Def()));
+		}
+		else
+		{
+			UE_LOG(LogNiagara, Error, TEXT("Malformed Variant for setting position data. Parameter was set to default/origin."));
+		}
+
+		SetPositionParameterValue(WorldPos, Key.GetName(), bAdd);
+	}
+	else if(Value.IsValid(Key.GetType()))
+	{
+		SetParameterData(Value.GetBytes(), Key, bAdd);
+	}
+	
+	if (bRebind)
+	{
+		OnLayoutChange();
+	}
+}
+
+void FNiagaraParameterStore::GetParameter_InternalUseOnly(const FNiagaraVariableBase& Key, FNiagaraVariant& Value)const
+{
+	Value.Reset();
+	if (Key.IsDataInterface())
+	{
+		int32 Index = IndexOf(Key);
+		if (Index != INDEX_NONE)
+		{
+			Value.SetDataInterface(GetDataInterface(Index));
+		}
+	}
+	else if (Key.IsUObject())
+	{
+		int32 Index = IndexOf(Key);
+		if (Index != INDEX_NONE)
+		{
+			Value.SetUObject(GetUObject(Index));
+		}
+	}
+	else if (Key.GetType() == FNiagaraTypeDefinition::GetPositionDef())
+	{
+		FVector PosValue(FVector::ZeroVector);
+		if (const FVector* PosData = GetPositionParameterValue(Key.GetName()))
+		{
+			PosValue = *PosData;
+			Value.SetBytesValue(Key.GetType(), PosValue);
+		}
+	}
+	else
+	{
+		if (const uint8* Data = GetParameterData(Key))
+		{
+			Value.SetBytes(Data, Key.GetSizeInBytes());
+		}
+	}
+}
+
 void FNiagaraParameterStore::OnLayoutChange()
 {
 	const int32 ExpectedSlack = PaddedParameterSize(ParameterData.Num());
diff --git a/Engine/Plugins/FX/Niagara/Source/Niagara/Private/NiagaraSystemInstanceController.cpp b/Engine/Plugins/FX/Niagara/Source/Niagara/Private/NiagaraSystemInstanceController.cpp
--- a/Engine/Plugins/FX/Niagara/Source/Niagara/Private/NiagaraSystemInstanceController.cpp
+++ b/Engine/Plugins/FX/Niagara/Source/Niagara/Private/NiagaraSystemInstanceController.cpp
@@ -387,6 +387,51 @@ void FNiagaraSystemInstanceController::UpdateEmitterMaterials()
 	OnMaterialsUpdatedDelegate.ExecuteIfBound();
 }
 
+void FNiagaraSystemInstanceController::SetVariable_InternalUseOnly_Deferred(const FNiagaraVariableBase& Variable, const FNiagaraVariant& Variant)
+{
+	SetVariable_Deferred(Variable, Variant);
+}
+
+void FNiagaraSystemInstanceController::SetVariable(FNiagaraVariableBase Variable, FNiagaraVariant Variant)
+{
+	if (OverrideParameters && Variant.IsValid())
+	{
+		//TODO: Unify this modification code with the typed setters so it doesn't diverge/become a maintainance headache.
+		bool bIsMaterial = Variable.GetType() == FNiagaraTypeDefinition::GetUMaterialDef();
+		bool bIsStaticMesh = Variable.GetType() == FNiagaraTypeDefinition::GetUStaticMeshDef();
+		bool bUpdateMaterialsIfModified = bIsMaterial || bIsStaticMesh;
+		bool bMarkNeedsTickOnModified = bIsStaticMesh;
+		bool bMarkNeedsRendererRecacheOnModified = bIsStaticMesh;
+		bool bCheckModified = Variant.GetMode() == ENiagaraVariantMode::Object && (bUpdateMaterialsIfModified || bMarkNeedsRendererRecacheOnModified || bMarkNeedsTickOnModified);
+
+		bool bUObjectModified = false;
+		if (bCheckModified)
+		{
+			UObject* CurrValue = OverrideParameters->GetUObject(Variable);
+			bUObjectModified = CurrValue != Variant.GetUObject();
+		}
+
+		OverrideParameters->SetParameter_InternalUseOnly(Variable, Variant, true);
+
+		//Handle any modification updates for some uobject types.
+		if (bUObjectModified)
+		{
+			if (bUpdateMaterialsIfModified)
+			{
+				bNeedsUpdateEmitterMaterials = true;
+			}
+			if (bMarkNeedsTickOnModified)
+			{
+				bNeedsOverrideParametersTicked = true;
+			}
+			if (bMarkNeedsRendererRecacheOnModified)
+			{
+				OnNeedsRendererRecache();
+			}
+		}
+	}
+}
+
 void FNiagaraSystemInstanceController::SetVariable(FName InVariableName, bool InValue)
 {
 	if (OverrideParameters)
diff --git a/Engine/Plugins/FX/Niagara/Source/Niagara/Private/NiagaraVariant.cpp b/Engine/Plugins/FX/Niagara/Source/Niagara/Private/NiagaraVariant.cpp
--- a/Engine/Plugins/FX/Niagara/Source/Niagara/Private/NiagaraVariant.cpp
+++ b/Engine/Plugins/FX/Niagara/Source/Niagara/Private/NiagaraVariant.cpp
@@ -43,6 +43,30 @@ FNiagaraVariant::FNiagaraVariant(const void* InBytes, int32 Size)
 	Object = nullptr;
 }
 
+bool FNiagaraVariant::IsValid(const FNiagaraTypeDefinition& Type)const
+{
+	if (Type.IsDataInterface())
+	{
+		return CurrentMode == ENiagaraVariantMode::DataInterface;
+	}
+	else if (Type.IsUObject())
+	{
+		return CurrentMode == ENiagaraVariantMode::Object;
+	}
+	else
+	{
+		return CurrentMode == ENiagaraVariantMode::Bytes && Bytes.Num() == Type.GetSize();
+	}
+}
+
+void FNiagaraVariant::Reset()
+{
+	CurrentMode = ENiagaraVariantMode::None;
+	Bytes.Reset();
+	DataInterface = nullptr;
+	Object = nullptr;
+}
+
 TObjectPtr<UObject> FNiagaraVariant::GetUObject() const
 {
 	ensure(CurrentMode == ENiagaraVariantMode::Object);
@@ -76,6 +100,7 @@ void FNiagaraVariant::AllocateBytes(int32 InCount)
 	check(InCount > 0);
 	ensure(CurrentMode == ENiagaraVariantMode::None || CurrentMode == ENiagaraVariantMode::Bytes);
 	Bytes.SetNumZeroed(InCount);
+	CurrentMode = ENiagaraVariantMode::Bytes;
 }
 
 void FNiagaraVariant::SetBytes(const uint8* InBytes, int32 InCount)
@@ -112,7 +137,15 @@ bool FNiagaraVariant::operator==(const FNiagaraVariant& Other) const
 				return Object == Other.Object;
 
 			case ENiagaraVariantMode::DataInterface:
-				return GetDataInterface()->Equals(Other.GetDataInterface());
+			{
+				UNiagaraDataInterface* DI = GetDataInterface();
+				if(ensure(DI))
+				{
+					return DI->Equals(Other.GetDataInterface());
+				}
+				return Other.GetDataInterface() == nullptr;
+			}
+			
 
 			case ENiagaraVariantMode::None:
 				return true;
diff --git a/Engine/Plugins/FX/Niagara/Source/Niagara/Public/NiagaraComponent.h b/Engine/Plugins/FX/Niagara/Source/Niagara/Public/NiagaraComponent.h
--- a/Engine/Plugins/FX/Niagara/Source/Niagara/Public/NiagaraComponent.h
+++ b/Engine/Plugins/FX/Niagara/Source/Niagara/Public/NiagaraComponent.h
@@ -615,6 +615,12 @@ public:
 	UFUNCTION(BlueprintCallable, Category = Niagara, meta = (DisplayName = "Get Niagara Variable (Matrix)", Keywords = "user parameter variable matrix"))
 	[[nodiscard]] NIAGARA_API FMatrix GetVariableMatrix(FName InVariableName, bool& bIsValid) const;
 
+	/**For internal use only. Do not use.*/
+	NIAGARA_API void SetVariable_InternalUseOnly(const FNiagaraVariableBase& Variable, const FNiagaraVariant& Variant);
+
+	/**For internal use only. Do not use.*/
+	NIAGARA_API void GetVariable_InternalUseOnly(const FNiagaraVariableBase& Variable, FNiagaraVariant& Variant)const;
+
 	/**
 	Override the parameter collection to use.
 	The order for resolving collections is Component -> System -> World.
diff --git a/Engine/Plugins/FX/Niagara/Source/Niagara/Public/NiagaraParameterStore.h b/Engine/Plugins/FX/Niagara/Source/Niagara/Public/NiagaraParameterStore.h
--- a/Engine/Plugins/FX/Niagara/Source/Niagara/Public/NiagaraParameterStore.h
+++ b/Engine/Plugins/FX/Niagara/Source/Niagara/Public/NiagaraParameterStore.h
@@ -3,6 +3,7 @@
 #pragma once
 
 #include "NiagaraTypes.h"
+#include "NiagaraVariant.h"
 #include "NiagaraParameterStore.generated.h"
 
 // When not cooked, sort by actual name to ensure deterministic cooked data
@@ -638,6 +639,10 @@ public:
 		}
 	}
 
+	//Internal Use Only. Do not use.
+	NIAGARA_API void SetParameter_InternalUseOnly(const FNiagaraVariableBase& Key, const FNiagaraVariant& Value, bool bAdd=false);
+	NIAGARA_API void GetParameter_InternalUseOnly(const FNiagaraVariableBase& Key, FNiagaraVariant& Value)const;
+
 	inline void OnParameterChange() 
 	{ 
 		bParametersDirty = true;
diff --git a/Engine/Plugins/FX/Niagara/Source/Niagara/Public/NiagaraSystemInstanceController.h b/Engine/Plugins/FX/Niagara/Source/Niagara/Public/NiagaraSystemInstanceController.h
--- a/Engine/Plugins/FX/Niagara/Source/Niagara/Public/NiagaraSystemInstanceController.h
+++ b/Engine/Plugins/FX/Niagara/Source/Niagara/Public/NiagaraSystemInstanceController.h
@@ -160,6 +160,8 @@ public:
 	#endif
 	void OnPooledReuse(UWorld& NewWorld) { SystemInstance->OnPooledReuse(NewWorld); }
 
+	//Internal use only. DO NOT USE.
+	void SetVariable_InternalUseOnly_Deferred(const FNiagaraVariableBase& Variable, const FNiagaraVariant& Variant);
 private:
 	using FMaterialOverride = FNiagaraSystemRenderData::FMaterialOverride;
 
@@ -184,11 +186,14 @@ private:
 	void SetVariable(FName InVariableName, TWeakObjectPtr<UTexture> Texture);
 	void SetVariable(FName InVariableName, TWeakObjectPtr<UTextureRenderTarget> TextureRenderTarget);
 
+	void SetVariable(FNiagaraVariableBase Variable, FNiagaraVariant Variant);
+
 	NIAGARA_SYSTEM_INSTANCE_CONTROLLER_SHIM(SetRequestedExecutionState,)
 
 	void OnNeedsRendererRecache() { bNeedsRendererRecache = true; }
 
 private:
+
 	FNiagaraWorldManager* WorldManager = nullptr;
 	FNiagaraUserRedirectionParameterStore* OverrideParameters = nullptr;
 	FNiagaraSystemInstancePtr SystemInstance;
diff --git a/Engine/Plugins/FX/Niagara/Source/Niagara/Public/NiagaraVariant.h b/Engine/Plugins/FX/Niagara/Source/Niagara/Public/NiagaraVariant.h
--- a/Engine/Plugins/FX/Niagara/Source/Niagara/Public/NiagaraVariant.h
+++ b/Engine/Plugins/FX/Niagara/Source/Niagara/Public/NiagaraVariant.h
@@ -3,6 +3,7 @@
 #pragma once
 
 #include "UObject/ObjectMacros.h"
+#include "NiagaraTypes.h"
 
 #include "NiagaraVariant.generated.h"
 
@@ -32,6 +33,8 @@ struct FNiagaraVariant
 	NIAGARA_API explicit FNiagaraVariant(const TArray<uint8>& InBytes);
 	NIAGARA_API FNiagaraVariant(const void* InBytes, int32 Size);
 
+	NIAGARA_API void Reset();
+
 	NIAGARA_API TObjectPtr<UObject> GetUObject() const;
 	NIAGARA_API void SetUObject(TObjectPtr<UObject> InObject);
 
@@ -43,6 +46,11 @@ struct FNiagaraVariant
 	NIAGARA_API uint8* GetBytes() const;
 	NIAGARA_API int32 GetNumBytes() const;
 
+ 	template<typename T> T GetBytesValue(const FNiagaraTypeDefinition& Type)const;
+	template<typename T> void SetBytesValue(const FNiagaraTypeDefinition& Type, const T& Value);
+
+	NIAGARA_API bool IsValid(const FNiagaraTypeDefinition& Type)const;
+	
 	bool IsValid() const { return CurrentMode != ENiagaraVariantMode::None; }
 	ENiagaraVariantMode GetMode() const { return CurrentMode; }
 
@@ -64,3 +72,22 @@ private:
 	UPROPERTY(EditAnywhere, Category=Variant)
 	ENiagaraVariantMode CurrentMode;
 };
+
+template<typename T> T FNiagaraVariant::GetBytesValue(const FNiagaraTypeDefinition& Type)const
+{
+	if (IsValid(Type))
+	{
+		check(sizeof(T) == Type.GetSize());
+		T Value;
+		FMemory::Memcpy(&Value, Bytes.GetData(), sizeof(T));
+		return Value;
+	}
+	return T();
+}
+
+template<typename T> void FNiagaraVariant::SetBytesValue(const FNiagaraTypeDefinition& Type, const T& Value)
+{
+	check(sizeof(T) == Type.GetSize());
+	AllocateBytes(sizeof(T));
+	FMemory::Memcpy(Bytes.GetData(), &Value, sizeof(T));
+}
diff --git a/Engine/Plugins/FX/Niagara/Source/NiagaraBlueprintNodes/Private/NiagaraBlueprintUtil.cpp b/Engine/Plugins/FX/Niagara/Source/NiagaraBlueprintNodes/Private/NiagaraBlueprintUtil.cpp
--- a/Engine/Plugins/FX/Niagara/Source/NiagaraBlueprintNodes/Private/NiagaraBlueprintUtil.cpp
+++ b/Engine/Plugins/FX/Niagara/Source/NiagaraBlueprintNodes/Private/NiagaraBlueprintUtil.cpp
@@ -61,6 +61,11 @@ FEdGraphPinType FNiagaraBlueprintUtil::TypeDefinitionToBlueprintType(const FNiag
 		Category = UEdGraphSchema_K2::PC_Byte;
 		SubCategoryObject = TypeDef.GetEnum();
 	}
+	else if (TypeDef.IsUObject())
+	{
+		Category = UEdGraphSchema_K2::PC_Object;
+		SubCategoryObject = TypeDef.GetClass();
+	}
 	else if (TypeDef.GetStruct())
 	{
 		Category = UEdGraphSchema_K2::PC_Struct;
diff --git a/Engine/Plugins/FX/Niagara/Source/NiagaraEditor/Private/EdGraphSchema_Niagara.cpp b/Engine/Plugins/FX/Niagara/Source/NiagaraEditor/Private/EdGraphSchema_Niagara.cpp
--- a/Engine/Plugins/FX/Niagara/Source/NiagaraEditor/Private/EdGraphSchema_Niagara.cpp
+++ b/Engine/Plugins/FX/Niagara/Source/NiagaraEditor/Private/EdGraphSchema_Niagara.cpp
@@ -52,6 +52,7 @@
 #include "UObject/LinkerLoad.h"
 #include "Algo/Sort.h"
 #include "Config/NiagaraFavoriteActionsConfig.h"
+#include "NiagaraVariant.h"
 
 #include UE_INLINE_GENERATED_CPP_BY_NAME(EdGraphSchema_Niagara)
 
@@ -1717,6 +1718,51 @@ bool UEdGraphSchema_Niagara::TryGetPinDefaultValueFromNiagaraVariable(const FNia
 	return false;
 }
 
+bool UEdGraphSchema_Niagara::TryGetPinDefaultValueFromNiagaraVariant(const FNiagaraVariableBase& Variable, const FNiagaraVariant& Value, FString& OutPinDefaultValue) const
+{
+	bool bValid = Value.IsValid(Variable.GetType());
+	OutPinDefaultValue.Reset();
+	if (Variable.IsDataInterface())
+	{
+		if (bValid)
+		{
+			OutPinDefaultValue = FSoftObjectPath(Value.GetDataInterface()).ToString();
+			return true;
+		}
+	}
+	else if (Variable.IsUObject())
+	{
+		if (bValid)
+		{
+			OutPinDefaultValue = FSoftObjectPath(Value.GetUObject()).ToString();
+			return true;
+		}
+	}
+	else
+	{
+		FNiagaraVariable PinDefaultVariable = Variable;//Just deal with data in the variable since that's what other code deals with
+		if (bValid)
+		{
+			PinDefaultVariable.AllocateData();
+			FMemory::Memcpy(PinDefaultVariable.GetData(), Value.GetBytes(), PinDefaultVariable.GetType().GetSize());
+		}
+		else
+		{
+			FNiagaraEditorUtilities::ResetVariableToDefaultValue(PinDefaultVariable);
+		}
+
+		FNiagaraEditorModule& NiagaraEditorModule = FModuleManager::GetModuleChecked<FNiagaraEditorModule>("NiagaraEditor");
+		TSharedPtr<INiagaraEditorTypeUtilities, ESPMode::ThreadSafe> TypeEditorUtilities = NiagaraEditorModule.GetTypeUtilities(Variable.GetType());
+		if (TypeEditorUtilities.IsValid() && TypeEditorUtilities->CanHandlePinDefaults())
+		{
+			OutPinDefaultValue = TypeEditorUtilities->GetPinDefaultStringFromValue(PinDefaultVariable);
+			return true;
+		}
+	}
+	OutPinDefaultValue = FString();
+	return false;
+}
+
 void UEdGraphSchema_Niagara::ConvertIllegalPinsInPlace(UEdGraphPin* Pin)
 {
 	if (Pin == nullptr)
diff --git a/Engine/Plugins/FX/Niagara/Source/NiagaraEditor/Private/NiagaraExternalSystemEditorUtilities.cpp b/Engine/Plugins/FX/Niagara/Source/NiagaraEditor/Private/NiagaraExternalSystemEditorUtilities.cpp
new file mode 100644
--- /dev/null
+++ b/Engine/Plugins/FX/Niagara/Source/NiagaraEditor/Private/NiagaraExternalSystemEditorUtilities.cpp
@@ -0,0 +1,322 @@
+// Copyright Epic Games, Inc. All Rights Reserved.
+
+#include "NiagaraExternalSystemEditorUtilities.h"
+#include "Stateless/NiagaraStatelessEmitter.h"
+#include "Stateless/NiagaraStatelessEmitterData.h"
+#include "NiagaraSystemEditorData.h"
+#include "NiagaraScriptMergeManager.h"
+#include "ViewModels/Stack/NiagaraStackGraphUtilities.h"
+
+#include "ScopedTransaction.h"
+
+#define LOCTEXT_NAMESPACE "FNiagaraExternalSystemEditorUtilities"
+
+FNiagaraEmitterHandle* FNiagaraExternalSystemEditorUtilities::GetEmitterHandle(UNiagaraSystem& System, FName EmitterName)
+{
+	for (FNiagaraEmitterHandle& EmitterHandle : System.GetEmitterHandles())
+	{
+		if (EmitterHandle.GetName() == EmitterName)
+		{
+			return &EmitterHandle;
+		}
+	}
+	return nullptr;
+}
+
+FNiagaraEmitterHandle* FNiagaraExternalSystemEditorUtilities::GetEmitterHandle(UNiagaraSystem& System, const FGuid& EmitterId)
+{
+	for (FNiagaraEmitterHandle& EmitterHandle : System.GetEmitterHandles())
+	{
+		if (EmitterHandle.GetId() == EmitterId)
+		{
+			return &EmitterHandle;
+		}
+	}
+	return nullptr;
+}
+
+void FNiagaraExternalSystemEditorUtilities::AddUserVariable(UNiagaraSystem& System, const FNiagaraVariableBase& Variable, const FNiagaraVariant& Variant, FString VariableDesc, FText TransactionDescription)
+{
+	if (!Variant.IsValid())
+	{
+		return;
+	}
+
+	const FScopedTransaction Transaction(TransactionDescription);
+	System.Modify();
+
+	System.GetExposedParameters().SetParameter_InternalUseOnly(Variable, Variant, true);
+
+	if (UNiagaraScriptVariable* ScriptVariable = FNiagaraEditorUtilities::UserParameters::GetScriptVariableForUserParameter(Variable, System))
+	{
+		ScriptVariable->Metadata.Description = FText::FromString(VariableDesc);
+	}
+	FPropertyChangedEvent PropertyChangedEvent(nullptr, EPropertyChangeType::ValueSet);
+	System.PostEditChangeProperty(PropertyChangedEvent);
+}
+
+void FNiagaraExternalSystemEditorUtilities::RemoveUserVariable(UNiagaraSystem& System, const FNiagaraVariableBase& Variable, FText TransactionDescription)
+{
+	const FScopedTransaction Transaction(TransactionDescription);
+	System.Modify();
+
+	System.GetExposedParameters().RemoveParameter(Variable);
+
+	FPropertyChangedEvent PropertyChangedEvent(nullptr, EPropertyChangeType::ValueSet);
+	System.PostEditChangeProperty(PropertyChangedEvent);
+}
+
+void FNiagaraExternalSystemEditorUtilities::GetUserVariableInfo(UNiagaraSystem& System, TArray<FNiagaraExternalSystemEditorUtilities::FUserVariableInfo>& OutVariableInfo)
+{
+	const FNiagaraUserRedirectionParameterStore& UserParamStore = System.GetExposedParameters();
+	TArray<FNiagaraVariable> UserVariables;
+	UserParamStore.GetUserParameters(UserVariables);
+
+	OutVariableInfo.Reset(UserVariables.Num());
+	for (const FNiagaraVariable& Var : UserVariables)
+	{
+		FNiagaraExternalSystemEditorUtilities::FUserVariableInfo& NewVarInfo = OutVariableInfo.AddDefaulted_GetRef();
+		NewVarInfo.Variable = Var;
+		const FNiagaraTypeDefinition& Type = Var.GetType();
+
+		if (const FNiagaraVariableWithOffset* StoreVarWithOffset = UserParamStore.FindParameterVariable(Var))
+		{
+			if (Var.IsUObject())
+			{
+				NewVarInfo.Value.SetUObject(UserParamStore.GetUObject(StoreVarWithOffset->Offset));
+			}
+			else if (Var.IsDataInterface())
+			{
+				NewVarInfo.Value.SetDataInterface(UserParamStore.GetDataInterface(StoreVarWithOffset->Offset));
+			}
+			else
+			{
+				check(Var.GetSizeInBytes() == StoreVarWithOffset->GetSizeInBytes());
+				int32 VarSize = Var.GetSizeInBytes();
+				NewVarInfo.Value.SetBytes(UserParamStore.GetParameterData(StoreVarWithOffset->Offset, VarSize), VarSize);
+			}
+		}
+
+		if (UNiagaraScriptVariable* ScriptVariable = FNiagaraEditorUtilities::UserParameters::GetScriptVariableForUserParameter(Var, System))
+		{
+			NewVarInfo.Description = ScriptVariable->Metadata.Description.ToString();
+		}
+	}
+}
+
+void FNiagaraExternalSystemEditorUtilities::GetSystemInfo(UNiagaraSystem& System, FSystemInfo& OutSystemInfo)
+{
+	OutSystemInfo.Emitters.Reset(System.GetEmitterHandles().Num());
+
+	for (const FNiagaraEmitterHandle& EmitterHandle : System.GetEmitterHandles())
+	{
+		FText ErrorText;
+		FEmitterInfo& NewInfo = OutSystemInfo.Emitters.AddDefaulted_GetRef();
+		GetEmitterInfo(EmitterHandle, NewInfo);
+	}
+}
+
+FName FNiagaraExternalSystemEditorUtilities::AddEmitterToSystem(UNiagaraSystem& System, UNiagaraEmitter& InEmitterToAdd, FText TransactionDescription)
+{
+	const FScopedTransaction Transaction(TransactionDescription);
+	System.Modify();
+
+	FGuid NewEmitterId = FNiagaraEditorUtilities::AddEmitterToSystem(System, InEmitterToAdd, InEmitterToAdd.GetExposedVersion().VersionGuid, true);
+	if (FNiagaraEmitterHandle* EmitterHandle = GetEmitterHandle(System, NewEmitterId))
+	{
+		//Try to place emitter overview node in a sensible location.
+		TArray<UNiagaraOverviewNode*> OverviewNodes;
+		UNiagaraSystemEditorData* SystemEditorData = CastChecked<UNiagaraSystemEditorData>(System.GetEditorData(), ECastCheckedType::NullChecked);
+		SystemEditorData->GetSystemOverviewGraph()->GetNodesOfClass<UNiagaraOverviewNode>(OverviewNodes);
+
+		UNiagaraOverviewNode* RightMostNode = nullptr;
+		UNiagaraOverviewNode* NewNode = nullptr;
+		for (UNiagaraOverviewNode* OverviewNode : OverviewNodes)
+		{
+			if (OverviewNode->GetEmitterHandleGuid() == NewEmitterId)
+			{
+				NewNode = OverviewNode;
+				continue;
+			}
+			if (RightMostNode == nullptr)
+			{
+				RightMostNode = OverviewNode;
+				continue;
+			}
+			if (OverviewNode->NodePosX > RightMostNode->NodePosX)
+			{
+				RightMostNode = OverviewNode;
+			}
+		}
+		if (NewNode)
+		{
+			FVector2D RightMostPos(0.0f);
+			if (RightMostNode)
+			{
+				RightMostPos = FVector2D(RightMostNode->NodePosX, RightMostNode->NodePosY);
+			}
+			NewNode->NodePosX = RightMostPos.X + 300;
+			NewNode->NodePosY = RightMostPos.Y;
+		}
+
+		System.OnSystemPostEditChange().Broadcast(&System);
+		return EmitterHandle->GetName();
+	}
+
+	return NAME_None;
+}
+
+bool FNiagaraExternalSystemEditorUtilities::RemoveEmitterFromSystem(UNiagaraSystem& System, FName EmitterToRemove, FText TransactionDescription)
+{
+	if (FNiagaraEmitterHandle* EmitterHandle = GetEmitterHandle(System, EmitterToRemove))
+	{
+		TSet<FGuid> EmitterIdsToRemove;
+		EmitterIdsToRemove.Add(EmitterHandle->GetId());
+		FNiagaraEditorUtilities::KillSystemInstances(System);
+
+		const FScopedTransaction Transaction(TransactionDescription);
+		System.Modify();
+
+		System.RemoveEmitterHandlesById(EmitterIdsToRemove);
+
+		FNiagaraScriptMergeManager::Get()->ClearMergeAdapterCache();
+		FNiagaraStackGraphUtilities::RebuildEmitterNodes(System);
+		UNiagaraSystemEditorData* SystemEditorData = CastChecked<UNiagaraSystemEditorData>(System.GetEditorData(), ECastCheckedType::NullChecked);
+		SystemEditorData->SynchronizeOverviewGraphWithSystem(System);
+
+		System.OnSystemPostEditChange().Broadcast(&System);
+
+		return true;
+	}
+	return false;
+}
+
+void FNiagaraExternalSystemEditorUtilities::GetEmitterInfo(const FNiagaraEmitterHandle& EmitterHandle, FEmitterInfo& OutEmitterInfo)
+{
+	OutEmitterInfo.Name = EmitterHandle.GetName();
+	OutEmitterInfo.Id = EmitterHandle.GetId();
+	OutEmitterInfo.EmitterMode = EmitterHandle.GetEmitterMode();
+	if (EmitterHandle.GetEmitterMode() == ENiagaraEmitterMode::Stateless)
+	{
+		if (const UNiagaraStatelessEmitter* StatelessEmitter = EmitterHandle.GetStatelessEmitter())
+		{
+			if (const FNiagaraStatelessEmitterDataPtr EmitterData = StatelessEmitter ? StatelessEmitter->GetEmitterData() : nullptr)
+			{
+				OutEmitterInfo.bLocalSpace = true;
+				OutEmitterInfo.bDeterministic = EmitterData->bDeterministic;
+				OutEmitterInfo.SimTarget = ENiagaraSimTarget::GPUComputeSim;
+			}
+		}		
+	}
+	else
+	{
+		const FVersionedNiagaraEmitterData* VersionedData = EmitterHandle.GetEmitterData();
+		if (VersionedData)
+		{
+			OutEmitterInfo.bLocalSpace = VersionedData->bLocalSpace;
+			OutEmitterInfo.bDeterministic = VersionedData->bDeterminism;
+			OutEmitterInfo.SimTarget = VersionedData->SimTarget;
+		}
+	}
+}
+
+bool FNiagaraExternalSystemEditorUtilities::GetEmitterInfo(UNiagaraSystem& System, FName EmitterName, FEmitterInfo& OutEmitterInfo, FText& OutErrorMessage)
+{
+	if (FNiagaraEmitterHandle* EmitterHandle = GetEmitterHandle(System, EmitterName))
+	{
+		GetEmitterInfo(*EmitterHandle, OutEmitterInfo);
+		return true;
+	}
+	OutErrorMessage = LOCTEXT("InvalidEmitter", "An emitter with this name does not exist on this System.");
+	return false;
+}
+
+bool FNiagaraExternalSystemEditorUtilities::SetEmitterName(UNiagaraSystem& System, FName EmitterName, FName Name, FText TransactionDescription, FText& OutErrorMessage)
+{
+	if (FNiagaraEmitterHandle* EmitterHandle = GetEmitterHandle(System, EmitterName))
+	{
+		const FScopedTransaction Transaction(TransactionDescription);
+		System.Modify();
+		EmitterHandle->SetName(Name, System);
+		System.OnSystemPostEditChange().Broadcast(&System);
+		return true;
+	}
+	OutErrorMessage = LOCTEXT("InvalidEmitter", "An emitter with this name does not exist on this System.");
+	return false;
+}
+
+bool FNiagaraExternalSystemEditorUtilities::SetEmitterLocalSpace(UNiagaraSystem& System, FName EmitterName, bool bLocalSpace, FText TransactionDescription, FText& OutErrorMessage)
+{
+	if (FNiagaraEmitterHandle* EmitterHandle = GetEmitterHandle(System, EmitterName))
+	{
+		const FScopedTransaction Transaction(TransactionDescription);
+		System.Modify();
+		if (EmitterHandle->GetEmitterMode() == ENiagaraEmitterMode::Stateless)
+		{
+			OutErrorMessage = LOCTEXT("SetStatelessLocalSpace", "Stateless Emitters are always local space.");
+		}
+		else
+		{
+			FVersionedNiagaraEmitterData* VersionedData = EmitterHandle->GetEmitterData();
+			if (VersionedData)
+			{
+				VersionedData->bLocalSpace = bLocalSpace;
+			}
+		}
+		System.OnSystemPostEditChange().Broadcast(&System);
+		return true;
+	}
+	OutErrorMessage = LOCTEXT("InvalidEmitter", "An emitter with this name does not exist on this System.");
+	return false;
+}
+
+bool FNiagaraExternalSystemEditorUtilities::SetEmitterDeterministic(UNiagaraSystem& System, FName EmitterName, bool bDeterministic, FText TransactionDescription, FText& OutErrorMessage)
+{
+	if (FNiagaraEmitterHandle* EmitterHandle = GetEmitterHandle(System, EmitterName))
+	{
+		const FScopedTransaction Transaction(TransactionDescription);
+		System.Modify();
+		if (EmitterHandle->GetEmitterMode() == ENiagaraEmitterMode::Stateless)
+		{
+			OutErrorMessage = LOCTEXT("SetStatelessDeterministic", "Modiftying this value for stateless emitters is not currently supported.");
+		}
+		else
+		{
+			FVersionedNiagaraEmitterData* VersionedData = EmitterHandle->GetEmitterData();
+			if (VersionedData)
+			{
+				VersionedData->bDeterminism = bDeterministic;
+			}
+		}
+		System.OnSystemPostEditChange().Broadcast(&System);
+		return true;
+	}
+	OutErrorMessage = LOCTEXT("InvalidEmitter", "An emitter with this name does not exist on this System.");
+	return false;
+}
+bool FNiagaraExternalSystemEditorUtilities::SetEmitterSimTarget(UNiagaraSystem& System, FName EmitterName, ENiagaraSimTarget SimTarget, FText TransactionDescription, FText& OutErrorMessage)
+{
+	if (FNiagaraEmitterHandle* EmitterHandle = GetEmitterHandle(System, EmitterName))
+	{
+		const FScopedTransaction Transaction(TransactionDescription);
+		System.Modify();
+		if (EmitterHandle->GetEmitterMode() == ENiagaraEmitterMode::Stateless)
+		{
+			OutErrorMessage = LOCTEXT("SetStatelessSimTarget", "Stateless Emitters determine their own Sim Target.");
+		}
+		else
+		{
+			FVersionedNiagaraEmitterData* VersionedData = EmitterHandle->GetEmitterData();
+			if (VersionedData)
+			{
+				VersionedData->SimTarget = SimTarget;
+			}
+		}
+		System.OnSystemPostEditChange().Broadcast(&System);
+		return true;
+	}
+	OutErrorMessage = LOCTEXT("InvalidEmitter", "An emitter with this name does not exist on this System.");
+	return false;
+}
+
+#undef LOCTEXT_NAMESPACE
diff --git a/Engine/Plugins/FX/Niagara/Source/NiagaraEditor/Private/TypeEditorUtilities/NiagaraVectorTypeEditorUtilities.cpp b/Engine/Plugins/FX/Niagara/Source/NiagaraEditor/Private/TypeEditorUtilities/NiagaraVectorTypeEditorUtilities.cpp
--- a/Engine/Plugins/FX/Niagara/Source/NiagaraEditor/Private/TypeEditorUtilities/NiagaraVectorTypeEditorUtilities.cpp
+++ b/Engine/Plugins/FX/Niagara/Source/NiagaraEditor/Private/TypeEditorUtilities/NiagaraVectorTypeEditorUtilities.cpp
@@ -198,7 +198,10 @@ bool FNiagaraEditorVector2TypeUtilities::CanHandlePinDefaults() const
 FString FNiagaraEditorVector2TypeUtilities::GetPinDefaultStringFromValue(const FNiagaraVariable& AllocatedVariable) const
 {
 	checkf(AllocatedVariable.IsDataAllocated(), TEXT("Can not generate a default value string for an unallocated variable."));
-	return AllocatedVariable.GetValue<FVector2f>().ToString();
+
+	// NOTE: We can not use ToString() here since the vector pin control doesn't use the standard 'X=0,Y=0' syntax. The string parsing breaks without the parentesis which are not added by default.
+	FVector2f Value = AllocatedVariable.GetValue<FVector2f>();
+	return FString::Printf(TEXT("(X=%3.3f, Y=%3.3f)"), Value.X, Value.Y);
 }
 
 bool FNiagaraEditorVector2TypeUtilities::SetValueFromPinDefaultString(const FString& StringValue, FNiagaraVariable& Variable) const
diff --git a/Engine/Plugins/FX/Niagara/Source/NiagaraEditor/Public/EdGraphSchema_Niagara.h b/Engine/Plugins/FX/Niagara/Source/NiagaraEditor/Public/EdGraphSchema_Niagara.h
--- a/Engine/Plugins/FX/Niagara/Source/NiagaraEditor/Public/EdGraphSchema_Niagara.h
+++ b/Engine/Plugins/FX/Niagara/Source/NiagaraEditor/Public/EdGraphSchema_Niagara.h
@@ -16,6 +16,7 @@ struct FNiagaraVariable;
 struct FNiagaraVariableInfo;
 struct FNiagaraTypeDefinition;
 enum class ENiagaraDataType : uint8;
+struct FNiagaraVariant;
 
 /** Action to add a node to the graph */
 USTRUCT()
@@ -142,13 +143,22 @@ class UEdGraphSchema_Niagara : public UEdGraphSchema
 	static NIAGARAEDITOR_API FNiagaraVariable PinToNiagaraVariable(const UEdGraphPin* Pin, bool bNeedsValue=false, ENiagaraStructConversion StructConversion = ENiagaraStructConversion::UserFacing);
 
 	/** 
-	  * Tries to get a default value string for a graph pin from a niagara variable.
+	  * Tries to get a default value string for a graph pin from a Niagara variable.
 	  * @param Variable The variable to get a default pin value for.
 	  * @param OutPinDefaultValue The pin default value string if this call was successful, otherwise empty.
 	  * @returns Whether or not we could get a default value form the supplied variable.
 	  */
 	NIAGARAEDITOR_API bool TryGetPinDefaultValueFromNiagaraVariable(const FNiagaraVariable& Variable, FString& OutPinDefaultValue) const;
 
+	/**
+	  * Tries to get a default value string for a graph pin from a Niagara variant.
+	  * @param Variable The variable.
+	  * @param Value The variant containing the data for the variable.
+	  * @param OutPinDefaultValue The pin default value string if this call was successful, otherwise empty.
+	  * @returns Whether or not we could get a default value form the supplied variable.
+	  */
+	NIAGARAEDITOR_API bool TryGetPinDefaultValueFromNiagaraVariant(const FNiagaraVariableBase& Variable, const FNiagaraVariant& Value, FString& OutPinDefaultValue) const;
+
 	static NIAGARAEDITOR_API void ConvertIllegalPinsInPlace(UEdGraphPin* Pin);
 	static NIAGARAEDITOR_API FNiagaraTypeDefinition PinToTypeDefinition(const UEdGraphPin* Pin, ENiagaraStructConversion StructConversion = ENiagaraStructConversion::UserFacing);
 	static NIAGARAEDITOR_API FNiagaraTypeDefinition PinTypeToTypeDefinition(const FEdGraphPinType& PinType);
diff --git a/Engine/Plugins/FX/Niagara/Source/NiagaraEditor/Public/NiagaraExternalSystemEditorUtilities.h b/Engine/Plugins/FX/Niagara/Source/NiagaraEditor/Public/NiagaraExternalSystemEditorUtilities.h
new file mode 100644
--- /dev/null
+++ b/Engine/Plugins/FX/Niagara/Source/NiagaraEditor/Public/NiagaraExternalSystemEditorUtilities.h
@@ -0,0 +1,63 @@
+// Copyright Epic Games, Inc. All Rights Reserved.
+
+#pragma once
+
+#include "NiagaraVariant.h"
+#include "NiagaraEditorUtilities.h"
+
+#define UE_API NIAGARAEDITOR_API
+
+//Tools for external code to modify Niagara Systems.
+//Provides minimal public API required for views into and modification of Niagara internals.
+class FNiagaraExternalSystemEditorUtilities
+{
+public:
+
+	struct FUserVariableInfo
+	{
+		FNiagaraVariableBase Variable;
+		FNiagaraVariant Value;
+		FString Description;
+	};
+	
+	//User Variables
+	NIAGARAEDITOR_API static void AddUserVariable(UNiagaraSystem& System, const FNiagaraVariableBase& Variable, const FNiagaraVariant& Variant, FString VariableDesc, FText TransactionDescription);
+	NIAGARAEDITOR_API static void RemoveUserVariable(UNiagaraSystem& System, const FNiagaraVariableBase& Variable, FText TransactionDescription);
+	NIAGARAEDITOR_API static void GetUserVariableInfo(UNiagaraSystem& System, TArray<FUserVariableInfo>& OutVariableInfo);
+
+	//System
+
+	struct FEmitterInfo
+	{
+		FGuid Id;
+		FName Name;
+		bool bLocalSpace = false;
+		bool bDeterministic = false;
+		ENiagaraSimTarget SimTarget = ENiagaraSimTarget::CPUSim;		
+		ENiagaraEmitterMode EmitterMode = ENiagaraEmitterMode::Standard;
+	};
+
+	struct FSystemInfo
+	{
+		TArray<FEmitterInfo> Emitters;
+	};
+	NIAGARAEDITOR_API static void GetSystemInfo(UNiagaraSystem& System, FSystemInfo& OutSystemInfo);
+
+	//Emitters
+	NIAGARAEDITOR_API static FName AddEmitterToSystem(UNiagaraSystem& System, UNiagaraEmitter& EmitterToAdd, FText TransactionDescription);
+	NIAGARAEDITOR_API static bool RemoveEmitterFromSystem(UNiagaraSystem& System, FName EmitterUniqueName, FText TransactionDescription);
+
+	NIAGARAEDITOR_API static bool GetEmitterInfo(UNiagaraSystem& System, FName EmitterUniqueName, FEmitterInfo& OutEmitterInfo, FText& OutErrorMessage);
+	NIAGARAEDITOR_API static bool SetEmitterName(UNiagaraSystem& System, FName EmitterUniqueName, FName Name, FText TransactionDescription, FText& OutErrorMessage);
+	NIAGARAEDITOR_API static bool SetEmitterLocalSpace(UNiagaraSystem& System, FName EmitterUniqueName, bool bLocalSpace, FText TransactionDescription, FText& OutErrorMessage);
+	NIAGARAEDITOR_API static bool SetEmitterDeterministic(UNiagaraSystem& System, FName EmitterUniqueName, bool bDeterministic, FText TransactionDescription, FText& OutErrorMessage);
+	NIAGARAEDITOR_API static bool SetEmitterSimTarget(UNiagaraSystem& System, FName EmitterUniqueName, ENiagaraSimTarget SimTarget, FText TransactionDescription, FText& OutErrorMessage);
+
+	private:
+		
+	static void GetEmitterInfo(const FNiagaraEmitterHandle& EmitterHandle, FEmitterInfo& OutEmitterInfo);
+	static FNiagaraEmitterHandle* GetEmitterHandle(UNiagaraSystem& System, FName EmitterUniqueName);
+	static FNiagaraEmitterHandle* GetEmitterHandle(UNiagaraSystem& System, const FGuid& EmitterId);
+};
+
+#undef UE_API
