Removed the option to render niagara system with default material if pso is not yet ready.

#rnx

#virtualized
#Synced-CL 50186389

diff --git a/Engine/Plugins/FX/Niagara/Source/Niagara/Private/NiagaraComponent.cpp b/Engine/Plugins/FX/Niagara/Source/Niagara/Private/NiagaraComponent.cpp
--- a/Engine/Plugins/FX/Niagara/Source/Niagara/Private/NiagaraComponent.cpp
+++ b/Engine/Plugins/FX/Niagara/Source/Niagara/Private/NiagaraComponent.cpp
@@ -2462,16 +2462,11 @@ FPrimitiveSceneProxy* UNiagaraComponent::CreateSceneProxy()
 			PrecacheAssetPSOs(Asset);
 		} 		
 
-		//Settings can opt to not delay rendering and just render the default material. In most cases we'll delay here until the PSO is ready.
-		bool bCanDelayRender = ReadinessSettings.PSOPrecacheSettings.DelayMode == ENiagaraSystemPSOPrecachingDelayMode::DelayRendering;
 		EPSOPrecachePriority PSOPrecachePriority = ReadinessSettings.PSOPrecacheSettings.GetPriorityBoost();
-		if (bCanDelayRender)
+		if (CheckPSOPrecachingAndBoostPriority(PSOPrecachePriority) && GetPSOPrecacheProxyCreationStrategy() != EPSOPrecacheProxyCreationStrategy::AlwaysCreate)
 		{
-			if (CheckPSOPrecachingAndBoostPriority(PSOPrecachePriority) && GetPSOPrecacheProxyCreationStrategy() != EPSOPrecacheProxyCreationStrategy::AlwaysCreate)
-			{
-				UE_LOG(LogNiagara, Verbose, TEXT("Skipping CreateSceneProxy for UNiagaraComponent %s (UNiagaraSystem PSOs are still compiling)"), *GetFullName());
-				return nullptr;
-			}
+			UE_LOG(LogNiagara, Verbose, TEXT("Skipping CreateSceneProxy for UNiagaraComponent %s (UNiagaraSystem PSOs are still compiling)"), *GetFullName());
+			return nullptr;
 		}
 	}
 #endif // UE_WITH_PSO_PRECACHING
diff --git a/Engine/Plugins/FX/Niagara/Source/Niagara/Private/NiagaraEmitterInstanceImpl.cpp b/Engine/Plugins/FX/Niagara/Source/Niagara/Private/NiagaraEmitterInstanceImpl.cpp
--- a/Engine/Plugins/FX/Niagara/Source/Niagara/Private/NiagaraEmitterInstanceImpl.cpp
+++ b/Engine/Plugins/FX/Niagara/Source/Niagara/Private/NiagaraEmitterInstanceImpl.cpp
@@ -635,11 +635,8 @@ void FNiagaraEmitterInstanceImpl::UpdatePSOStatus()
 		{
 			if (ParentSystem->IsLoadingPSOCache(PSOPrecacheParams.MaterialInterface))
 			{
-				if (EmitterData->Importance != ENiagaraEmitterImportance::Normal || ParentSystem->GetReadinessSettings().PSOPrecacheSettings.DelayMode != ENiagaraSystemPSOPrecachingDelayMode::None)
-				{
-					bWaitingForPSO = true;
-					break;
-				}
+				bWaitingForPSO = true;
+				break;
 			}
 		}
 	}
diff --git a/Engine/Plugins/FX/Niagara/Source/Niagara/Public/NiagaraScalabilitySettings.h b/Engine/Plugins/FX/Niagara/Source/Niagara/Public/NiagaraScalabilitySettings.h
--- a/Engine/Plugins/FX/Niagara/Source/Niagara/Public/NiagaraScalabilitySettings.h
+++ b/Engine/Plugins/FX/Niagara/Source/Niagara/Public/NiagaraScalabilitySettings.h
@@ -20,8 +20,6 @@ NIAGARA_API UE_DECLARE_GAMEPLAY_TAG_EXTERN(TAG_Niagara_ScalabilityProfile_Low);
 UENUM()
 enum class ENiagaraSystemPSOPrecachingDelayMode : uint8
 {
-	/** The Niagara System is never delayed and will simulate and render immediately. Can cause incorrect rendering / default material if the PSO is not ready. */
-	None,
 	/** The Niagara System will delay its rendering until the PSO is ready. Simulation will still occur but will not be rendered until the PSO is ready. */
 	DelayRendering,
 	/** The Niagara System will delay its activation until the PSO is ready. No simulation or rendering will occur until the PSO is ready. */
