- Handle potential OOB / invalid access to readback data

#Synced-CL 50261662

diff --git a/Engine/Plugins/FX/Niagara/Source/Niagara/Private/NiagaraDebugHud.cpp b/Engine/Plugins/FX/Niagara/Source/Niagara/Private/NiagaraDebugHud.cpp
--- a/Engine/Plugins/FX/Niagara/Source/Niagara/Private/NiagaraDebugHud.cpp
+++ b/Engine/Plugins/FX/Niagara/Source/Niagara/Private/NiagaraDebugHud.cpp
@@ -1359,6 +1359,11 @@ const FNiagaraDataSet* FNiagaraDebugHud::GetParticleDataSet(FNiagaraSystemInstan
 {
 	using namespace NiagaraDebugLocal;
 
+	if (EmitterInstance == nullptr || SystemInstance == nullptr)
+	{
+		return nullptr;
+	}
+
 	// For GPU context we need to readback and cache the data
 	if (EmitterInstance->GetSimTarget() == ENiagaraSimTarget::GPUComputeSim)
 	{
@@ -1375,6 +1380,12 @@ const FNiagaraDataSet* FNiagaraDebugHud::GetParticleDataSet(FNiagaraSystemInstan
 			GpuCachedData->CurrentEmitterData.AddDefaulted(NumEmitters);
 			GpuCachedData->PendingEmitterData.AddDefaulted(NumEmitters);
 		}
+		else if (!GpuCachedData->PendingEmitterData.IsValidIndex(iEmitter))
+		{
+			// In case we didn't get the change notification make sure we update the arrays
+			GpuCachedData->CurrentEmitterData.SetNum(iEmitter + 1);
+			GpuCachedData->PendingEmitterData.SetNum(iEmitter + 1);
+		}
 		GpuCachedData->LastAccessedCycles = FPlatformTime::Cycles64();
 
 		// Pending readback complete?
