- Remove newline from region warnings to avoid the information being lost

#rnx

#Synced-CL 50255171

diff --git a/Engine/Plugins/FX/Niagara/Source/Niagara/Private/NiagaraDataInterfaceSkeletalMesh.cpp b/Engine/Plugins/FX/Niagara/Source/Niagara/Private/NiagaraDataInterfaceSkeletalMesh.cpp
--- a/Engine/Plugins/FX/Niagara/Source/Niagara/Private/NiagaraDataInterfaceSkeletalMesh.cpp
+++ b/Engine/Plugins/FX/Niagara/Source/Niagara/Private/NiagaraDataInterfaceSkeletalMesh.cpp
@@ -3899,7 +3899,7 @@ int32 UNiagaraDataInterfaceSkeletalMesh::CalculateLODIndexAndSamplingRegions(USk
 					}
 					else if (RegionLODIndex != LastRegionLODIndex)
 					{
-						UE_LOG(LogNiagara, Warning, TEXT("Skeletal Mesh Data Interface is trying to use regions on LODs levels that are either streamed or cooked out. This is currently unsupported.\nInterface: %s\nMesh: %s\nRegion: %s"),
+						UE_LOG(LogNiagara, Warning, TEXT("Skeletal Mesh Data Interface is trying to use regions on LODs levels that are either streamed or cooked out. This is currently unsupported. Interface(%s) Mesh(%s) Region(%s)"),
 							*GetFullName(),
 							*SkinnedMesh->GetFullName(),
 							*RegionName.ToString());
@@ -3913,7 +3913,7 @@ int32 UNiagaraDataInterfaceSkeletalMesh::CalculateLODIndexAndSamplingRegions(USk
 					}
 					else
 					{
-						UE_LOG(LogNiagara, Warning, TEXT("Skeletal Mesh Data Interface is trying to use a region with no associated triangles.\nLOD: %d\nInterface: %s\nMesh: %s\nRegion: %s"),
+						UE_LOG(LogNiagara, Warning, TEXT("Skeletal Mesh Data Interface is trying to use a region with no associated triangles. LOD(%d) Interface(%s) Mesh(%s) Region(%s)"),
 							 RegionLODIndex,
 							*GetFullName(),
 							*SkinnedMesh->GetFullName(),
@@ -3924,7 +3924,7 @@ int32 UNiagaraDataInterfaceSkeletalMesh::CalculateLODIndexAndSamplingRegions(USk
 				}
 				else
 				{
-					UE_LOG(LogNiagara, Warning, TEXT("Skeletal Mesh Data Interface is trying to use a region on a mesh that does not provide this region.\nInterface: %s\nMesh: %s\nRegion: %s"),
+					UE_LOG(LogNiagara, Warning, TEXT("Skeletal Mesh Data Interface is trying to use a region on a mesh that does not provide this region. Interface(%s) Mesh(%s) Region(%s)"),
 						*GetFullName(),
 						*SkinnedMesh->GetFullName(),
 						*RegionName.ToString());
