Be sure to destroy simulation instances in PostApplyToComponent to reflect the fact that the instance parameters could be changed out from under us and our DI bindings could go stale

#tests Opened level that was crashing

#Synced-CL 50355583

diff --git a/Engine/Plugins/FX/Niagara/Source/Niagara/Private/NiagaraComponent.cpp b/Engine/Plugins/FX/Niagara/Source/Niagara/Private/NiagaraComponent.cpp
--- a/Engine/Plugins/FX/Niagara/Source/Niagara/Private/NiagaraComponent.cpp
+++ b/Engine/Plugins/FX/Niagara/Source/Niagara/Private/NiagaraComponent.cpp
@@ -2120,6 +2120,10 @@ void UNiagaraComponent::PostApplyToComponent()
 	Super::PostApplyToComponent();
 
 #if WITH_EDITOR
+	// Changes to the component can be broad ranging and include things like manipulating the instance parameters.  In particular if the
+	// data interfaces are changed out, we may need to rebind VM functions, so be sure to recreate the instance.
+	DestroyInstanceNotComponent();
+
 	// Ensure update the override parameters as they may be incomplete
 	// Without this calls to set override data interface values can fail as the DI may not be present in the override parameter store
 	if (Asset != nullptr)
