PR #14353: Fixes FNDCScopedReader to use template param as member type

#Synced-CL 50261662

diff --git a/Engine/Plugins/FX/Niagara/Source/Niagara/Public/NiagaraDataChannelAccessor.h b/Engine/Plugins/FX/Niagara/Source/Niagara/Public/NiagaraDataChannelAccessor.h
--- a/Engine/Plugins/FX/Niagara/Source/Niagara/Public/NiagaraDataChannelAccessor.h
+++ b/Engine/Plugins/FX/Niagara/Source/Niagara/Public/NiagaraDataChannelAccessor.h
@@ -404,7 +404,7 @@ template<typename TNDCReaderType>
 struct FNDCScopedReader
 {
 private:
-	FNDCReaderBase& Reader;
+	TNDCReaderType& Reader;
 public:
 	FNDCScopedReader(TNDCReaderType& InReader) : Reader(InReader){ }
 	~FNDCScopedReader() { Reader.EndRead();	}
