- Always allocate a critical section for the runtime allocation
- This removes the confusion around having to initialize and hopefully fixes a bug where it's occasionally nullptr

#rnx

#Synced-CL 50409238

diff --git a/Engine/Plugins/FX/Niagara/Source/Niagara/Classes/NiagaraEmitter.h b/Engine/Plugins/FX/Niagara/Source/Niagara/Classes/NiagaraEmitter.h
--- a/Engine/Plugins/FX/Niagara/Source/Niagara/Classes/NiagaraEmitter.h
+++ b/Engine/Plugins/FX/Niagara/Source/Niagara/Classes/NiagaraEmitter.h
@@ -231,7 +231,7 @@ struct FMemoryRuntimeEstimation
 	std::atomic<bool> IsEstimationDirty = false;
 	int32 AllocationEstimate = 0;
 
-	FMemoryRuntimeEstimation() = default;
+	FMemoryRuntimeEstimation();
 	FMemoryRuntimeEstimation(const FMemoryRuntimeEstimation& Other) { *this = Other; }
 
 	FMemoryRuntimeEstimation& operator=(const FMemoryRuntimeEstimation& Other)
@@ -258,7 +258,7 @@ struct FMemoryRuntimeEstimation
 	}
 
 	NIAGARA_API FCriticalSection* GetCriticalSection();
-	NIAGARA_API void Init();
+
 private:
 	TSharedPtr<FCriticalSection> EstimationCriticalSection;
 };
diff --git a/Engine/Plugins/FX/Niagara/Source/Niagara/Private/NiagaraEmitter.cpp b/Engine/Plugins/FX/Niagara/Source/Niagara/Private/NiagaraEmitter.cpp
--- a/Engine/Plugins/FX/Niagara/Source/Niagara/Private/NiagaraEmitter.cpp
+++ b/Engine/Plugins/FX/Niagara/Source/Niagara/Private/NiagaraEmitter.cpp
@@ -215,7 +215,6 @@ void FVersionedNiagaraEmitterData::CopyFrom(const FVersionedNiagaraEmitterData&
 #if STATS
 	StatDatabase.Init();
 #endif
-	RuntimeEstimation.Init();
 }
 
 FName GetVersionedName(const FString& BaseName, const FNiagaraAssetVersion& VersionPostfix)
@@ -702,7 +701,6 @@ void FVersionedNiagaraEmitterData::PostLoad(UNiagaraEmitter& Emitter, int32 Niag
 #if STATS
 	StatDatabase.Init();
 #endif
-	RuntimeEstimation.Init();
 
 #if WITH_EDITORONLY_DATA
 	//Force old data to not update Initial values to maintain existing behavior.
@@ -1169,7 +1167,6 @@ void UNiagaraEmitter::PostDuplicate(EDuplicateMode::Type DuplicateMode)
 #if STATS
 		Data.StatDatabase.Init();
 #endif
-		Data.RuntimeEstimation.Init();
 	}
 }
 
@@ -1195,7 +1192,6 @@ void UNiagaraEmitter::PostEditChangeProperty(FPropertyChangedEvent& PropertyChan
 #if STATS
 		Data.StatDatabase.Init();
 #endif
-		Data.RuntimeEstimation.Init();
 		Data.UpdateDebugName(*this, nullptr);
 	}
 	UpdateChangeId(TEXT("PostEditChangeProperty"));
@@ -2762,8 +2758,6 @@ FGuid UNiagaraEmitter::AddNewVersion(int32 MajorVersion, int32 MinorVersion)
 #if STATS
 			NewVersionData.StatDatabase.Init();
 #endif
-			NewVersionData.RuntimeEstimation.Init();
-			
 			break;
 		}
 	}
@@ -3481,15 +3475,15 @@ int32 FVersionedNiagaraEmitterData::GetMaxParticleCountEstimate()
 	return RuntimeEstimation.AllocationEstimate;
 }
 
-FCriticalSection* FMemoryRuntimeEstimation::GetCriticalSection()
+FMemoryRuntimeEstimation::FMemoryRuntimeEstimation()
 {
-	ensure(EstimationCriticalSection.IsValid());
-	return EstimationCriticalSection.Get();
+	EstimationCriticalSection = MakeShared<FCriticalSection>();
 }
 
-void FMemoryRuntimeEstimation::Init()
+FCriticalSection* FMemoryRuntimeEstimation::GetCriticalSection()
 {
-	EstimationCriticalSection = MakeShared<FCriticalSection>();
+	ensure(EstimationCriticalSection.IsValid());
+	return EstimationCriticalSection.Get();
 }
 
 void UNiagaraEmitter::UpdateStatID() const
