Optimized FNiagaraTypeDefinition::GeneratePathNameHash() to avoid any string allocations and just walk the path up the outers using FNames.  The resulting hash is not persistent.

#Synced-CL 50404854

diff --git a/Engine/Plugins/FX/Niagara/Source/Niagara/Private/NiagaraModule.cpp b/Engine/Plugins/FX/Niagara/Source/Niagara/Private/NiagaraModule.cpp
--- a/Engine/Plugins/FX/Niagara/Source/Niagara/Private/NiagaraModule.cpp
+++ b/Engine/Plugins/FX/Niagara/Source/Niagara/Private/NiagaraModule.cpp
@@ -1649,7 +1649,19 @@ bool FNiagaraTypeDefinition::Serialize(FArchive& Ar)
 
 void FNiagaraTypeDefinition::GeneratePathNameHash()
 {
-	UnderlyingTypePathNameHash = GetTypeHash(GetFullNameSafe(ClassStructOrEnum));
+	if (ClassStructOrEnum)
+	{
+		UnderlyingTypePathNameHash = GetTypeHash(ClassStructOrEnum->GetFName());
+		const UObject* Parent = ClassStructOrEnum;
+		while ((Parent = Parent->GetOuter()) != nullptr)
+		{
+			UnderlyingTypePathNameHash = HashCombineFast(UnderlyingTypePathNameHash, GetTypeHash(Parent->GetFName()));
+		}
+	}
+	else
+	{
+		UnderlyingTypePathNameHash = GetTypeHash(NAME_None);
+	}
 }
 
 //////////////////////////////////////////////////////////////////////////
