- Fix using byte buffer vs buffer for audio spectrum

#Synced-CL 50261662

diff --git a/Engine/Plugins/FX/Niagara/Source/Niagara/Private/NiagaraDataInterfaceAudioSpectrum.cpp b/Engine/Plugins/FX/Niagara/Source/Niagara/Private/NiagaraDataInterfaceAudioSpectrum.cpp
--- a/Engine/Plugins/FX/Niagara/Source/Niagara/Private/NiagaraDataInterfaceAudioSpectrum.cpp
+++ b/Engine/Plugins/FX/Niagara/Source/Niagara/Private/NiagaraDataInterfaceAudioSpectrum.cpp
@@ -110,7 +110,7 @@ namespace UE::Niagara::DataInterfaceAudioSpectrum::Private
 				const uint64 InitialDataSize = InstanceData.DataToUpload.Num() * InstanceData.DataToUpload.GetTypeSize();
 				const uint64 BufferSize = Align(InitialDataSize, 16);
 
-				FRDGBufferDesc BufferDesc = FRDGBufferDesc::CreateByteAddressDesc(IntCastChecked<uint32>(BufferSize));
+				FRDGBufferDesc BufferDesc = FRDGBufferDesc::CreateBufferDesc(sizeof(float), IntCastChecked<uint32>(BufferSize / sizeof(float)));
 				ResizeBufferIfNeeded(GraphBuilder, InstanceData.PooledBuffer, BufferDesc, TEXT("NiagaraAudioSpectrumBuffer"));
 
 				GraphBuilder.QueueBufferUpload(
