- Mesh Exporter Changes
- Add selection for TexCoords allowing data to be send from Niagara in dynamic parameters and packed into Tex Coord slots
- Make default capture local space for debugger mesh exporter

#preflght

#Synced-CL 50409238

diff --git a/Engine/Plugins/FX/Niagara/Shaders/Private/NiagaraRibbonVertexFactory.ush b/Engine/Plugins/FX/Niagara/Shaders/Private/NiagaraRibbonVertexFactory.ush
--- a/Engine/Plugins/FX/Niagara/Shaders/Private/NiagaraRibbonVertexFactory.ush
+++ b/Engine/Plugins/FX/Niagara/Shaders/Private/NiagaraRibbonVertexFactory.ush
@@ -184,16 +184,16 @@ struct FVertexFactoryIntermediates
 	float ParticleRandom;
 #endif
 
-#if (DYNAMIC_PARAMETERS_MASK & 1)
+#if (DYNAMIC_PARAMETERS_MASK & 1) || WITH_NIAGARA_VERTEX_FACTORY_EXPORT
 	float4 DynamicParameter;
 #endif
-#if (DYNAMIC_PARAMETERS_MASK & 2)
+#if (DYNAMIC_PARAMETERS_MASK & 2) || WITH_NIAGARA_VERTEX_FACTORY_EXPORT
 	float4 DynamicParameter1;
 #endif
-#if (DYNAMIC_PARAMETERS_MASK & 4)
+#if (DYNAMIC_PARAMETERS_MASK & 4) || WITH_NIAGARA_VERTEX_FACTORY_EXPORT
 	float4 DynamicParameter2;
 #endif
-#if (DYNAMIC_PARAMETERS_MASK & 8)
+#if (DYNAMIC_PARAMETERS_MASK & 8) || WITH_NIAGARA_VERTEX_FACTORY_EXPORT
 	float4 DynamicParameter3;
 #endif
 	/** Cached primitive and instance data */
@@ -720,7 +720,7 @@ FVertexFactoryIntermediates GetVertexFactoryIntermediates(FVertexFactoryInput In
 		Intermediates.Color = lerp(Color0, Color1, InterpVtxData.Alpha);
 	}
 
-#if (DYNAMIC_PARAMETERS_MASK & 1)
+#if (DYNAMIC_PARAMETERS_MASK & 1) || WITH_NIAGARA_VERTEX_FACTORY_EXPORT
 	// Overrides (FVertexFactoryIntermediates)0; 
 	Intermediates.DynamicParameter = float4(1.0f, 1.0f, 1.0f, 1.0f);
 
@@ -730,7 +730,7 @@ FVertexFactoryIntermediates GetVertexFactoryIntermediates(FVertexFactoryInput In
 		Intermediates.DynamicParameter = lerp(GetVec4(NiagaraRibbonVF.MaterialParamDataOffset, InterpVtxData.ParticleId0), GetVec4(NiagaraRibbonVF.MaterialParamDataOffset, InterpVtxData.ParticleId1), InterpVtxData.Alpha);
 	}
 #endif
-#if (DYNAMIC_PARAMETERS_MASK & 2)
+#if (DYNAMIC_PARAMETERS_MASK & 2) || WITH_NIAGARA_VERTEX_FACTORY_EXPORT
 	// Overrides (FVertexFactoryIntermediates)0; 
 	Intermediates.DynamicParameter1 = float4(1.0f, 1.0f, 1.0f, 1.0f);
 
@@ -740,7 +740,7 @@ FVertexFactoryIntermediates GetVertexFactoryIntermediates(FVertexFactoryInput In
 		Intermediates.DynamicParameter1 = lerp(GetVec4(NiagaraRibbonVF.MaterialParam1DataOffset, InterpVtxData.ParticleId0), GetVec4(NiagaraRibbonVF.MaterialParam1DataOffset, InterpVtxData.ParticleId1), InterpVtxData.Alpha);
 	}
 #endif
-#if (DYNAMIC_PARAMETERS_MASK & 4)
+#if (DYNAMIC_PARAMETERS_MASK & 4) || WITH_NIAGARA_VERTEX_FACTORY_EXPORT
 	// Overrides (FVertexFactoryIntermediates)0; 
 	Intermediates.DynamicParameter2 = float4(1.0f, 1.0f, 1.0f, 1.0f);
 
@@ -750,7 +750,7 @@ FVertexFactoryIntermediates GetVertexFactoryIntermediates(FVertexFactoryInput In
 		Intermediates.DynamicParameter2 = lerp(GetVec4(NiagaraRibbonVF.MaterialParam2DataOffset, InterpVtxData.ParticleId0), GetVec4(NiagaraRibbonVF.MaterialParam2DataOffset, InterpVtxData.ParticleId1), InterpVtxData.Alpha);
 	}
 #endif
-#if (DYNAMIC_PARAMETERS_MASK & 8)
+#if (DYNAMIC_PARAMETERS_MASK & 8) || WITH_NIAGARA_VERTEX_FACTORY_EXPORT
 	// Overrides (FVertexFactoryIntermediates)0; 
 	Intermediates.DynamicParameter3 = float4(1.0f, 1.0f, 1.0f, 1.0f);
 
@@ -1174,6 +1174,40 @@ float2 NiagaraVertexFactoryExport_GetTexCoord(FVertexFactoryInput Input, FVertex
 		default: return float2(0,0);
 	}
 }
+
+float4 NiagaraVertexFactoryExport_GetDynamicParameter(FVertexFactoryInput Input, FVertexFactoryIntermediates Intermediates, uint ParameterIndex)
+{
+	float4 Value = 1.0f;
+	if ( ParameterIndex == 0 )
+	{
+		Value = Intermediates.DynamicParameter;
+	}
+	else if ( ParameterIndex == 1 )
+	{
+		Value = Intermediates.DynamicParameter1;
+	}
+	else if ( ParameterIndex == 2 )
+	{
+		Value = Intermediates.DynamicParameter2;
+	}
+	else if ( ParameterIndex == 3 )
+	{
+		Value = Intermediates.DynamicParameter3;
+	}
+	else
+	{
+		return 1.0f;
+	}
+
+	uint MaskBit = ParameterIndex * 4u;
+	for ( int i=0; i < 4; ++i )
+	{
+		Value[i] = ((NiagaraRibbonVF.MaterialParamValidMask & (1u << MaskBit)) != 0u) ? Value[i] : 1.0f;
+		++MaskBit;
+	}
+	return Value;
+}
+
 #endif //WITH_NIAGARA_VERTEX_FACTORY_EXPORT
 
 #include "/Engine/Private/VertexFactoryDefaultInterface.ush"
diff --git a/Engine/Plugins/FX/Niagara/Shaders/Private/NiagaraSpriteVertexFactory.ush b/Engine/Plugins/FX/Niagara/Shaders/Private/NiagaraSpriteVertexFactory.ush
--- a/Engine/Plugins/FX/Niagara/Shaders/Private/NiagaraSpriteVertexFactory.ush
+++ b/Engine/Plugins/FX/Niagara/Shaders/Private/NiagaraSpriteVertexFactory.ush
@@ -218,16 +218,16 @@ struct FVertexFactoryIntermediates
 	float3 CustomFacingVector;
 	float3 CustomAlignmentVector;
 
-#if (DYNAMIC_PARAMETERS_MASK & 1)
+#if (DYNAMIC_PARAMETERS_MASK & 1) || WITH_NIAGARA_VERTEX_FACTORY_EXPORT
 	float4 DynamicParameter;
 #endif
-#if (DYNAMIC_PARAMETERS_MASK & 2)
+#if (DYNAMIC_PARAMETERS_MASK & 2) || WITH_NIAGARA_VERTEX_FACTORY_EXPORT
 	float4 DynamicParameter1;
 #endif
-#if (DYNAMIC_PARAMETERS_MASK & 4)
+#if (DYNAMIC_PARAMETERS_MASK & 4) || WITH_NIAGARA_VERTEX_FACTORY_EXPORT
 	float4 DynamicParameter2;
 #endif
-#if (DYNAMIC_PARAMETERS_MASK & 8)
+#if (DYNAMIC_PARAMETERS_MASK & 8) || WITH_NIAGARA_VERTEX_FACTORY_EXPORT
 	float4 DynamicParameter3;
 #endif
 
@@ -853,25 +853,25 @@ float GetNiagaraNormalizedAge(uint InstanceID)
 }
 #endif
 
-#if (DYNAMIC_PARAMETERS_MASK & 1)
+#if (DYNAMIC_PARAMETERS_MASK & 1) || WITH_NIAGARA_VERTEX_FACTORY_EXPORT
 float4 GetNiagaraParticleDynamicParameters(uint InstanceID)
 {
 	return SafeGetVec4(NiagaraSpriteVF.MaterialParamDataOffset, InstanceID, NiagaraSpriteVF.DefaultDynamicMaterialParameter0);
 }
 #endif
-#if (DYNAMIC_PARAMETERS_MASK & 2)
+#if (DYNAMIC_PARAMETERS_MASK & 2) || WITH_NIAGARA_VERTEX_FACTORY_EXPORT
 float4 GetNiagaraParticleDynamicParameters1(uint InstanceID)
 {
 	return SafeGetVec4(NiagaraSpriteVF.MaterialParam1DataOffset, InstanceID, NiagaraSpriteVF.DefaultDynamicMaterialParameter1);
 }
 #endif
-#if (DYNAMIC_PARAMETERS_MASK & 4)
+#if (DYNAMIC_PARAMETERS_MASK & 4) || WITH_NIAGARA_VERTEX_FACTORY_EXPORT
 float4 GetNiagaraParticleDynamicParameters2(uint InstanceID)
 {
 	return SafeGetVec4(NiagaraSpriteVF.MaterialParam2DataOffset, InstanceID, NiagaraSpriteVF.DefaultDynamicMaterialParameter2);
 }
 #endif
-#if (DYNAMIC_PARAMETERS_MASK & 8)
+#if (DYNAMIC_PARAMETERS_MASK & 8) || WITH_NIAGARA_VERTEX_FACTORY_EXPORT
 float4 GetNiagaraParticleDynamicParameters3(uint InstanceID)
 {
 	return SafeGetVec4(NiagaraSpriteVF.MaterialParam3DataOffset, InstanceID, NiagaraSpriteVF.DefaultDynamicMaterialParameter3);
@@ -938,16 +938,16 @@ FVertexFactoryIntermediates GetVertexFactoryIntermediates(FVertexFactoryInput In
 	Intermediates.CustomAlignmentVector = SafeNormalize(GetNiagaraParticleAlignmentVector(ParticleID));
 	Intermediates.UVScale = UVScale;
 
-#if (DYNAMIC_PARAMETERS_MASK & 1)
+#if (DYNAMIC_PARAMETERS_MASK & 1) || WITH_NIAGARA_VERTEX_FACTORY_EXPORT
 	Intermediates.DynamicParameter = GetNiagaraParticleDynamicParameters(ParticleID);
 #endif
-#if (DYNAMIC_PARAMETERS_MASK & 2)
+#if (DYNAMIC_PARAMETERS_MASK & 2) || WITH_NIAGARA_VERTEX_FACTORY_EXPORT
 	Intermediates.DynamicParameter1 = GetNiagaraParticleDynamicParameters1(ParticleID);
 #endif
-#if (DYNAMIC_PARAMETERS_MASK & 4)
+#if (DYNAMIC_PARAMETERS_MASK & 4) || WITH_NIAGARA_VERTEX_FACTORY_EXPORT
 	Intermediates.DynamicParameter2 = GetNiagaraParticleDynamicParameters2(ParticleID);
 #endif
-#if (DYNAMIC_PARAMETERS_MASK & 8)
+#if (DYNAMIC_PARAMETERS_MASK & 8) || WITH_NIAGARA_VERTEX_FACTORY_EXPORT
 	Intermediates.DynamicParameter3 = GetNiagaraParticleDynamicParameters3(ParticleID);
 #endif
 
@@ -1395,6 +1395,39 @@ float2 NiagaraVertexFactoryExport_GetTexCoord(FVertexFactoryInput Input, FVertex
 {
 	return GetUV(Intermediates, CoordIndex);
 }
+
+float4 NiagaraVertexFactoryExport_GetDynamicParameter(FVertexFactoryInput Input, FVertexFactoryIntermediates Intermediates, uint ParameterIndex)
+{
+	float4 Value = 1.0f;
+	if ( ParameterIndex == 0 )
+	{
+		Value = Intermediates.DynamicParameter;
+	}
+	else if ( ParameterIndex == 1 )
+	{
+		Value = Intermediates.DynamicParameter1;
+	}
+	else if ( ParameterIndex == 2 )
+	{
+		Value = Intermediates.DynamicParameter2;
+	}
+	else if ( ParameterIndex == 3 )
+	{
+		Value = Intermediates.DynamicParameter3;
+	}
+	else
+	{
+		return 1.0f;
+	}
+
+	uint MaskBit = ParameterIndex * 4u;
+	for ( int i=0; i < 4; ++i )
+	{
+		Value[i] = ((NiagaraSpriteVF.MaterialParamValidMask & (1u << MaskBit)) != 0u) ? Value[i] : 1.0f;
+		++MaskBit;
+	}
+	return Value;
+}
 #endif //WITH_NIAGARA_VERTEX_FACTORY_EXPORT
 
 #include "/Engine/Private/VertexFactoryDefaultInterface.ush"
diff --git a/Engine/Plugins/FX/Niagara/Shaders/Private/NiagaraVertexFactoryExport.usf b/Engine/Plugins/FX/Niagara/Shaders/Private/NiagaraVertexFactoryExport.usf
--- a/Engine/Plugins/FX/Niagara/Shaders/Private/NiagaraVertexFactoryExport.usf
+++ b/Engine/Plugins/FX/Niagara/Shaders/Private/NiagaraVertexFactoryExport.usf
@@ -24,6 +24,7 @@ uint VertexColorOffset;
 uint VertexTangentBasisOffset;
 uint VertexTexCoordOffset;
 uint VertexTexCoordNum;
+uint4 VertexTexCoordBindings[(8 * 2) / 4];		// NUM_TEXCOORDS * Channels
 uint VertexOutputOffset;
 RWByteAddressBuffer	RWVertexData;
 
@@ -130,8 +131,29 @@ void VertexFactoryExportCS(uint3 GroupId : SV_GroupID, int GroupThreadIndex : SV
 	}
 	OutputColor(VertexIndex, VertexColor);
 	
-	for ( uint i=0; i < VertexTexCoordNum; ++i )
+	if ( VertexTexCoordNum > 0 )
 	{
-		OutputTexCoord(VertexIndex, i, NiagaraVertexFactoryExport_GetTexCoord(Input, VFIntermediates, i));
+		float4 BindingValues[8];
+		for ( int i=0; i < 4; ++ i)
+		{
+			BindingValues[i] = NiagaraVertexFactoryExport_GetDynamicParameter(Input, VFIntermediates, i);
+		}
+		for ( int i=0; i < 4; ++ i)
+		{
+			BindingValues[4 + i].xy = NiagaraVertexFactoryExport_GetTexCoord(Input, VFIntermediates, (i * 2) + 0);
+			BindingValues[4 + i].zw = NiagaraVertexFactoryExport_GetTexCoord(Input, VFIntermediates, (i * 2) + 1);
+		}
+
+		for ( uint i=0; i < VertexTexCoordNum; ++i )
+		{
+			const uint UInt4Index = i / 2;
+			const uint UInt4Offset = (i % 2) * 2;
+			
+			const uint XBinding = VertexTexCoordBindings[UInt4Index][UInt4Offset + 0];
+			const uint YBinding = VertexTexCoordBindings[UInt4Index][UInt4Offset + 1];
+			const float2 Value = float2(BindingValues[XBinding / 4][XBinding % 4], BindingValues[YBinding / 4][YBinding % 4]);
+			
+			OutputTexCoord(VertexIndex, i, Value);
+		}
 	}
 }
diff --git a/Engine/Plugins/FX/Niagara/Source/Niagara/Internal/NiagaraRendererReadback.h b/Engine/Plugins/FX/Niagara/Source/Niagara/Internal/NiagaraRendererReadback.h
--- a/Engine/Plugins/FX/Niagara/Source/Niagara/Internal/NiagaraRendererReadback.h
+++ b/Engine/Plugins/FX/Niagara/Source/Niagara/Internal/NiagaraRendererReadback.h
@@ -32,6 +32,58 @@ enum class ENiagaraRendererReadbackCoordinateSpace : uint8
 	Local,
 };
 
+// Note: Changing this requires changing NiagaraVertexFactoryExport.usf
+UENUM()
+enum class ENiagaraRendererReadbackAttribute : uint32
+{
+	DynamicParameter0X,
+	DynamicParameter0Y,
+	DynamicParameter0Z,
+	DynamicParameter0W,
+
+	DynamicParameter1X,
+	DynamicParameter1Y,
+	DynamicParameter1Z,
+	DynamicParameter1W,
+
+	DynamicParameter2X,
+	DynamicParameter2Y,
+	DynamicParameter2Z,
+	DynamicParameter2W,
+
+	DynamicParameter3X,
+	DynamicParameter3Y,
+	DynamicParameter3Z,
+	DynamicParameter3W,
+
+	TexCoord0U,
+	TexCoord0V,
+
+	TexCoord1U,
+	TexCoord1V,
+
+	TexCoord2U,
+	TexCoord2V,
+
+	TexCoord3U,
+	TexCoord3V,
+
+	TexCoord4U,
+	TexCoord4V,
+
+	TexCoord5U,
+	TexCoord5V,
+
+	TexCoord6U,
+	TexCoord6V,
+
+	TexCoord7U,
+	TexCoord7V,
+
+	Num UMETA(Hidden)
+};
+
+
 USTRUCT(BlueprintType)
 struct FNiagaraRendererReadbackParameters
 {
@@ -58,6 +110,11 @@ struct FNiagaraRendererReadbackParameters
 	UPROPERTY(EditAnywhere, BlueprintReadWrite, Category = "Parameters", meta=(UIMin="0", UIMax="8", ClampMin="0", ClampMax="8"));
 	int32 ExportNumTexCoords = 1;
 
+	// Optional custom bindings for texture coordinate packing from the export process
+	// TexCoord0.X -> TexCoord0.Y, TexCoord1.X -> etc
+	UPROPERTY(EditAnywhere, BlueprintReadWrite, Category = "Parameters")
+	TArray<ENiagaraRendererReadbackAttribute> ExportTexCoordBindings;
+
 	// When enabled we will attempt to export the materials used with each section
 	// When disabled no materials will be assigned, so the default material will be used
 	UPROPERTY(EditAnywhere, BlueprintReadWrite, Category = "Parameters")
@@ -83,6 +140,14 @@ struct FNiagaraRendererReadbackResult
 		uint32								NumTriangles = 0;
 	};
 
+	FNiagaraRendererReadbackResult()
+	{
+		for (int i = 0; i < UE_ARRAY_COUNT(VertexTexCoordBindings); ++i)
+		{
+			VertexTexCoordBindings[i] = int(ENiagaraRendererReadbackAttribute::TexCoord0U) + i;
+		}
+	}
+
 	bool HasPosition() const { return VertexPositionOffset != INDEX_NONE; }
 	bool HasColor() const { return VertexColorOffset != INDEX_NONE; }
 	bool HasTangentBasis() const { return VertexTangentBasisOffset != INDEX_NONE; }
@@ -102,6 +167,7 @@ struct FNiagaraRendererReadbackResult
 	uint32				VertexTangentBasisOffset = INDEX_NONE;
 	uint32				VertexTexCoordOffset = INDEX_NONE;
 	uint32				VertexTexCoordNum = 0;
+	uint32				VertexTexCoordBindings[8 * 2];
 	TArray<uint8>		VertexData;
 	TArray<FSection>	Sections;
 
diff --git a/Engine/Plugins/FX/Niagara/Source/Niagara/Private/NiagaraRendererReadback.cpp b/Engine/Plugins/FX/Niagara/Source/Niagara/Private/NiagaraRendererReadback.cpp
--- a/Engine/Plugins/FX/Niagara/Source/Niagara/Private/NiagaraRendererReadback.cpp
+++ b/Engine/Plugins/FX/Niagara/Source/Niagara/Private/NiagaraRendererReadback.cpp
@@ -138,6 +138,7 @@ namespace NiagaraRendererReadback
 			SingleShaderBindings.Add(Shader->VertexTangentBasisOffset,	VertexLayout.VertexTangentBasisOffset);
 			SingleShaderBindings.Add(Shader->VertexTexCoordOffset,		VertexLayout.VertexTexCoordOffset);
 			SingleShaderBindings.Add(Shader->VertexTexCoordNum,			VertexLayout.VertexTexCoordNum);
+			SingleShaderBindings.Add(Shader->VertexTexCoordBindings,	VertexLayout.VertexTexCoordBindings);
 			SingleShaderBindings.Add(Shader->VertexOutputOffset,		ReadbackMeshBatch.OutputVertexOffset);
 			SingleShaderBindings.Add(Shader->SectionInfoOutputOffset,	SectionInfoOffset);
 
@@ -446,6 +447,10 @@ namespace NiagaraRendererReadback
 				Result.VertexTexCoordNum = NumTexCoords;
 				Result.VertexTexCoordOffset = Result.VertexStride;
 				Result.VertexStride += sizeof(FVector2f) * NumTexCoords;
+				for (int32 i=0; i < FMath::Min<int32>(Parameters.ExportTexCoordBindings.Num(), UE_ARRAY_COUNT(Result.VertexTexCoordBindings)); ++i )
+				{
+					Result.VertexTexCoordBindings[i] = uint32(Parameters.ExportTexCoordBindings[i]);
+				}
 			}
 		}
 
diff --git a/Engine/Plugins/FX/Niagara/Source/NiagaraEditor/Private/NiagaraBakerRendererOutputStaticMesh.cpp b/Engine/Plugins/FX/Niagara/Source/NiagaraEditor/Private/NiagaraBakerRendererOutputStaticMesh.cpp
--- a/Engine/Plugins/FX/Niagara/Source/NiagaraEditor/Private/NiagaraBakerRendererOutputStaticMesh.cpp
+++ b/Engine/Plugins/FX/Niagara/Source/NiagaraEditor/Private/NiagaraBakerRendererOutputStaticMesh.cpp
@@ -218,6 +218,7 @@ bool FNiagaraBakerRendererOutputStaticMesh::ConvertReadbackResultsToStaticMesh(c
 		FMeshBuildSettings& MeshBuildSettings = StaticMesh->GetSourceModel(0).BuildSettings;
 		MeshBuildSettings.bRecomputeNormals = false;
 		MeshBuildSettings.bRecomputeTangents = false;
+		MeshBuildSettings.bGenerateLightmapUVs = false;
 	}
 
 	StaticMesh->SetStaticMaterials(StaticMaterials);
diff --git a/Engine/Plugins/FX/Niagara/Source/NiagaraEditor/Private/Widgets/SNiagaraDebuggerMeshCapture.cpp b/Engine/Plugins/FX/Niagara/Source/NiagaraEditor/Private/Widgets/SNiagaraDebuggerMeshCapture.cpp
--- a/Engine/Plugins/FX/Niagara/Source/NiagaraEditor/Private/Widgets/SNiagaraDebuggerMeshCapture.cpp
+++ b/Engine/Plugins/FX/Niagara/Source/NiagaraEditor/Private/Widgets/SNiagaraDebuggerMeshCapture.cpp
@@ -51,6 +51,9 @@ namespace NiagaraDebuggerMeshCapturePrivate
 
 void SNiagaraDebuggerMeshCapture::Construct(const FArguments& InArgs)
 {
+	// Setup defaults
+	CaptureSettings.ExportParameters.CoordinateSpace = ENiagaraRendererReadbackCoordinateSpace::Local;
+
 	// Create details view for the export parameters
 	FPropertyEditorModule& PropertyModule = FModuleManager::LoadModuleChecked<FPropertyEditorModule>("PropertyEditor");
 
diff --git a/Engine/Plugins/FX/Niagara/Source/NiagaraVertexFactories/Internal/NiagaraVertexFactoryExport.h b/Engine/Plugins/FX/Niagara/Source/NiagaraVertexFactories/Internal/NiagaraVertexFactoryExport.h
--- a/Engine/Plugins/FX/Niagara/Source/NiagaraVertexFactories/Internal/NiagaraVertexFactoryExport.h
+++ b/Engine/Plugins/FX/Niagara/Source/NiagaraVertexFactories/Internal/NiagaraVertexFactoryExport.h
@@ -66,6 +66,7 @@ public:
 	LAYOUT_FIELD(FShaderParameter,			VertexTangentBasisOffset);
 	LAYOUT_FIELD(FShaderParameter,			VertexTexCoordOffset);
 	LAYOUT_FIELD(FShaderParameter,			VertexTexCoordNum);
+	LAYOUT_FIELD(FShaderParameter,			VertexTexCoordBindings);
 	LAYOUT_FIELD(FShaderParameter,			VertexOutputOffset);
 	LAYOUT_FIELD(FShaderResourceParameter,	RWVertexData);
 };
diff --git a/Engine/Plugins/FX/Niagara/Source/NiagaraVertexFactories/Private/NiagaraVertexFactoryExport.cpp b/Engine/Plugins/FX/Niagara/Source/NiagaraVertexFactories/Private/NiagaraVertexFactoryExport.cpp
--- a/Engine/Plugins/FX/Niagara/Source/NiagaraVertexFactories/Private/NiagaraVertexFactoryExport.cpp
+++ b/Engine/Plugins/FX/Niagara/Source/NiagaraVertexFactories/Private/NiagaraVertexFactoryExport.cpp
@@ -31,6 +31,7 @@ FNiagaraVertexFactoryExportCS::FNiagaraVertexFactoryExportCS(const FMeshMaterial
 	VertexTangentBasisOffset.Bind(Initializer.ParameterMap, TEXT("VertexTangentBasisOffset"));
 	VertexTexCoordOffset.Bind(Initializer.ParameterMap, TEXT("VertexTexCoordOffset"));
 	VertexTexCoordNum.Bind(Initializer.ParameterMap, TEXT("VertexTexCoordNum"));
+	VertexTexCoordBindings.Bind(Initializer.ParameterMap, TEXT("VertexTexCoordBindings"));
 	VertexOutputOffset.Bind(Initializer.ParameterMap, TEXT("VertexOutputOffset"));
 	RWVertexData.Bind(Initializer.ParameterMap, TEXT("RWVertexData"));
 }
