- Improve tooltips on sprite renderer

#rnx

#Synced-CL 50548957

diff --git a/Engine/Plugins/FX/Niagara/Source/Niagara/Public/NiagaraSpriteRendererProperties.h b/Engine/Plugins/FX/Niagara/Source/Niagara/Public/NiagaraSpriteRendererProperties.h
--- a/Engine/Plugins/FX/Niagara/Source/Niagara/Public/NiagaraSpriteRendererProperties.h
+++ b/Engine/Plugins/FX/Niagara/Source/Niagara/Public/NiagaraSpriteRendererProperties.h
@@ -172,11 +172,17 @@ public:
 	UPROPERTY(EditAnywhere, Category = "Sprite Rendering")
 	ENiagaraRendererSourceDataMode SourceMode = ENiagaraRendererSourceDataMode::Particles;
 
-	/** Imagine the particle texture having an arrow pointing up, these modes define how the particle aligns that texture to other particle attributes.*/
+	/**
+	Imagine the particle texture having an arrow pointing up, these modes define how the particle aligns that texture to other particle attributes.
+	Automatic mode will select between Unaligned & CustomAlignment depending on if SpriteAlignment binding is valid.
+	*/
 	UPROPERTY(EditAnywhere, Category = "Sprite Rendering")
 	ENiagaraSpriteAlignment Alignment = ENiagaraSpriteAlignment::Automatic;
 
-	/** Determines how the particle billboard orients itself relative to the camera.*/
+	/**
+	Determines how the particle billboard orients itself relative to the camera.
+	Automatic mode will select FaceCamera or CustomFacingVector depending on if SpriteFacing binding is valid
+	*/
 	UPROPERTY(EditAnywhere, Category = "Sprite Rendering")
 	ENiagaraSpriteFacingMode FacingMode = ENiagaraSpriteFacingMode::Automatic;
 
