Fixed a bug where pso checking was done on a worker thread instead of game thread

#rnx

#Synced-CL 50596076

diff --git a/Engine/Plugins/FX/Niagara/Source/Niagara/Classes/NiagaraEmitterInstance.h b/Engine/Plugins/FX/Niagara/Source/Niagara/Classes/NiagaraEmitterInstance.h
--- a/Engine/Plugins/FX/Niagara/Source/Niagara/Classes/NiagaraEmitterInstance.h
+++ b/Engine/Plugins/FX/Niagara/Source/Niagara/Classes/NiagaraEmitterInstance.h
@@ -102,6 +102,8 @@ public:
 	[[nodiscard]] NIAGARA_API UNiagaraDataInterface* FindDataInterface(const FNiagaraVariable& InVariable) const;
 
 	[[nodiscard]] FNiagaraEmitterID GetEmitterID() const{ return FNiagaraEmitterID(EmitterIndex); }
+	
+	virtual void UpdatePSOStatus();
 
 	//////////////////////////////////////////////////////////////////////////////////////////////////////////////////
 
diff --git a/Engine/Plugins/FX/Niagara/Source/Niagara/Internal/NiagaraEmitterInstanceImpl.h b/Engine/Plugins/FX/Niagara/Source/Niagara/Internal/NiagaraEmitterInstanceImpl.h
--- a/Engine/Plugins/FX/Niagara/Source/Niagara/Internal/NiagaraEmitterInstanceImpl.h
+++ b/Engine/Plugins/FX/Niagara/Source/Niagara/Internal/NiagaraEmitterInstanceImpl.h
@@ -112,9 +112,10 @@ public:
 
 	bool IsParticleComponentActive(FObjectKey ComponentKey, int32 ParticleID) const;
 
+	virtual void UpdatePSOStatus() override;
+
 private:
 	void CheckForErrors();
-	void UpdatePSOStatus();
 
 	void BuildConstantBufferTable(
 		const FNiagaraScriptExecutionContext& ExecContext,
diff --git a/Engine/Plugins/FX/Niagara/Source/Niagara/Private/NiagaraEmitterInstance.cpp b/Engine/Plugins/FX/Niagara/Source/Niagara/Private/NiagaraEmitterInstance.cpp
--- a/Engine/Plugins/FX/Niagara/Source/Niagara/Private/NiagaraEmitterInstance.cpp
+++ b/Engine/Plugins/FX/Niagara/Source/Niagara/Private/NiagaraEmitterInstance.cpp
@@ -177,6 +177,11 @@ UNiagaraDataInterface* FNiagaraEmitterInstance::FindDataInterface(const FNiagara
 	return nullptr;
 }
 
+void FNiagaraEmitterInstance::UpdatePSOStatus()
+{
+	// nothing to do here, can be overriden by the emitter implementation
+}
+
 FNiagaraScriptExecutionContext& FNiagaraEmitterInstance::GetSpawnExecutionContext()
 {
 	FNiagaraEmitterInstanceImpl* StatefulEmitter = AsStateful();
diff --git a/Engine/Plugins/FX/Niagara/Source/Niagara/Private/NiagaraEmitterInstanceImpl.cpp b/Engine/Plugins/FX/Niagara/Source/Niagara/Private/NiagaraEmitterInstanceImpl.cpp
--- a/Engine/Plugins/FX/Niagara/Source/Niagara/Private/NiagaraEmitterInstanceImpl.cpp
+++ b/Engine/Plugins/FX/Niagara/Source/Niagara/Private/NiagaraEmitterInstanceImpl.cpp
@@ -994,11 +994,6 @@ void FNiagaraEmitterInstanceImpl::PostTick()
 				}
 			}
 		}
-		
-		if (bWaitingForPSO)
-		{
-			UpdatePSOStatus();
-		}
 
 #if STATS
 		EmitterData->GetStatData().AddStatCapture(TTuple<uint64, ENiagaraScriptUsage>((uint64)this, ENiagaraScriptUsage::ParticleSpawnScript), GetSpawnExecutionContext().ReportStats());
diff --git a/Engine/Plugins/FX/Niagara/Source/Niagara/Private/NiagaraSystemInstance.cpp b/Engine/Plugins/FX/Niagara/Source/Niagara/Private/NiagaraSystemInstance.cpp
--- a/Engine/Plugins/FX/Niagara/Source/Niagara/Private/NiagaraSystemInstance.cpp
+++ b/Engine/Plugins/FX/Niagara/Source/Niagara/Private/NiagaraSystemInstance.cpp
@@ -2627,6 +2627,17 @@ void FNiagaraSystemInstance::Tick_GameThread(float DeltaSeconds)
 	Age += DeltaSeconds;
 	TickCount += 1;
 	
+#if UE_WITH_PSO_PRECACHING
+	for (const FNiagaraEmitterInstanceRef& EmitterRef : Emitters)
+	{
+		//TODO (mga) add a flag to the system instance to skip iterating over the emitters if they are all done with PSOs
+		if (EmitterRef->IsWaitingForPSOCaching())
+		{
+			EmitterRef->UpdatePSOStatus();
+		}
+	}
+#endif
+	
 #if WITH_EDITOR
 	// We need to tick the rapid iteration parameters when in the editor
 	for (const FNiagaraEmitterInstanceRef& EmitterRef : Emitters)
