- Niagara Debug HUD changes
- Fix Include Cascade option not actually doing anything
- Add Include Niagara, allows the user to see just Cascasde
- Fix filters being confusing by adding a separate master control and making sure all sub filters respect it.  Previous we had a master control but the sub filters could still be active and not visible

#Synced-CL 50301641

diff --git a/Engine/Plugins/FX/Niagara/Source/Niagara/Private/NiagaraDebugHud.cpp b/Engine/Plugins/FX/Niagara/Source/Niagara/Private/NiagaraDebugHud.cpp
--- a/Engine/Plugins/FX/Niagara/Source/Niagara/Private/NiagaraDebugHud.cpp
+++ b/Engine/Plugins/FX/Niagara/Source/Niagara/Private/NiagaraDebugHud.cpp
@@ -1100,8 +1100,8 @@ void FNiagaraDebugHud::GatherSystemInfo()
 	for (TObjectIterator<UFXSystemComponent> It; It; ++It)
 	{
 		UFXSystemComponent* FXComponent = *It;
-		UNiagaraComponent* NiagaraComponent = Cast<UNiagaraComponent>(FXComponent);
-		UParticleSystemComponent* CascadeComponent = Cast<UParticleSystemComponent>(FXComponent);
+		UNiagaraComponent* NiagaraComponent = Settings.bIncludeNiagara ? Cast<UNiagaraComponent>(FXComponent) : nullptr;
+		UParticleSystemComponent* CascadeComponent = Settings.bIncludeCascade ? Cast<UParticleSystemComponent>(FXComponent) : nullptr;
 
 		FNiagaraSystemInstanceControllerPtr SystemInstanceController = NiagaraComponent ? NiagaraComponent->GetSystemInstanceController() : nullptr;
 		FNiagaraSystemInstance* SystemInstance = SystemInstanceController.IsValid() ? SystemInstanceController->GetSystemInstance_Unsafe() : nullptr;
@@ -1163,7 +1163,7 @@ void FNiagaraDebugHud::GatherSystemInfo()
 				LongestSystemPrettyName = SystemDebugInfo.SystemPrettyName;
 			}
 		}
-		SystemDebugInfo.bPassesSystemFilter = !Settings.bSystemFilterEnabled || Settings.SystemFilter.IsEmpty() || SystemDebugInfo.SystemName.MatchesWildcard(Settings.SystemFilter);
+		SystemDebugInfo.bPassesSystemFilter = !(Settings.bFiltersEnabled && Settings.bSystemFilterEnabled) || Settings.SystemFilter.IsEmpty() || SystemDebugInfo.SystemName.MatchesWildcard(Settings.SystemFilter);
 		SystemDebugInfo.bShowInWorld = (Settings.SystemDebugVerbosity != ENiagaraDebugHudVerbosity::None) && SystemDebugInfo.bPassesSystemFilter;
 
 		const bool bCanShowInWorld = 
@@ -1176,14 +1176,14 @@ void FNiagaraDebugHud::GatherSystemInfo()
 			bool bIsMatch = true;
 
 			// Filter by actor
-			if ( Settings.bActorFilterEnabled && !Settings.ActorFilter.IsEmpty())
+			if (Settings.bFiltersEnabled && Settings.bActorFilterEnabled && !Settings.ActorFilter.IsEmpty())
 			{
 				AActor* Actor = FXComponent->GetOwner();
 				bIsMatch &= (Actor != nullptr) && Actor->GetActorNameOrLabel().MatchesWildcard(Settings.ActorFilter);
 			}
 
 			// Filter by component
-			if ( bIsMatch && Settings.bComponentFilterEnabled && !Settings.ComponentFilter.IsEmpty())
+			if ( bIsMatch && Settings.bFiltersEnabled && Settings.bComponentFilterEnabled && !Settings.ComponentFilter.IsEmpty())
 			{
 				bIsMatch &= FXComponent->GetName().MatchesWildcard(Settings.ComponentFilter);
 			}
diff --git a/Engine/Plugins/FX/Niagara/Source/Niagara/Public/NiagaraDebuggerCommon.h b/Engine/Plugins/FX/Niagara/Source/Niagara/Public/NiagaraDebuggerCommon.h
--- a/Engine/Plugins/FX/Niagara/Source/Niagara/Public/NiagaraDebuggerCommon.h
+++ b/Engine/Plugins/FX/Niagara/Source/Niagara/Public/NiagaraDebuggerCommon.h
@@ -484,10 +484,14 @@ struct FNiagaraDebugHUDSettingsData
 	UPROPERTY(EditAnywhere, Category = "Debug Overview", meta = (DisplayName = "Debug Overview Sort Mode"))
 	ENiagaraDebugHUDDOverviewSort OverviewSortMode = ENiagaraDebugHUDDOverviewSort::Name;
 
-	/** When enabled the overview display will include cascade FX. */
+	/** When enabled the overview display will include Cascade FX. */
 	UPROPERTY(EditAnywhere, Category = "Debug Overview", AdvancedDisplay)
 	bool bIncludeCascade = true;
 
+	/** When enabled the overview display will include Niagara FX. */
+	UPROPERTY(EditAnywhere, Category = "Debug Overview", AdvancedDisplay)
+	bool bIncludeNiagara = true;
+
 	UPROPERTY(EditAnywhere, Category = "Debug Overview")
 	bool bShowRegisteredComponents = false;
 	
@@ -497,11 +501,15 @@ struct FNiagaraDebugHUDSettingsData
 
 	UPROPERTY(EditAnywhere, Category = "Debug Overview", AdvancedDisplay)
 	bool bShowGlobalBudgetInfo = false;
-	
+
+	/** Master control for enabling and disabling filtering by asset / actor / component */
+	UPROPERTY(EditAnywhere, Category = "Debug Filter")
+	bool bFiltersEnabled = false;
+
 	/**
 	If enabled, then only systems matching the filters will be considered by the debug hud
 	*/
-	UPROPERTY(EditAnywhere, Category = "Debug Filter", meta = (DisplayName = "Filter Displayed Systems"))
+	UPROPERTY(EditAnywhere, Category = "Debug Filter", meta = (InlineEditConditionToggle))
 	bool bSystemFilterEnabled = false;
 
 	/**
diff --git a/Engine/Plugins/FX/Niagara/Source/NiagaraEditor/Private/Customizations/NiagaraDebugHUDCustomization.cpp b/Engine/Plugins/FX/Niagara/Source/NiagaraEditor/Private/Customizations/NiagaraDebugHUDCustomization.cpp
--- a/Engine/Plugins/FX/Niagara/Source/NiagaraEditor/Private/Customizations/NiagaraDebugHUDCustomization.cpp
+++ b/Engine/Plugins/FX/Niagara/Source/NiagaraEditor/Private/Customizations/NiagaraDebugHUDCustomization.cpp
@@ -269,11 +269,11 @@ void FNiagaraDebugHUDSettingsDetailsCustomization::CustomizeDetails(IDetailLayou
 	{
 		IDetailCategoryBuilder& FilterCategory = DetailBuilder.EditCategory("Debug Filter");
 		
-		FilterCategory.AddProperty(DetailBuilder.GetProperty(GET_MEMBER_NAME_CHECKED(FNiagaraDebugHUDSettingsData, bSystemFilterEnabled), FNiagaraDebugHUDSettingsData::StaticStruct()));
-		MakeCustomAssetSearch(DetailBuilder, FilterCategory, DetailBuilder.GetProperty(GET_MEMBER_NAME_CHECKED(FNiagaraDebugHUDSettingsData, SystemFilter), FNiagaraDebugHUDSettingsData::StaticStruct()), UNiagaraSystem::StaticClass(), [&]() -> bool& { return WeakSettings.Get()->Data.bSystemFilterEnabled; }, [&]() -> bool { return true; }, false);
-		MakeCustomAssetSearch(DetailBuilder, FilterCategory, DetailBuilder.GetProperty(GET_MEMBER_NAME_CHECKED(FNiagaraDebugHUDSettingsData, EmitterFilter), FNiagaraDebugHUDSettingsData::StaticStruct()), UNiagaraEmitter::StaticClass(), [&]() -> bool& { return WeakSettings.Get()->Data.bEmitterFilterEnabled; }, [&]() -> bool { return WeakSettings.Get()->Data.bSystemFilterEnabled; });
-		MakeCustomAssetSearch(DetailBuilder, FilterCategory, DetailBuilder.GetProperty(GET_MEMBER_NAME_CHECKED(FNiagaraDebugHUDSettingsData, ActorFilter), FNiagaraDebugHUDSettingsData::StaticStruct()), AActor::StaticClass(), [&]() -> bool& { return WeakSettings.Get()->Data.bActorFilterEnabled; }, [&]() -> bool { return WeakSettings.Get()->Data.bSystemFilterEnabled; });
-		MakeCustomAssetSearch(DetailBuilder, FilterCategory, DetailBuilder.GetProperty(GET_MEMBER_NAME_CHECKED(FNiagaraDebugHUDSettingsData, ComponentFilter), FNiagaraDebugHUDSettingsData::StaticStruct()), UNiagaraComponent::StaticClass(), [&]() -> bool& { return WeakSettings.Get()->Data.bComponentFilterEnabled; }, [&]() -> bool { return WeakSettings.Get()->Data.bSystemFilterEnabled; });
+		FilterCategory.AddProperty(DetailBuilder.GetProperty(GET_MEMBER_NAME_CHECKED(FNiagaraDebugHUDSettingsData, bFiltersEnabled), FNiagaraDebugHUDSettingsData::StaticStruct()));
+		MakeCustomAssetSearch(DetailBuilder, FilterCategory, DetailBuilder.GetProperty(GET_MEMBER_NAME_CHECKED(FNiagaraDebugHUDSettingsData, SystemFilter), FNiagaraDebugHUDSettingsData::StaticStruct()), UNiagaraSystem::StaticClass(), [&]() -> bool& { return WeakSettings.Get()->Data.bSystemFilterEnabled; }, [&]() -> bool { return WeakSettings.Get()->Data.bFiltersEnabled; });
+		MakeCustomAssetSearch(DetailBuilder, FilterCategory, DetailBuilder.GetProperty(GET_MEMBER_NAME_CHECKED(FNiagaraDebugHUDSettingsData, EmitterFilter), FNiagaraDebugHUDSettingsData::StaticStruct()), UNiagaraEmitter::StaticClass(), [&]() -> bool& { return WeakSettings.Get()->Data.bEmitterFilterEnabled; }, [&]() -> bool { return WeakSettings.Get()->Data.bFiltersEnabled; });
+		MakeCustomAssetSearch(DetailBuilder, FilterCategory, DetailBuilder.GetProperty(GET_MEMBER_NAME_CHECKED(FNiagaraDebugHUDSettingsData, ActorFilter), FNiagaraDebugHUDSettingsData::StaticStruct()), AActor::StaticClass(), [&]() -> bool& { return WeakSettings.Get()->Data.bActorFilterEnabled; }, [&]() -> bool { return WeakSettings.Get()->Data.bFiltersEnabled; });
+		MakeCustomAssetSearch(DetailBuilder, FilterCategory, DetailBuilder.GetProperty(GET_MEMBER_NAME_CHECKED(FNiagaraDebugHUDSettingsData, ComponentFilter), FNiagaraDebugHUDSettingsData::StaticStruct()), UNiagaraComponent::StaticClass(), [&]() -> bool& { return WeakSettings.Get()->Data.bComponentFilterEnabled; }, [&]() -> bool { return WeakSettings.Get()->Data.bFiltersEnabled; });
 	}
 }
 
