- Use unchecked method as we handle cast failure

#Synced-CL 50610072

diff --git a/Engine/Plugins/FX/Niagara/Source/NiagaraEditor/Private/ViewModels/Stack/NiagaraStackGraphUtilities.cpp b/Engine/Plugins/FX/Niagara/Source/NiagaraEditor/Private/ViewModels/Stack/NiagaraStackGraphUtilities.cpp
--- a/Engine/Plugins/FX/Niagara/Source/NiagaraEditor/Private/ViewModels/Stack/NiagaraStackGraphUtilities.cpp
+++ b/Engine/Plugins/FX/Niagara/Source/NiagaraEditor/Private/ViewModels/Stack/NiagaraStackGraphUtilities.cpp
@@ -828,7 +828,7 @@ void GetGroupNodesRecursive(const TArray<UNiagaraNode*>& CurrentStartNodes, UNia
 					{
 						for (UEdGraphPin* InputLinkedPin : InputPin->LinkedTo)
 						{
-							UNiagaraNode* LinkedNode = Cast<UNiagaraNode>(InputLinkedPin->GetOwningNode());
+							UNiagaraNode* LinkedNode = Cast<UNiagaraNode>(InputLinkedPin->GetOwningNodeUnchecked());
 							if (LinkedNode != nullptr)
 							{
 								OutAllNodes.AddUnique(LinkedNode);
@@ -848,7 +848,7 @@ void GetGroupNodesRecursive(const TArray<UNiagaraNode*>& CurrentStartNodes, UNia
 				{
 					for (UEdGraphPin* LinkedPin : OutputPin->LinkedTo)
 					{
-						UNiagaraNode* LinkedNode = Cast<UNiagaraNode>(LinkedPin->GetOwningNode());
+						UNiagaraNode* LinkedNode = Cast<UNiagaraNode>(LinkedPin->GetOwningNodeUnchecked());
 						if (LinkedNode != nullptr)
 						{
 							LinkedNodes.Add(LinkedNode);
