- Add missing attributes to mesh renderer
- Clean up some code while here

#Synced-CL 50610072

diff --git a/Engine/Plugins/FX/Niagara/Source/Niagara/Private/NiagaraMeshRendererProperties.cpp b/Engine/Plugins/FX/Niagara/Source/Niagara/Private/NiagaraMeshRendererProperties.cpp
--- a/Engine/Plugins/FX/Niagara/Source/Niagara/Private/NiagaraMeshRendererProperties.cpp
+++ b/Engine/Plugins/FX/Niagara/Source/Niagara/Private/NiagaraMeshRendererProperties.cpp
@@ -1002,6 +1002,31 @@ TArray<FNiagaraVariable> UNiagaraMeshRendererProperties::GetBoundAttributes() co
 	{
 		BoundAttributes.AddUnique(MaterialParamBinding.GetParamMapBindableVariable());
 	}
+
+	if (MeshesBinding.IsSetToDefault() == false)
+	{
+		BoundAttributes.AddUnique(MeshesBinding.AliasedParameter);
+	}
+
+	for (const FNiagaraMeshRendererMeshProperties& Mesh : Meshes)
+	{
+		if (Mesh.MeshParameterBinding.IsSetToDefault() == false)
+		{
+			BoundAttributes.AddUnique(Mesh.MeshParameterBinding.AliasedParameter);
+		}
+	}
+
+	if (bOverrideMaterials)
+	{
+		for (const FNiagaraMeshMaterialOverride& OverrideMaterial : OverrideMaterials)
+		{
+			if (OverrideMaterial.UserParamBinding.Parameter.IsValid())
+			{
+				BoundAttributes.AddUnique(OverrideMaterial.UserParamBinding.Parameter);
+			}
+		}
+	}
+
 	return BoundAttributes;
 }
 #endif
@@ -1477,7 +1502,6 @@ void UNiagaraMeshRendererProperties::CheckMaterialUsage()
 				UMaterialInterface *MaterialInterface = MeshProperties.Mesh->GetMaterial(Section.MaterialIndex);
 				if (MaterialInterface)
 				{
-					const UMaterial* Material = MaterialInterface->GetMaterial();
 					MaterialInterface->CheckMaterialUsage(MATUSAGE_NiagaraMeshParticles);
 				}
 			}
diff --git a/Engine/Plugins/FX/Niagara/Source/Niagara/Private/NiagaraRendererMeshes.cpp b/Engine/Plugins/FX/Niagara/Source/Niagara/Private/NiagaraRendererMeshes.cpp
--- a/Engine/Plugins/FX/Niagara/Source/Niagara/Private/NiagaraRendererMeshes.cpp
+++ b/Engine/Plugins/FX/Niagara/Source/Niagara/Private/NiagaraRendererMeshes.cpp
@@ -1970,11 +1970,5 @@ int FNiagaraRendererMeshes::GetDynamicDataSize()const
 
 bool FNiagaraRendererMeshes::IsMaterialValid(const UMaterialInterface* Mat)const
 {
-	bool bIsMaterialValid = false;
-	if (Mat)
-	{
-		bIsMaterialValid = Mat->CheckMaterialUsage_Concurrent(MATUSAGE_NiagaraMeshParticles);
-	}
-
-	return bIsMaterialValid;
+	return Mat && Mat->CheckMaterialUsage_Concurrent(MATUSAGE_NiagaraMeshParticles);
 }
