- Add double click navigation to source for renderers

#Synced-CL 50610072

diff --git a/Engine/Plugins/FX/Niagara/Source/NiagaraEditor/Private/ViewModels/Stack/NiagaraStackRendererItem.cpp b/Engine/Plugins/FX/Niagara/Source/NiagaraEditor/Private/ViewModels/Stack/NiagaraStackRendererItem.cpp
--- a/Engine/Plugins/FX/Niagara/Source/NiagaraEditor/Private/ViewModels/Stack/NiagaraStackRendererItem.cpp
+++ b/Engine/Plugins/FX/Niagara/Source/NiagaraEditor/Private/ViewModels/Stack/NiagaraStackRendererItem.cpp
@@ -31,6 +31,7 @@
 
 #include "Materials/MaterialInterface.h"
 #include "MaterialEditorValidation.h"
+#include "SourceCodeNavigation.h"
 
 #include UE_INLINE_GENERATED_CPP_BY_NAME(NiagaraStackRendererItem)
 
@@ -273,6 +274,14 @@ const UNiagaraRendererProperties* UNiagaraStackRendererItem::GetRendererProperti
 	return RendererProperties.Get();
 }
 
+void UNiagaraStackRendererItem::OpenSourceCodeLocation() const
+{
+	if (UNiagaraRendererProperties* RendererProps = RendererProperties.Get())
+	{
+		FSourceCodeNavigation::NavigateToClass(RendererProps->GetClass());
+	}
+}
+
 FText UNiagaraStackRendererItem::GetDisplayName() const
 {
 	if(DisplayNameCache.IsSet() == false)
@@ -332,7 +341,7 @@ bool UNiagaraStackRendererItem::TestCanCopyWithMessage(FText& OutMessage) const
 
 void UNiagaraStackRendererItem::Copy(UNiagaraClipboardContent* ClipboardContent) const
 {
-	ClipboardContent->Renderers.Add(UNiagaraClipboardRenderer::CreateRenderer(ClipboardContent, RendererProperties.Get(), GetStackNoteData()));
+	ClipboardContent->Renderers.Add(UNiagaraClipboardRenderer::CreateRenderer(ClipboardContent, RendererProperties.Get(), GetStackNoteData(), GetAlternateDisplayName()));
 }
 
 bool UNiagaraStackRendererItem::TestCanPasteWithMessage(const UNiagaraClipboardContent* ClipboardContent, FText& OutMessage) const
@@ -802,4 +811,3 @@ TArray<FNiagaraStackItemCustomIcon> UNiagaraStackRendererItem::GetCustomIcons()
 }
 
 #undef LOCTEXT_NAMESPACE
-
diff --git a/Engine/Plugins/FX/Niagara/Source/NiagaraEditor/Public/ViewModels/Stack/NiagaraStackRendererItem.h b/Engine/Plugins/FX/Niagara/Source/NiagaraEditor/Public/ViewModels/Stack/NiagaraStackRendererItem.h
--- a/Engine/Plugins/FX/Niagara/Source/NiagaraEditor/Public/ViewModels/Stack/NiagaraStackRendererItem.h
+++ b/Engine/Plugins/FX/Niagara/Source/NiagaraEditor/Public/ViewModels/Stack/NiagaraStackRendererItem.h
@@ -25,7 +25,9 @@ public:
 
 	NIAGARAEDITOR_API UNiagaraRendererProperties* GetRendererProperties();
 	NIAGARAEDITOR_API const UNiagaraRendererProperties* GetRendererProperties() const;
-	
+
+	NIAGARAEDITOR_API void OpenSourceCodeLocation() const;
+
 	NIAGARAEDITOR_API virtual FText GetDisplayName() const override;
 
 	virtual bool SupportsCut() const override { return true; }
diff --git a/Engine/Plugins/FX/Niagara/Source/NiagaraEditorWidgets/Private/SNiagaraOverviewStack.cpp b/Engine/Plugins/FX/Niagara/Source/NiagaraEditorWidgets/Private/SNiagaraOverviewStack.cpp
--- a/Engine/Plugins/FX/Niagara/Source/NiagaraEditorWidgets/Private/SNiagaraOverviewStack.cpp
+++ b/Engine/Plugins/FX/Niagara/Source/NiagaraEditorWidgets/Private/SNiagaraOverviewStack.cpp
@@ -311,7 +311,6 @@ class SNiagaraSystemOverviewEntryListRow : public STableRow<UNiagaraStackEntry*>
 	{
 		if (UNiagaraStackModuleItem* ModuleItem = Cast<UNiagaraStackModuleItem>(StackEntry))
 		{
-			
 			// if the user double-clicks a collapsed module we uncollapse it
 			UNiagaraStackEditorData& StackEditorData = ModuleItem->GetSystemViewModel()->GetEditorData().GetStackEditorData();
 			if (ModuleItem->GetIsEnabled() == false && StackEditorData.bHideDisabledModules)
@@ -325,6 +324,11 @@ class SNiagaraSystemOverviewEntryListRow : public STableRow<UNiagaraStackEntry*>
 				return FReply::Handled();
 			}			
 		}
+		else if (UNiagaraStackRendererItem* RenderItem = Cast<UNiagaraStackRendererItem>(StackEntry))
+		{
+			RenderItem->OpenSourceCodeLocation();
+			return FReply::Handled();
+		}
 		return FReply::Unhandled();
 	}
 
diff --git a/Engine/Plugins/FX/Niagara/Source/NiagaraEditorWidgets/Private/Stack/SNiagaraStackTableRow.cpp b/Engine/Plugins/FX/Niagara/Source/NiagaraEditorWidgets/Private/Stack/SNiagaraStackTableRow.cpp
--- a/Engine/Plugins/FX/Niagara/Source/NiagaraEditorWidgets/Private/Stack/SNiagaraStackTableRow.cpp
+++ b/Engine/Plugins/FX/Niagara/Source/NiagaraEditorWidgets/Private/Stack/SNiagaraStackTableRow.cpp
@@ -364,8 +364,17 @@ void SNiagaraStackTableRow::AddFillRowContextMenuHandler(FOnFillRowContextMenu F
 	OnFillRowContextMenuHanders.Add(FillRowContextMenuHandler);
 }
 
-FReply SNiagaraStackTableRow::OnMouseButtonDoubleClick(const FGeometry& InMyGeometry, const FPointerEvent& InMouseEvent)
+FReply SNiagaraStackTableRow::OnMouseButtonDoubleClick(const FGeometry& InMyGeometry, const FPointerEvent& MouseEvent)
 {
+	if (MouseEvent.GetEffectingButton() == EKeys::LeftMouseButton)
+	{
+		if (UNiagaraStackRendererItem* RendererItem = Cast<UNiagaraStackRendererItem>(StackEntry))
+		{
+			RendererItem->OpenSourceCodeLocation();
+			return FReply::Handled();
+		}
+	}
+
 	return FReply::Unhandled();
 }
 
