- Set editor renderer name when copy and pasting

#Synced-CL 50610072

diff --git a/Engine/Plugins/FX/Niagara/Source/NiagaraEditor/Private/NiagaraClipboard.cpp b/Engine/Plugins/FX/Niagara/Source/NiagaraEditor/Private/NiagaraClipboard.cpp
--- a/Engine/Plugins/FX/Niagara/Source/NiagaraEditor/Private/NiagaraClipboard.cpp
+++ b/Engine/Plugins/FX/Niagara/Source/NiagaraEditor/Private/NiagaraClipboard.cpp
@@ -131,14 +131,18 @@ bool UNiagaraClipboardFunctionInput::CopyValuesFrom(const UNiagaraClipboardFunct
 	return true;
 }
 
-const UNiagaraClipboardRenderer* UNiagaraClipboardRenderer::CreateRenderer(UObject* InOuter, UNiagaraRendererProperties* Renderer, TOptional<FNiagaraStackNoteData> StackNoteData)
+const UNiagaraClipboardRenderer* UNiagaraClipboardRenderer::CreateRenderer(UObject* InOuter, UNiagaraRendererProperties* Renderer, TOptional<FNiagaraStackNoteData> StackNoteData, TOptional<FText> AlternateDisplayName)
 {
 	UNiagaraClipboardRenderer* NewRenderer = Cast<UNiagaraClipboardRenderer>(NewObject<UNiagaraClipboardRenderer>(InOuter));
 	NewRenderer->RendererProperties = CastChecked<UNiagaraRendererProperties>(StaticDuplicateObject(Renderer, InOuter));
 
-	if(StackNoteData.IsSet())
+	if (StackNoteData.IsSet() && StackNoteData.GetValue().IsValid())
 	{
-		NewRenderer->StackNoteData = StackNoteData.GetValue();
+		NewRenderer->StackNoteData = StackNoteData;
+	}
+	if (AlternateDisplayName.IsSet() && !AlternateDisplayName.GetValue().IsEmpty())
+	{
+		NewRenderer->AlternateDisplayName = AlternateDisplayName;
 	}
 
 	return NewRenderer;
diff --git a/Engine/Plugins/FX/Niagara/Source/NiagaraEditor/Private/ViewModels/Stack/NiagaraStackRenderItemGroup.cpp b/Engine/Plugins/FX/Niagara/Source/NiagaraEditor/Private/ViewModels/Stack/NiagaraStackRenderItemGroup.cpp
--- a/Engine/Plugins/FX/Niagara/Source/NiagaraEditor/Private/ViewModels/Stack/NiagaraStackRenderItemGroup.cpp
+++ b/Engine/Plugins/FX/Niagara/Source/NiagaraEditor/Private/ViewModels/Stack/NiagaraStackRenderItemGroup.cpp
@@ -193,9 +193,13 @@ void UNiagaraStackRenderItemGroup::Paste(const UNiagaraClipboardContent* Clipboa
 
 			UNiagaraRendererProperties* NewRenderer = ClipboardRenderer->RendererProperties->StaticDuplicateWithNewMergeId(RenderersOwner->GetOwnerObject());
 			RenderersOwner->AddRenderer(NewRenderer);
-			if(ClipboardRenderer->StackNoteData.IsValid())
+			if(ClipboardRenderer->StackNoteData.IsSet())
 			{
-				GetStackEditorData().AddOrReplaceStackNote(FNiagaraStackGraphUtilities::StackKeys::GenerateStackRendererEditorDataKey(*NewRenderer), ClipboardRenderer->StackNoteData);
+				GetStackEditorData().AddOrReplaceStackNote(FNiagaraStackGraphUtilities::StackKeys::GenerateStackRendererEditorDataKey(*NewRenderer), ClipboardRenderer->StackNoteData.GetValue());
+			}
+			if (ClipboardRenderer->AlternateDisplayName.IsSet())
+			{
+				GetStackEditorData().SetStackEntryDisplayName(FNiagaraStackGraphUtilities::StackKeys::GenerateStackRendererEditorDataKey(*NewRenderer), ClipboardRenderer->AlternateDisplayName.GetValue());
 			}
 		}
 	}
diff --git a/Engine/Plugins/FX/Niagara/Source/NiagaraEditor/Public/NiagaraClipboard.h b/Engine/Plugins/FX/Niagara/Source/NiagaraEditor/Public/NiagaraClipboard.h
--- a/Engine/Plugins/FX/Niagara/Source/NiagaraEditor/Public/NiagaraClipboard.h
+++ b/Engine/Plugins/FX/Niagara/Source/NiagaraEditor/Public/NiagaraClipboard.h
@@ -101,13 +101,16 @@ class UNiagaraClipboardRenderer : public UObject
 	GENERATED_BODY()
 
 public:
-	static NIAGARAEDITOR_API const UNiagaraClipboardRenderer* CreateRenderer(UObject* InOuter, UNiagaraRendererProperties* Renderer, TOptional<FNiagaraStackNoteData> StackNoteData = TOptional<FNiagaraStackNoteData>());
+	static NIAGARAEDITOR_API const UNiagaraClipboardRenderer* CreateRenderer(UObject* InOuter, UNiagaraRendererProperties* Renderer, TOptional<FNiagaraStackNoteData> StackNoteData = TOptional<FNiagaraStackNoteData>(), TOptional<FText> AlternateDisplayName = TOptional<FText>());
 
 	UPROPERTY()
 	TObjectPtr<UNiagaraRendererProperties> RendererProperties;
 
 	UPROPERTY()
-	FNiagaraStackNoteData StackNoteData;
+	TOptional<FNiagaraStackNoteData> StackNoteData;
+
+	UPROPERTY()
+	TOptional<FText> AlternateDisplayName;
 };
 
 UENUM()
