- Fix potential nullptr access

#Synced-CL 50617493

diff --git a/Engine/Plugins/FX/Niagara/Source/Niagara/Private/NiagaraDebugHud.cpp b/Engine/Plugins/FX/Niagara/Source/Niagara/Private/NiagaraDebugHud.cpp
--- a/Engine/Plugins/FX/Niagara/Source/Niagara/Private/NiagaraDebugHud.cpp
+++ b/Engine/Plugins/FX/Niagara/Source/Niagara/Private/NiagaraDebugHud.cpp
@@ -1120,7 +1120,7 @@ void FNiagaraDebugHud::GatherSystemInfo()
 		}
 
 		const bool bIsActive = FXComponent->IsActive();
-		const bool bHasScalability = NiagaraComponent ? NiagaraComponent->IsRegisteredWithScalabilityManager() : CascadeComponent->bIsManagingSignificance;
+		const bool bHasScalability = NiagaraComponent ? NiagaraComponent->IsRegisteredWithScalabilityManager() : (CascadeComponent ? CascadeComponent->bIsManagingSignificance : false);
 		const bool bUsingCullProxy = NiagaraComponent ? NiagaraComponent->IsUsingCullProxy() : false;
 		const bool bIsCullProxy = FXComponent->IsA<UNiagaraCullProxyComponent>();
 
@@ -1273,9 +1273,8 @@ void FNiagaraDebugHud::GatherSystemInfo()
 					ActiveParticles += EmitterInstance->GetNumParticles();
 				}
 			}
-			else 
+			else if (CascadeComponent)
 			{
-				check(CascadeComponent);				
 				for(FParticleEmitterInstance* Emitter : CascadeComponent->EmitterInstances)
 				{					
 					if(Emitter)
