- Display information about how many spawn infos we are discarding in LWE

#Synced-CL 50657737

diff --git a/Engine/Plugins/FX/Niagara/Source/Niagara/Private/Stateless/NiagaraStatelessEmitterData.cpp b/Engine/Plugins/FX/Niagara/Source/Niagara/Private/Stateless/NiagaraStatelessEmitterData.cpp
--- a/Engine/Plugins/FX/Niagara/Source/Niagara/Private/Stateless/NiagaraStatelessEmitterData.cpp
+++ b/Engine/Plugins/FX/Niagara/Source/Niagara/Private/Stateless/NiagaraStatelessEmitterData.cpp
@@ -64,6 +64,10 @@ uint32 FNiagaraStatelessEmitterData::CalculateActiveParticles(int32 InRandomSeed
 	int32	GpuSpawnIndex = 0;
 	uint32	TotalActiveParticles = 0;
 
+#if WITH_EDITOR && WITH_NIAGARA_DEBUG_EMITTER_NAME
+	uint32	AmountOverLimit = 0;
+#endif
+
 	for (const FNiagaraStatelessRuntimeSpawnInfo& SpawnInfo : RuntimeSpawnInfos)
 	{
 		uint32 SpawnInfoTotalParticles = 0;
@@ -111,9 +115,11 @@ uint32 FNiagaraStatelessEmitterData::CalculateActiveParticles(int32 InRandomSeed
 		if (GpuSpawnIndex >= NiagaraStateless::MaxGpuSpawnInfos)
 		{
 		#if WITH_EDITOR && WITH_NIAGARA_DEBUG_EMITTER_NAME
-			GEngine->AddOnScreenDebugMessage(uint64(this), 1.f, FColor::White, *FString::Printf(TEXT("Stateless Simulation (%s:%s) has run out of space to store spawn infos results may flicker.  Reduce the number of spawn infos / loop time / particle age / increase code limit to fix."), *DebugSimulationName.ToString(), *DebugEmitterName.ToString()));
-		#endif
+			++AmountOverLimit;
+			continue;
+		#else
 			break;
+		#endif
 		}
 
 		if (SpawnParameters)
@@ -131,6 +137,23 @@ uint32 FNiagaraStatelessEmitterData::CalculateActiveParticles(int32 InRandomSeed
 		TotalActiveParticles += NumActive;
 	}
 
+#if WITH_EDITOR && WITH_NIAGARA_DEBUG_EMITTER_NAME
+	if (AmountOverLimit > 0)
+	{
+		GEngine->AddOnScreenDebugMessage(
+			uint64(this),
+			1.f,
+			FColor::White,
+			*FString::Printf(
+				TEXT("Stateless Simulation (%s:%s) has overflowed spawn infos (%d of %d), this can result in flickering / missing particles.  Reduce the number of spawn infos / loop time / particle age / increase code limit to fix."),
+				*DebugSimulationName.ToString(), *DebugEmitterName.ToString(),
+				NiagaraStateless::MaxGpuSpawnInfos + AmountOverLimit,
+				NiagaraStateless::MaxGpuSpawnInfos
+			)
+		);
+	}
+#endif
+
 	if (SpawnParameters)
 	{
 		while (GpuSpawnIndex < NiagaraStateless::MaxGpuSpawnInfos)
