- Fix for invalid texture coordinates when sprite size is 0.0

#Synced-CL 50731656

#Synced-CL 50731656

diff --git a/Engine/Plugins/FX/Niagara/Shaders/Private/NiagaraSpriteVertexFactory.ush b/Engine/Plugins/FX/Niagara/Shaders/Private/NiagaraSpriteVertexFactory.ush
--- a/Engine/Plugins/FX/Niagara/Shaders/Private/NiagaraSpriteVertexFactory.ush
+++ b/Engine/Plugins/FX/Niagara/Shaders/Private/NiagaraSpriteVertexFactory.ush
@@ -57,6 +57,11 @@ bool GetCutoutUseSubImage()
 	return (NiagaraSpriteVFLooseParameters.CutoutParameters & 0x1) != 0;
 }
 
+float2 NoneZeroSign(float2 v)
+{
+	return float2(v.x >= 0.0f ? 1.0f : -1.0f, v.y >= 0.0f ? 1.0f : -1.0f);
+}
+
 // RHI_RAYTRACING
 #if COMPUTESHADER || RAYHITGROUPSHADER
 FVertexFactoryInput LoadVertexFactoryInputForHGS(uint TriangleIndex, int VertexIndex)
@@ -637,7 +642,7 @@ float3x3 CalcTangentBasis(FVertexFactoryIntermediates Intermediates)
 	// So there are artifacts with tangent space calculations like the TransformVector node,
 	// But this allows lighting based on a world space shape via the normal while still having normal maps camera aligned.
 	float3x3 Result;
-	const float2 UVSign = sign(Intermediates.UVScale);
+	const float2 UVSign = NoneZeroSign(Intermediates.UVScale);
 	Result[0] = Intermediates.TangentRight * UVSign.x;
 	Result[1] = Intermediates.TangentUp * UVSign.y;
 
@@ -684,7 +689,7 @@ void ComputeBillboardUVs(FVertexFactoryInput Input, float2 ParticleSize, float S
 		const float NumFrames = NiagaraSpriteVF.SubImageSize.x * NiagaraSpriteVF.SubImageSize.y;
 		const uint SubImageIndexInt = (uint)fmod(SubImageIndex, NumFrames);
 
-		const float2 UVFlip = sign(ParticleSize) * sign(UVScale);
+		const float2 UVFlip = NoneZeroSign(ParticleSize) * NoneZeroSign(UVScale);
 		
 		// Reverse the winding order of the polygon when only flipping in one direction to counteract UV flipping
 		FLATTEN
@@ -703,7 +708,7 @@ void ComputeBillboardUVs(FVertexFactoryInput Input, float2 ParticleSize, float S
 		UVForPosition.x = UVFlip.x < 0.0f ? 1.0f - UVForTexturing.x : UVForTexturing.x;
 		UVForPosition.y = UVFlip.y < 0.0f ? 1.0f - UVForTexturing.y : UVForTexturing.y;
 
-		UVForTexturing = float2(0.5f, 0.5f) + ((UVForPosition - float2(0.5f, 0.5f)) * UVScale * sign(ParticleSize));
+		UVForTexturing = float2(0.5f, 0.5f) + ((UVForPosition - float2(0.5f, 0.5f)) * UVScale * NoneZeroSign(ParticleSize));
 	}
 	else
 #endif
@@ -711,7 +716,7 @@ void ComputeBillboardUVs(FVertexFactoryInput Input, float2 ParticleSize, float S
 		// Note: not inverting positions, as that would change the winding order
 		UVForPosition = Input.TexCoord.xy;
 		UVForTexturingUnflipped = UVForPosition;
-		UVForTexturing = float2(0.5f, 0.5f) + ((UVForTexturingUnflipped - float2(0.5f, 0.5f)) * UVScale * sign(ParticleSize));
+		UVForTexturing = float2(0.5f, 0.5f) + ((UVForTexturingUnflipped - float2(0.5f, 0.5f)) * UVScale * NoneZeroSign(ParticleSize));
 	}
 }
 
@@ -1380,10 +1385,7 @@ FLWCVector3 NiagaraVertexFactoryExport_GetPosition(FVertexFactoryInput Input, FV
 
 float4 NiagaraVertexFactoryExport_GetColor(FVertexFactoryInput Input, FVertexFactoryIntermediates Intermediates)
 {
-#if NEEDS_PARTICLE_COLOR
 	return Intermediates.Color;
-#endif
-	return float4(1,1,1,1);
 }
 
 float3x3 NiagaraVertexFactoryExport_GetTangentBasis(FVertexFactoryInput Input, FVertexFactoryIntermediates Intermediates)
