- Fix for CIS issue no caught by PF

#rnx

#Synced-CL 50731656

diff --git a/Engine/Plugins/FX/Niagara/Source/Niagara/Internal/DataInterface/NiagaraDataInterfacePropertyInterface.h b/Engine/Plugins/FX/Niagara/Source/Niagara/Internal/DataInterface/NiagaraDataInterfacePropertyInterface.h
--- a/Engine/Plugins/FX/Niagara/Source/Niagara/Internal/DataInterface/NiagaraDataInterfacePropertyInterface.h
+++ b/Engine/Plugins/FX/Niagara/Source/Niagara/Internal/DataInterface/NiagaraDataInterfacePropertyInterface.h
@@ -61,9 +61,7 @@ public:
 		return !InterfaceName.IsNone() && !InterfacePath.IsNone() && !InterfacePackage.IsNone();
 	}
 
-#if WITH_EDITORONLY_DATA
 	NIAGARA_API TArray<FNiagaraVariableBase> FindUsedVariables(UNiagaraSystem* OwningSystem);
-#endif //WITH_EDITORONLY_DATA
 
 public:
 	/** The Class name for the interface */
diff --git a/Engine/Plugins/FX/Niagara/Source/Niagara/Private/DataInterface/NiagaraDataInterfacePropertyInterface.cpp b/Engine/Plugins/FX/Niagara/Source/Niagara/Private/DataInterface/NiagaraDataInterfacePropertyInterface.cpp
--- a/Engine/Plugins/FX/Niagara/Source/Niagara/Private/DataInterface/NiagaraDataInterfacePropertyInterface.cpp
+++ b/Engine/Plugins/FX/Niagara/Source/Niagara/Private/DataInterface/NiagaraDataInterfacePropertyInterface.cpp
@@ -661,11 +661,13 @@ void UNiagaraDataInterfacePropertyInterface::DrawDebugHud(FNDIDrawDebugHudContex
 }
 #endif //WITH_NIAGARA_DEBUGGER
 
-#if WITH_EDITORONLY_DATA
 TArray<FNiagaraVariableBase> UNiagaraDataInterfacePropertyInterface::FindUsedVariables(UNiagaraSystem* OwningSystem)
 {
 	using namespace NDIPropertyInterfaceLocal;
 
+	TArray<FNiagaraVariableBase> Variables;
+
+#if WITH_EDITORONLY_DATA
 	INiagaraModule& NiagaraModule = FModuleManager::GetModuleChecked<INiagaraModule>("Niagara");
 	const INiagaraEditorOnlyDataUtilities& EditorOnlyDataUtilities = NiagaraModule.GetEditorOnlyDataUtilities();
 	UNiagaraDataInterface* RuntimeInstanceOfThis =
@@ -673,8 +675,6 @@ TArray<FNiagaraVariableBase> UNiagaraDataInterfacePropertyInterface::FindUsedVar
 		? EditorOnlyDataUtilities.GetResolvedRuntimeInstanceForEditorDataInterfaceInstance(*OwningSystem, *this)
 		: this;
 
-	TArray<FNiagaraVariableBase> Variables;
-
 	FNiagaraDataInterfaceUtilities::ForEachVMFunction(
 		RuntimeInstanceOfThis,
 		OwningSystem,
@@ -706,9 +706,9 @@ TArray<FNiagaraVariableBase> UNiagaraDataInterfacePropertyInterface::FindUsedVar
 			return true;
 		}
 	);
+#endif
 
 	return Variables;
 }
-#endif //WITH_EDITORONLY_DATA
 
 #undef LOCTEXT_NAMESPACE
