[Backout] - CL50781810
#rnx
Original CL Desc
-----------------------------------------------------------------
- Niagara Verse Property Interfaces (IsA)
- Uses a data interface to read from the interface data (temporary)
- Digest generation reworked to handle getter / setter which is required for interfaces
- Getter / setter currently disabled for regular user parameters, that will be enabled in a future change

#Synced-CL 50731656

#Synced-CL 50731656

diff --git a/Engine/Plugins/FX/Niagara/Shaders/Private/NiagaraDataInterfacePropertyInterfaceTemplate.ush b/Engine/Plugins/FX/Niagara/Shaders/Private/NiagaraDataInterfacePropertyInterfaceTemplate.ush
deleted file mode 100644
diff --git a/Engine/Plugins/FX/Niagara/Source/Niagara/Internal/DataInterface/NiagaraDataInterfacePropertyInterface.h b/Engine/Plugins/FX/Niagara/Source/Niagara/Internal/DataInterface/NiagaraDataInterfacePropertyInterface.h
deleted file mode 100644
diff --git a/Engine/Plugins/FX/Niagara/Source/Niagara/Private/DataInterface/NiagaraDataInterfaceDataTable.cpp b/Engine/Plugins/FX/Niagara/Source/Niagara/Private/DataInterface/NiagaraDataInterfaceDataTable.cpp
--- a/Engine/Plugins/FX/Niagara/Source/Niagara/Private/DataInterface/NiagaraDataInterfaceDataTable.cpp
+++ b/Engine/Plugins/FX/Niagara/Source/Niagara/Private/DataInterface/NiagaraDataInterfaceDataTable.cpp
@@ -556,6 +556,7 @@ bool UNiagaraDataInterfaceDataTable::Equals(const UNiagaraDataInterface* Other)
 		return false;
 	}
 
+	// Note: We intentionally do not include bReadsAnyRow for equality as that's part of PostCompile
 	const UNiagaraDataInterfaceDataTable* OtherTyped = CastChecked<const UNiagaraDataInterfaceDataTable>(Other);
 	return
 		OtherTyped->DataTable == DataTable &&
diff --git a/Engine/Plugins/FX/Niagara/Source/Niagara/Private/DataInterface/NiagaraDataInterfacePropertyInterface.cpp b/Engine/Plugins/FX/Niagara/Source/Niagara/Private/DataInterface/NiagaraDataInterfacePropertyInterface.cpp
deleted file mode 100644
diff --git a/Engine/Plugins/FX/Niagara/Source/Niagara/Private/NiagaraSceneComponentUtils.cpp b/Engine/Plugins/FX/Niagara/Source/Niagara/Private/NiagaraSceneComponentUtils.cpp
--- a/Engine/Plugins/FX/Niagara/Source/Niagara/Private/NiagaraSceneComponentUtils.cpp
+++ b/Engine/Plugins/FX/Niagara/Source/Niagara/Private/NiagaraSceneComponentUtils.cpp
@@ -56,16 +56,6 @@ FNiagaraActorSceneComponentUtils::FNiagaraActorSceneComponentUtils(UNiagaraCompo
 {
 }
 
-UObject* FNiagaraActorSceneComponentUtils::GetOwner() const
-{
-	return WeakOwnerComponent.Get();
-}
-
-UClass* FNiagaraActorSceneComponentUtils::FindPropertyInterface(UObject* SourceObject, FName InterfaceFullPath) const
-{
-	return nullptr;
-}
-
 void FNiagaraActorSceneComponentUtils::ResolveStaticMesh(bool bRecurseParents, UObject*& OutComponent, UStaticMesh*& OutStaticMesh) const
 {
 	USceneComponent* OwnerComponent = WeakOwnerComponent.Get();
diff --git a/Engine/Plugins/FX/Niagara/Source/Niagara/Public/NiagaraSceneComponentUtils.h b/Engine/Plugins/FX/Niagara/Source/Niagara/Public/NiagaraSceneComponentUtils.h
--- a/Engine/Plugins/FX/Niagara/Source/Niagara/Public/NiagaraSceneComponentUtils.h
+++ b/Engine/Plugins/FX/Niagara/Source/Niagara/Public/NiagaraSceneComponentUtils.h
@@ -21,11 +21,6 @@ class INiagaraSceneComponentUtils
 public:
 	virtual ~INiagaraSceneComponentUtils() = default;
 
-	// Gets the underlying owner object
-	virtual UObject* GetOwner() const = 0;
-	// Finds a property interface for the provided source object
-	virtual UClass* FindPropertyInterface(UObject* SourceObject, FName InterfaceFullPath) const = 0;
-
 	// Resolve the static mesh from what the interface's owner object
 	virtual void ResolveStaticMesh(bool bRecurseParents, UObject*& OutComponent, UStaticMesh*& OutStaticMesh) const = 0;
 	// Resolve the static mesh from the provided object which could be a static mesh / component / actor / entity / etc
@@ -50,9 +45,6 @@ class FNiagaraActorSceneComponentUtils final : public INiagaraSceneComponentUtil
 public:
 	explicit FNiagaraActorSceneComponentUtils(UNiagaraComponent* OwnerComponent);
 
-	virtual UObject* GetOwner() const override;
-	virtual UClass* FindPropertyInterface(UObject* SourceObject, FName InterfaceFullPath) const override;
-
 	virtual void ResolveStaticMesh(bool bRecurseParents, UObject*& OutComponent, UStaticMesh*& OutStaticMesh) const override;
 	virtual void ResolveStaticMesh(UObject* ObjectFrom, bool bRecurseParents, UObject*& OutComponent, UStaticMesh*& OutStaticMesh) const override;
 	virtual bool GetStaticMeshTransforms(UObject* Component, FTransform& OutComponentTransform, TArray<FTransform>& OutInstanceTransforms) const override;
