Fix use of PreviewMesh in SkelMeshDI to wait for synchronous load to actually be complete before we access the streaming data

#Synced-CL 50694553

diff --git a/Engine/Plugins/FX/Niagara/Source/Niagara/Private/NiagaraDataInterfaceSkeletalMesh.cpp b/Engine/Plugins/FX/Niagara/Source/Niagara/Private/NiagaraDataInterfaceSkeletalMesh.cpp
--- a/Engine/Plugins/FX/Niagara/Source/Niagara/Private/NiagaraDataInterfaceSkeletalMesh.cpp
+++ b/Engine/Plugins/FX/Niagara/Source/Niagara/Private/NiagaraDataInterfaceSkeletalMesh.cpp
@@ -1773,7 +1773,19 @@ USkinnedAsset* UNiagaraDataInterfaceSkeletalMesh::GetSkinnedMesh(FNiagaraSystemI
 	else if (!SystemInstance || !SystemInstance->GetWorld()->IsGameWorld())
 	{
 		// NOTE: We don't fall back on the preview mesh if we have a valid skeletal mesh component referenced
-		Mesh = PreviewMesh.LoadSynchronous();
+		Mesh = PreviewMesh.Get();
+
+		if (!Mesh)
+		{
+			Mesh = PreviewMesh.LoadSynchronous();
+
+			// if the soft object needed to be loaded we need to make sure that all the necessary streaming
+			// meta data has been initialized
+			if (Mesh)
+			{
+				Mesh->GetResourceForRendering();
+			}
+		}
 	}
 #endif
 
