- Include resolved data interfaces parameter bindings in parameter read count

#Synced-CL 50801738

diff --git a/Engine/Plugins/FX/Niagara/Source/NiagaraEditor/Private/ViewModels/NiagaraParameterPanelViewModel.cpp b/Engine/Plugins/FX/Niagara/Source/NiagaraEditor/Private/ViewModels/NiagaraParameterPanelViewModel.cpp
--- a/Engine/Plugins/FX/Niagara/Source/NiagaraEditor/Private/ViewModels/NiagaraParameterPanelViewModel.cpp
+++ b/Engine/Plugins/FX/Niagara/Source/NiagaraEditor/Private/ViewModels/NiagaraParameterPanelViewModel.cpp
@@ -6,6 +6,7 @@
 #include "NiagaraActions.h"
 #include "NiagaraClipboard.h"
 #include "NiagaraConstants.h"
+#include "NiagaraDataInterfaceUtilities.h"
 #include "NiagaraEditorData.h"
 #include "NiagaraEditorModule.h"
 #include "NiagaraEditorUtilities.h"
@@ -65,8 +66,6 @@ TArray<FNiagaraParameterPanelCategory> FNiagaraScriptToolkitParameterPanelViewMo
 TArray<FNiagaraParameterPanelCategory> FNiagaraScriptToolkitParameterPanelViewModel::DefaultAdvancedCategories;
 TArray<FNiagaraParameterPanelCategory> FNiagaraParameterDefinitionsToolkitParameterPanelViewModel::DefaultCategories;
 
-
-
 ///////////////////////////////////////////////////////////////////////////////
 /// System Toolkit Parameter Panel Utilities								///
 ///////////////////////////////////////////////////////////////////////////////
@@ -2246,61 +2245,133 @@ TArray<FNiagaraParameterPanelItem> FNiagaraSystemToolkitParameterPanelViewModel:
 		// Add active renderers usage variables to the counts
 		for (TSharedRef<FNiagaraEmitterHandleViewModel> EmitterVMS : SystemViewModel->GetEmitterHandleViewModels())
 		{
-			if (EmitterVMS.Get().IsValid() && EmitterVMS.Get().GetIsEnabled() && EmitterVMS.Get().GetEmitterHandle())
+			if (!EmitterVMS.Get().IsValid() || !EmitterVMS.Get().GetIsEnabled() || !EmitterVMS.Get().GetEmitterHandle())
+			{
+				continue;
+			}
+
+			FNiagaraEmitterHandle* EmitterHandle = EmitterVMS.Get().GetEmitterHandle();
+			UObject* SourceObject = EmitterHandle->GetEmitterBase();
+			if (const FVersionedNiagaraEmitterData* EmitterData = EmitterVMS.Get().GetEmitterHandle()->GetEmitterData())
 			{
-				FNiagaraEmitterHandle* EmitterHandle = EmitterVMS.Get().GetEmitterHandle();
-				UObject* SourceObject = EmitterHandle->GetEmitterBase();
-				if (const FVersionedNiagaraEmitterData* EmitterData = EmitterVMS.Get().GetEmitterHandle()->GetEmitterData())
+				if (UNiagaraScriptSource* ScriptSource = Cast<UNiagaraScriptSource>(EmitterData->GraphSource))
 				{
-					SourceObject = Cast<UNiagaraScriptSource>(EmitterData->GraphSource)->NodeGraph;
+					SourceObject = ScriptSource->NodeGraph;
 				}
+			}
 
-				FNiagaraAliasContext ResolveAliasesContext(FNiagaraAliasContext::ERapidIterationParameterMode::EmitterOrParticleScript);
-				ResolveAliasesContext.ChangeEmitterNameToEmitter(EmitterHandle->GetUniqueInstanceName());
+			FNiagaraAliasContext ResolveAliasesContext(FNiagaraAliasContext::ERapidIterationParameterMode::EmitterOrParticleScript);
+			ResolveAliasesContext.ChangeEmitterNameToEmitter(EmitterHandle->GetUniqueInstanceName());
 
-				EmitterHandle->ForEachEnabledRendererWithIndex(
-					[&](const UNiagaraRendererProperties* RenderProperties, int32 RendererIndex)
+			EmitterHandle->ForEachEnabledRendererWithIndex(
+				[&](const UNiagaraRendererProperties* RenderProperties, int32 RendererIndex)
+				{
+					TArray<FNiagaraVariable> BoundAttributes = RenderProperties->GetBoundAttributes();
+
+					// Find attributes first
+					for (const FNiagaraVariableAttributeBinding* AttributeBinding : RenderProperties->GetAttributeBindings())
 					{
-						TArray<FNiagaraVariable> BoundAttributes = RenderProperties->GetBoundAttributes();
+						const FNiagaraVariable BoundAttribute = RenderProperties->GetBoundAttribute(AttributeBinding);
+						if (BoundAttribute.IsValid() == false)
+						{
+							continue;
+						}
+						BoundAttributes.RemoveSwap(BoundAttribute, EAllowShrinking::No);
 
-						// Find attributes first
-						for (const FNiagaraVariableAttributeBinding* AttributeBinding : RenderProperties->GetAttributeBindings())
+						const FNiagaraVariable ResolvedAttribute = FNiagaraUtilities::ResolveAliases(BoundAttribute, ResolveAliasesContext);
+						if (FNiagaraParameterPanelItem* ItemPtr = VisitedParameterToItemMap.Find(ResolvedAttribute))
 						{
-							const FNiagaraVariable BoundAttribute = RenderProperties->GetBoundAttribute(AttributeBinding);
-							if (BoundAttribute.IsValid() == false)
-							{
-								continue;
-							}
-							BoundAttributes.RemoveSwap(BoundAttribute, EAllowShrinking::No);
+							const FName SourceModuleName(RenderProperties->GetWidgetDisplayName().ToString() + " (" + AttributeBinding->GetDataSetBindableVariable().GetName().ToString() + " Binding)");
+							ItemPtr->AddToReadCount(FNiagaraParameterReferencePath(SourceObject, SourceModuleName));
+						}
+					}
 
-							const FNiagaraVariable ResolvedAttribute = FNiagaraUtilities::ResolveAliases(BoundAttribute, ResolveAliasesContext);
-							if (FNiagaraParameterPanelItem* ItemPtr = VisitedParameterToItemMap.Find(ResolvedAttribute))
-							{
-								const FName SourceModuleName(RenderProperties->GetWidgetDisplayName().ToString() + " (" + AttributeBinding->GetDataSetBindableVariable().GetName().ToString() + " Binding)");
-								ItemPtr->AddToReadCount(FNiagaraParameterReferencePath(SourceObject, SourceModuleName));
-							}
+					// Add remaining bindings
+					for (const FNiagaraVariable& BoundAttribute : BoundAttributes)
+					{
+						if (BoundAttribute.IsValid() == false)
+						{
+							continue;
+						}
+
+						const FNiagaraVariable ResolvedAttribute = FNiagaraUtilities::ResolveAliases(BoundAttribute, ResolveAliasesContext);
+						if (FNiagaraParameterPanelItem* ItemPtr = VisitedParameterToItemMap.Find(ResolvedAttribute))
+						{
+							const FName SourceModuleName(*RenderProperties->GetWidgetDisplayName().ToString());
+							ItemPtr->AddToReadCount(FNiagaraParameterReferencePath(SourceObject, SourceModuleName));
+						}
+					}
+				}
+			);
+		}
+
+		// Find all data interfaces bindings
+		{
+			TSet<UNiagaraDataInterface*> DiscoveredDIs;
+			const auto FindDataInterfaceBindings =
+				[&](const UObject* Context, UNiagaraDataInterface* DataInterface, const FNiagaraVariableBase& Variable)
+				{
+					if (!DataInterface || DiscoveredDIs.Contains(DataInterface))
+					{
+						return true;
+					}
+					DiscoveredDIs.Add(DataInterface);
+
+					for (TFieldIterator<FStructProperty> PropertyIt(DataInterface->GetClass()); PropertyIt; ++PropertyIt)
+					{
+						FStructProperty* StructProperty = *PropertyIt;
+						if (!StructProperty || !StructProperty->Struct)
+						{
+							continue;
 						}
 
-						// Add remaining bindings
-						for (const FNiagaraVariable& BoundAttribute : BoundAttributes)
+						if (StructProperty->Struct->IsChildOf(FNiagaraUserParameterBinding::StaticStruct()))
 						{
-							if (BoundAttribute.IsValid() == false)
+							FNiagaraUserParameterBinding* ParameterBinding = StructProperty->ContainerPtrToValuePtr<FNiagaraUserParameterBinding>(DataInterface);
+							if (ParameterBinding->Parameter.IsValid())
 							{
-								continue;
+								if (FNiagaraParameterPanelItem* ItemPtr = VisitedParameterToItemMap.Find(ParameterBinding->Parameter))
+								{
+									ItemPtr->AddToReadCount(FNiagaraParameterReferencePath(Context, Variable.GetName()));
+								}
 							}
-
-							const FNiagaraVariable ResolvedAttribute = FNiagaraUtilities::ResolveAliases(BoundAttribute, ResolveAliasesContext);
-							if (FNiagaraParameterPanelItem* ItemPtr = VisitedParameterToItemMap.Find(ResolvedAttribute))
+						}
+						else if (StructProperty->Struct->IsChildOf(FNiagaraParameterBinding::StaticStruct()))
+						{
+							FNiagaraParameterBinding* ParameterBinding = StructProperty->ContainerPtrToValuePtr<FNiagaraParameterBinding>(DataInterface);
+							if (ParameterBinding->ResolvedParameter.IsValid())
 							{
-								const FName SourceModuleName(*RenderProperties->GetWidgetDisplayName().ToString());
-								ItemPtr->AddToReadCount(FNiagaraParameterReferencePath(SourceObject, SourceModuleName));
+								if (FNiagaraParameterPanelItem* ItemPtr = VisitedParameterToItemMap.Find(ParameterBinding->ResolvedParameter))
+								{
+									ItemPtr->AddToReadCount(FNiagaraParameterReferencePath(Context, Variable.GetName()));
+								}
 							}
 						}
 					}
-				);
+					return true;
+				};
+
+			{
+				const FNiagaraParameterStore& UserParameterStore = SystemViewModel->GetSystem().GetExposedParameters();
+				for (const FNiagaraVariableWithOffset& UserVariable : UserParameterStore.ReadParameterVariables())
+				{
+					if (UserVariable.GetType().IsDataInterface() == false)
+					{
+						continue;
+					}
+					FindDataInterfaceBindings(&SystemViewModel->GetSystem(), UserParameterStore.GetDataInterface(UserVariable.Offset), UserVariable);
+				}
 			}
-		}
 
+			FNiagaraDataInterfaceUtilities::ForEachDataInterface(
+				&SystemViewModel->GetSystem(),
+				[&](const FNiagaraDataInterfaceUtilities::FDataInterfaceUsageContext& UsageContext)
+				{
+					FindDataInterfaceBindings(UsageContext.OwnerObject, UsageContext.DataInterface, UsageContext.Variable);
+					return true;
+				}
+			);
+		}
 	}
 	else if (bForceScript)
 	{
