- Add method to convert transform in a 3x3 matrix

#rnx

#Synced-CL 50731656

diff --git a/Engine/Plugins/FX/Niagara/Shaders/Private/NiagaraTransformUtils.ush b/Engine/Plugins/FX/Niagara/Shaders/Private/NiagaraTransformUtils.ush
--- a/Engine/Plugins/FX/Niagara/Shaders/Private/NiagaraTransformUtils.ush
+++ b/Engine/Plugins/FX/Niagara/Shaders/Private/NiagaraTransformUtils.ush
@@ -50,3 +50,9 @@ float4x4 TransformToMatrix44(FNiagaraTransform T)
 		float4(T.Translation, 1.0f)
 	);
 }
+
+float3x3 TransformToMatrix33(FNiagaraTransform T)
+{
+	const float3x3 Rotation = NiagaraQuatTo3x3(T.Rotation);
+	return float3x3(Rotation[0] * T.Scale.x,  Rotation[1] * T.Scale.y, Rotation[2] * T.Scale.z);
+}
