- Make debug shaders optional references to handle cases where they are cooked out but the code can get in here

#Synced-CL 50731656

diff --git a/Engine/Plugins/FX/Niagara/Source/NiagaraShader/Private/NiagaraDebugShaders.cpp b/Engine/Plugins/FX/Niagara/Source/NiagaraShader/Private/NiagaraDebugShaders.cpp
--- a/Engine/Plugins/FX/Niagara/Source/NiagaraShader/Private/NiagaraDebugShaders.cpp
+++ b/Engine/Plugins/FX/Niagara/Source/NiagaraShader/Private/NiagaraDebugShaders.cpp
@@ -216,8 +216,12 @@ void NiagaraDebugShaders::DrawDebugLines(
 {
 	const FGlobalShaderMap* ShaderMap = GetGlobalShaderMap(View.GetFeatureLevel());
 
-	TShaderMapRef<FNiagaraDebugDrawLineVS> VertexShader(ShaderMap);
-	TShaderMapRef<FNiagaraDebugDrawLinePS> PixelShader(ShaderMap);
+	TOptionalShaderMapRef<FNiagaraDebugDrawLineVS> VertexShader(ShaderMap);
+	TOptionalShaderMapRef<FNiagaraDebugDrawLinePS> PixelShader(ShaderMap);
+	if (VertexShader.IsNull() || PixelShader.IsNull())
+	{
+		return;
+	}
 
 	const FIntRect ViewRect = UE::FXRenderingUtils::GetRawViewRectUnsafe(View);
 
@@ -263,8 +267,12 @@ void NiagaraDebugShaders::DrawDebugLines(FRDGBuilder& GraphBuilder, const FScene
 {
 	const FGlobalShaderMap* ShaderMap = GetGlobalShaderMap(View.GetFeatureLevel());
 
-	TShaderMapRef<FNiagaraDebugDrawLineVS> VertexShader(ShaderMap);
-	TShaderMapRef<FNiagaraDebugDrawLinePS> PixelShader(ShaderMap);
+	TOptionalShaderMapRef<FNiagaraDebugDrawLineVS> VertexShader(ShaderMap);
+	TOptionalShaderMapRef<FNiagaraDebugDrawLinePS> PixelShader(ShaderMap);
+	if (VertexShader.IsNull() || PixelShader.IsNull())
+	{
+		return;
+	}
 
 	const FIntRect ViewRect = UE::FXRenderingUtils::GetRawViewRectUnsafe(View);
 
@@ -399,7 +407,11 @@ void NiagaraDebugShaders::VisualizeTexture(
 
 	auto ShaderMap = GetGlobalShaderMap(GMaxRHIFeatureLevel);
 	TShaderMapRef<FScreenVS> VertexShader(ShaderMap);
-	TShaderMapRef<FNiagaraVisualizeTexturePS> PixelShader(ShaderMap, PermutationVector);
+	TOptionalShaderMapRef<FNiagaraVisualizeTexturePS> PixelShader(ShaderMap, PermutationVector);
+	if (PixelShader.IsNull())
+	{
+		return;
+	}
 
 	FIntPoint DisplaySize(TextureSize.X, TextureSize.Y);
 	if (DisplayHeight > 0)
