- Fix for ribbon export relying on the particle color node being present to export color

#Synced-CL 50731656

diff --git a/Engine/Plugins/FX/Niagara/Shaders/Private/NiagaraRibbonVertexFactory.ush b/Engine/Plugins/FX/Niagara/Shaders/Private/NiagaraRibbonVertexFactory.ush
--- a/Engine/Plugins/FX/Niagara/Shaders/Private/NiagaraRibbonVertexFactory.ush
+++ b/Engine/Plugins/FX/Niagara/Shaders/Private/NiagaraRibbonVertexFactory.ush
@@ -1150,10 +1150,7 @@ FLWCVector3 NiagaraVertexFactoryExport_GetPosition(FVertexFactoryInput Input, FV
 
 float4 NiagaraVertexFactoryExport_GetColor(FVertexFactoryInput Input, FVertexFactoryIntermediates Intermediates)
 {
-#if NEEDS_PARTICLE_COLOR
 	return Intermediates.Color;
-#endif
-	return float4(1,1,1,1);
 }
 
 float3x3 NiagaraVertexFactoryExport_GetTangentBasis(FVertexFactoryInput Input, FVertexFactoryIntermediates Intermediates)
