- Add runtime type utilities
- Allows plugins to define types that can be used with the property interface data interface

#Synced-CL 50868321

diff --git a/Engine/Plugins/FX/Niagara/Source/Niagara/Private/DataInterface/NiagaraDataInterfacePropertyInterface.cpp b/Engine/Plugins/FX/Niagara/Source/Niagara/Private/DataInterface/NiagaraDataInterfacePropertyInterface.cpp
--- a/Engine/Plugins/FX/Niagara/Source/Niagara/Private/DataInterface/NiagaraDataInterfacePropertyInterface.cpp
+++ b/Engine/Plugins/FX/Niagara/Source/Niagara/Private/DataInterface/NiagaraDataInterfacePropertyInterface.cpp
@@ -7,6 +7,7 @@
 #include "NiagaraShaderParametersBuilder.h"
 #include "NiagaraParameterStore.h"
 #include "NiagaraRenderThreadDeletor.h"
+#include "NiagaraRuntimeTypeUtilities.h"
 #include "NiagaraSceneComponentUtils.h"
 #include "NiagaraSystemInstance.h"
 
@@ -135,22 +136,19 @@ namespace NDIPropertyInterfaceLocal
 		TConstArrayView<uint8> GetCpuPropertyInterfaceData() const { return CpuPropertyInterfaceData; }
 
 	private:
-		template<typename TTypeFrom, typename TTypeTo>
-		void AddProperty(UStruct* Struct, FProperty* InterfaceProperty, FProperty* ClassProperty, const uint8* SourceObject, const FNiagaraTypeDefinition& TypeDef)
+		void AddProperty(UStruct* Struct, FProperty* InterfaceProperty, FProperty* ClassProperty, const uint8* SourceObject, const FNiagaraRuntimeTypeUtilities::FTypeConverter* TypeConverter)
 		{
+			check(TypeConverter && TypeConverter->ToNiagara);
+
 			const FString UnmangledPropertyName = Verse::Names::Private::UnmangleCasedName(InterfaceProperty->GetFName());
 			const FName PropertyName(*UnmangledPropertyName);
-			VariableToOffsets.Emplace(FNiagaraVariableBase(TypeDef, PropertyName), CpuPropertyInterfaceData.Num());
+			VariableToOffsets.Emplace(FNiagaraVariableBase(TypeConverter->TypeDef, PropertyName), CpuPropertyInterfaceData.Num());
 
 			const uint32 SrcOffset = ClassProperty->GetOffset_ForInternal();
 			const uint32 DstOffset = CpuPropertyInterfaceData.Num();
-			CpuPropertyInterfaceData.AddUninitialized(sizeof(TTypeTo));
-
-			TTypeFrom FromValue;
-			FMemory::Memcpy(&FromValue, SourceObject + SrcOffset, sizeof(TTypeFrom));
-
-			const TTypeTo ToValue = TTypeTo(FromValue);
-			FMemory::Memcpy(CpuPropertyInterfaceData.GetData() + DstOffset, &ToValue, sizeof(TTypeTo));
+			
+			CpuPropertyInterfaceData.AddUninitialized(TypeConverter->TypeDef.GetSize());
+			TypeConverter->ToNiagara(SourceObject + SrcOffset, CpuPropertyInterfaceData.GetData() + DstOffset);
 		}
 
 		void BuildProperties(UStruct* ClassStruct, UStruct* InterfaceStruct, const uint8* SourceObject)
@@ -160,6 +158,7 @@ namespace NDIPropertyInterfaceLocal
 				FNameBuilder PropertyNameBuilder;
 				InterfaceProperty->GetFName().ToString(PropertyNameBuilder);
 				PropertyNameBuilder.Append(TEXT("_Internal"));
+
 				FName PropertyName(PropertyNameBuilder.ToString());
 
 				// Interface proprties do not have the correct offset to the data so we need to find the property in the base class
@@ -170,41 +169,9 @@ namespace NDIPropertyInterfaceLocal
 					continue;
 				}
 
-				if (Property->IsA<const FIntProperty>())
-				{
-					AddProperty<int32, int32>(InterfaceStruct, InterfaceProperty, Property, SourceObject, FNiagaraTypeDefinition::GetIntDef());
-				}
-				else if (Property->IsA<const FFloatProperty>())
-				{
-					AddProperty<float, float>(InterfaceStruct, InterfaceProperty, Property, SourceObject, FNiagaraTypeDefinition::GetFloatDef());
-				}
-				else if (Property->IsA<const FDoubleProperty>())
-				{
-					AddProperty<double, float>(InterfaceStruct, InterfaceProperty, Property, SourceObject, FNiagaraTypeDefinition::GetFloatDef());
-				}
-				else if ( const FStructProperty* StructProperty = CastField<const FStructProperty>(Property))
+				if (const FNiagaraRuntimeTypeUtilities::FTypeConverter* TypeConverter = FNiagaraRuntimeTypeUtilities::FindTypeConverterForProperty(Property))
 				{
-					if (StructProperty->Struct == TBaseStructure<FVector2D>::Get())
-					{
-						AddProperty<FVector2D, FVector2f>(InterfaceStruct, InterfaceProperty, Property, SourceObject, FNiagaraTypeDefinition::GetVec2Def());
-					}
-					else if (StructProperty->Struct == TBaseStructure<FVector>::Get())
-					{
-						AddProperty<FVector, FVector3f>(InterfaceStruct, InterfaceProperty, Property, SourceObject, FNiagaraTypeDefinition::GetVec3Def());
-					}
-					else if (StructProperty->Struct == TBaseStructure<FVector4>::Get())
-					{
-						AddProperty<FVector4, FVector4f>(InterfaceStruct, InterfaceProperty, Property, SourceObject, FNiagaraTypeDefinition::GetVec4Def());
-					}
-					else if (StructProperty->Struct == TBaseStructure<FQuat>::Get())
-					{
-						AddProperty<FQuat, FQuat4f>(InterfaceStruct, InterfaceProperty, Property, SourceObject, FNiagaraTypeDefinition::GetQuatDef());
-					}
-					else if (StructProperty->Struct == TBaseStructure<FLinearColor>::Get())
-					{
-						AddProperty<FLinearColor, FLinearColor>(InterfaceStruct, InterfaceProperty, Property, SourceObject, FNiagaraTypeDefinition::GetColorDef());
-					}
-					// Note does not recurse the properties we only allow for top level
+					AddProperty(InterfaceStruct, InterfaceProperty, Property, SourceObject, TypeConverter);
 				}
 			}
 		}
diff --git a/Engine/Plugins/FX/Niagara/Source/Niagara/Private/NiagaraModule.cpp b/Engine/Plugins/FX/Niagara/Source/Niagara/Private/NiagaraModule.cpp
--- a/Engine/Plugins/FX/Niagara/Source/Niagara/Private/NiagaraModule.cpp
+++ b/Engine/Plugins/FX/Niagara/Source/Niagara/Private/NiagaraModule.cpp
@@ -12,6 +12,7 @@
 #include "NiagaraSettings.h"
 #include "NiagaraDataChannelManager.h"
 #include "NiagaraDataChannelDebugUtilities.h"
+#include "NiagaraRuntimeTypeUtilities.h"
 #include "NiagaraSimCache.h"
 #include "NiagaraSimulationTaskPriority.h"
 #include "NiagaraUseOpenVDB.h"
@@ -323,6 +324,8 @@ void INiagaraModule::StartupModule()
 	FNiagaraTypeDefinition::Init();
 	FNiagaraTypeRegistry::ProcessRegistryQueue();
 
+	FNiagaraRuntimeTypeUtilities::Initialize();
+
 	if (IsRunningCookCommandlet())
 	{
 		//Never allow NDCs when cooking.
diff --git a/Engine/Plugins/FX/Niagara/Source/Niagara/Private/NiagaraRuntimeTypeUtilities.cpp b/Engine/Plugins/FX/Niagara/Source/Niagara/Private/NiagaraRuntimeTypeUtilities.cpp
new file mode 100644
--- /dev/null
+++ b/Engine/Plugins/FX/Niagara/Source/Niagara/Private/NiagaraRuntimeTypeUtilities.cpp
@@ -0,0 +1,74 @@
+// Copyright Epic Games, Inc. All Rights Reserved.
+
+#include "NiagaraRuntimeTypeUtilities.h"
+#include "UObject/ObjectKey.h"
+
+namespace FNiagaraRuntimeTypeUtilities
+{
+	static FTypeConverter					IntConverter;
+	static FTypeConverter					FloatConverter;
+	static FTypeConverter					DoubleConverter;
+	static TMap<FObjectKey, FTypeConverter> StructTypeConverters;
+
+	template<typename TTypeFrom, typename TTypeTo>
+	void CopyValue(const void* FromPointer, void* ToPointer)
+	{
+		TTypeFrom FromValue;
+		FMemory::Memcpy(&FromValue, FromPointer, sizeof(TTypeFrom));
+
+		const TTypeTo ToValue = TTypeTo(FromValue);
+		FMemory::Memcpy(ToPointer, &ToValue, sizeof(TTypeTo));
+	}
+
+	template<typename TType, typename TNiagaraType>
+	void RegisterStructTypeConverterBase(FNiagaraTypeDefinition TypeDef)
+	{
+		RegisterStructTypeConverter(TBaseStructure<TType>::Get(), FTypeConverter(TypeDef, CopyValue<TType, TNiagaraType>, CopyValue<TNiagaraType, TType>));
+	}
+
+	void Initialize()
+	{
+		IntConverter = FTypeConverter(FNiagaraTypeDefinition::GetIntDef(), CopyValue<int, int>, CopyValue<int, int>);
+		FloatConverter = FTypeConverter(FNiagaraTypeDefinition::GetFloatDef(), CopyValue<float, float>, CopyValue<float, float>);
+		DoubleConverter = FTypeConverter(FNiagaraTypeDefinition::GetFloatDef(), CopyValue<double, float>, CopyValue<float, double>);
+
+		RegisterStructTypeConverterBase<FVector2D,		FVector2f>(FNiagaraTypeDefinition::GetVec2Def());
+		RegisterStructTypeConverterBase<FVector,		FVector3f>(FNiagaraTypeDefinition::GetVec3Def());
+		RegisterStructTypeConverterBase<FVector4,		FVector4f>(FNiagaraTypeDefinition::GetVec4Def());
+		RegisterStructTypeConverterBase<FQuat,			FQuat4f>(FNiagaraTypeDefinition::GetQuatDef());
+		RegisterStructTypeConverterBase<FLinearColor,	FLinearColor>(FNiagaraTypeDefinition::GetColorDef());
+	}
+
+	void RegisterStructTypeConverter(const UScriptStruct* Struct, FTypeConverter Converter)
+	{
+		check(StructTypeConverters.Contains(Struct) == false);
+		StructTypeConverters.Emplace(Struct, MoveTemp(Converter));
+	}
+
+	void UnregisterStructTypeConverter(const UScriptStruct* Struct)
+	{
+		check(StructTypeConverters.Contains(Struct) == true);
+		StructTypeConverters.Remove(Struct);
+	}
+
+	const FTypeConverter* FindTypeConverterForProperty(const FProperty* Property)
+	{
+		if (Property->IsA<const FIntProperty>())
+		{
+			return &IntConverter;
+		}
+		else if (Property->IsA<const FFloatProperty>())
+		{
+			return &FloatConverter;
+		}
+		else if (Property->IsA<const FDoubleProperty>())
+		{
+			return &DoubleConverter;
+		}
+		else if (const FStructProperty* StructProperty = CastField<const FStructProperty>(Property))
+		{
+			return StructTypeConverters.Find(StructProperty->Struct);
+		}
+		return nullptr;
+	}
+}
diff --git a/Engine/Plugins/FX/Niagara/Source/Niagara/Public/NiagaraRuntimeTypeUtilities.h b/Engine/Plugins/FX/Niagara/Source/Niagara/Public/NiagaraRuntimeTypeUtilities.h
new file mode 100644
--- /dev/null
+++ b/Engine/Plugins/FX/Niagara/Source/Niagara/Public/NiagaraRuntimeTypeUtilities.h
@@ -0,0 +1,45 @@
+// Copyright Epic Games, Inc. All Rights Reserved.
+
+#pragma once
+
+#include "CoreMinimal.h"
+#include "NiagaraTypes.h"
+
+namespace FNiagaraRuntimeTypeUtilities
+{
+	struct FTypeConverter
+	{
+		using FCopyFunc = void(*)(const void*, void*);
+
+		FTypeConverter() = default;
+		FTypeConverter(FNiagaraTypeDefinition InTypeDef, FCopyFunc InToNiagara, FCopyFunc InFromNiagara)
+			: TypeDef(InTypeDef)
+			, ToNiagara(InToNiagara)
+			, FromNiagara(InFromNiagara)
+		{
+		}
+
+		FNiagaraTypeDefinition	TypeDef;
+		FCopyFunc				ToNiagara = nullptr;
+		FCopyFunc				FromNiagara = nullptr;
+	};
+
+	void Initialize();
+
+	NIAGARA_API void RegisterStructTypeConverter(const UScriptStruct* Struct, FTypeConverter Converter);
+	NIAGARA_API void UnregisterStructTypeConverter(const UScriptStruct* Struct);
+
+	template<typename TType>
+	void RegisterStructTypeConverter(FTypeConverter Converter)
+	{
+		RegisterStructTypeConverter(StaticStruct<TType>(), Converter);
+	}
+
+	template<typename TType>
+	void UnregisterStructTypeConverter()
+	{
+		UnregisterStructTypeConverter(StaticStruct<TType>());
+	}
+
+	NIAGARA_API const FTypeConverter* FindTypeConverterForProperty(const FProperty* Property);
+}
