Adding logging when encountering Niagara shader compile errors.

#Synced-CL 50913622

diff --git a/Engine/Plugins/FX/Niagara/Source/NiagaraEditor/Private/NiagaraCompiler.cpp b/Engine/Plugins/FX/Niagara/Source/NiagaraEditor/Private/NiagaraCompiler.cpp
--- a/Engine/Plugins/FX/Niagara/Source/NiagaraEditor/Private/NiagaraCompiler.cpp
+++ b/Engine/Plugins/FX/Niagara/Source/NiagaraEditor/Private/NiagaraCompiler.cpp
@@ -7,6 +7,7 @@
 #include "EdGraphUtilities.h"
 #include "Interfaces/IShaderFormat.h"
 #include "INiagaraEditorTypeUtilities.h"
+#include "Logging/StructuredLog.h"
 #include "Misc/FileHelper.h"
 #include "Misc/PathViews.h"
 #include "Modules/ModuleManager.h"
@@ -33,6 +34,7 @@
 #include "Serialization/MemoryReader.h"
 #include "ShaderCompiler.h"
 #include "ShaderCore.h"
+#include "ShaderDiagnostics.h"
 #include "VectorVMTestCompile.h"
 
 #define LOCTEXT_NAMESPACE "NiagaraCompiler"
@@ -2240,13 +2242,26 @@ TOptional<FNiagaraCompileResults> FHlslNiagaraCompiler::GetCompileResult(int32 J
 	}
 	else if (CompilationJob->ShaderCompileJob->Output.Errors.Num() > 0)
 	{
-		FString Errors;
-		for (FShaderCompilerError ShaderError : CompilationJob->ShaderCompileJob->Output.Errors)
+		TArray<FShaderCommonCompileJobPtr> FailedJobs;
+		FailedJobs.Emplace(CompilationJob->ShaderCompileJob);
+
+		const FShaderDiagnosticInfo ShaderDiagInfo(FailedJobs);
+		if (ShaderDiagInfo.UniqueErrors.Num())
 		{
-			Errors += ShaderError.StrippedErrorMessage + "\n";
+			const FString AllUniqueErrors = FString::Join(ShaderDiagInfo.UniqueErrors, TEXT("\n"));
+
+			UE_LOGFMT_LOC(LogNiagaraCompiler, Error, "VectorVMShaderCompilationFailed", "{NumErrors} errors encountered compiling Vector VM shaders:\n{Errors}",
+				("NumErrors", ShaderDiagInfo.UniqueErrors.Num()),
+				("Errors", AllUniqueErrors)
+			);
+
+			Error(FText::Format(LOCTEXT("VectorVMCompileErrorMessageFormat", "The Vector VM compile failed.  Errors:\n{0}"), FText::FromString(AllUniqueErrors)));
+
+			if (UE_LOG_ACTIVE(LogNiagaraCompiler, Verbose))
+			{
+				DumpHLSLText(Results.Data->LastHlslTranslation, CompilationJob->CompileResults.DumpDebugInfoPath);
+			}
 		}
-		Error(FText::Format(LOCTEXT("VectorVMCompileErrorMessageFormat", "The Vector VM compile failed.  Errors:\n{0}"), FText::FromString(Errors)));
-		DumpHLSLText(Results.Data->LastHlslTranslation, CompilationJob->CompileResults.DumpDebugInfoPath);
 	}
 
 	if (Results.bVMSucceeded)
