Fixed a crash where pso tasks could be started on the render thread instead of game thread by CreateSceneProxy()

#rnx
#Synced-CL 50596076

diff --git a/Engine/Plugins/FX/Niagara/Source/Niagara/Private/NiagaraComponent.cpp b/Engine/Plugins/FX/Niagara/Source/Niagara/Private/NiagaraComponent.cpp
--- a/Engine/Plugins/FX/Niagara/Source/Niagara/Private/NiagaraComponent.cpp
+++ b/Engine/Plugins/FX/Niagara/Source/Niagara/Private/NiagaraComponent.cpp
@@ -2464,9 +2464,17 @@ FPrimitiveSceneProxy* UNiagaraComponent::CreateSceneProxy()
 	{
 		// If we are relying on proxy creation PSOs only then call precache PSOs here because asset doesn't have collected PSO precache data
 		const FNiagaraReadinessSettings& ReadinessSettings = GetReadinessSettings();
-		if (ReadinessSettings.PSOPrecacheSettings.PrecachingTime == ENiagaraPSOPrecachingTime::OnRenderStateCreation)
+		if (ReadinessSettings.PSOPrecacheSettings.PrecachingTime == ENiagaraPSOPrecachingTime::OnRenderStateCreation && !IsPSOPrecaching())
 		{
-			PrecachePSOs();
+			// Make sure the pso is being precached on the GT as we're likely not creating the scene proxy here, so we can wait for next frame
+			TWeakObjectPtr<UNiagaraComponent> WeakThis(this);
+			AsyncTask(ENamedThreads::GameThread, [WeakThis]()
+			{
+				if (WeakThis.IsValid())
+				{
+					WeakThis->PrecachePSOs();
+				}
+			});
 		}
 
 		if (IsPSOPrecaching() && SystemInstanceController.IsValid())
