Adding functions for getting transform/invtransform

#Synced-CL 50944596

diff --git a/Engine/Plugins/FX/Niagara/Shaders/Private/NiagaraDataInterfaceRigidMeshCollisionQuery.ush b/Engine/Plugins/FX/Niagara/Shaders/Private/NiagaraDataInterfaceRigidMeshCollisionQuery.ush
--- a/Engine/Plugins/FX/Niagara/Shaders/Private/NiagaraDataInterfaceRigidMeshCollisionQuery.ush
+++ b/Engine/Plugins/FX/Niagara/Shaders/Private/NiagaraDataInterfaceRigidMeshCollisionQuery.ush
@@ -5,6 +5,8 @@ NiagaraDataInterfaceRigidMeshCollisionQuery.ush
 =============================================================================*/
 
 #include "/Plugin/FX/Niagara/Private/NiagaraDistanceFieldQueries.ush"
+#include "/Plugin/FX/Niagara/Private/NiagaraQuaternionUtils.ush"
+#include "/Plugin/FX/Niagara/Private/NiagaraTransformUtils.ush"
 
 /* -----------------------------------------------------------------
  * static mesh constants and context
@@ -681,6 +683,44 @@ void GetMaxEncodedDistanceMeshDistanceField_{ParameterName}(in int ElementIndex,
 	}
 }
 
+void GetElementTransform_{ParameterName}(in int ElementIndex, in float TimeFraction, out float3 OutTranslation, out float4 OutRotation, out float3 OutScale, out bool OutIsValid)
+{
+	OutIsValid = false;
+
+	int ElementCount = -1;
+	GetNumElements_{ParameterName}(ElementCount);
+	if (ElementIndex >= 0 && ElementIndex < ElementCount)
+	{
+		OutIsValid = true;
+
+		FNiagaraTransform PrevTransform = DecomposePackedMatrix(GetPreviousTransform_{ParameterName}(ElementIndex));
+		FNiagaraTransform CurrTransform = DecomposePackedMatrix(GetCurrentTransform_{ParameterName}(ElementIndex));
+
+		OutTranslation = lerp(PrevTransform.Translation, CurrTransform.Translation, TimeFraction);
+		OutRotation = SlerpQuat(PrevTransform.Rotation, CurrTransform.Rotation, TimeFraction);
+		OutScale = lerp(PrevTransform.Scale, CurrTransform.Scale, TimeFraction);
+	}
+}
+
+void GetElementInverseTransform_{ParameterName}(in int ElementIndex, in float TimeFraction, out float3 OutTranslation, out float4 OutRotation, out float3 OutScale, out bool OutIsValid)
+{
+	OutIsValid = false;
+
+	int ElementCount = -1;
+	GetNumElements_{ParameterName}(ElementCount);
+	if (ElementIndex >= 0 && ElementIndex < ElementCount)
+	{
+		OutIsValid = true;
+
+		FNiagaraTransform PrevTransform = DecomposePackedMatrix(GetPreviousInverse_{ParameterName}(ElementIndex));
+		FNiagaraTransform CurrTransform = DecomposePackedMatrix(GetInverseTransform_{ParameterName}(ElementIndex));
+
+		OutTranslation = lerp(PrevTransform.Translation, CurrTransform.Translation, TimeFraction);
+		OutRotation = SlerpQuat(PrevTransform.Rotation, CurrTransform.Rotation, TimeFraction);
+		OutScale = lerp(PrevTransform.Scale, CurrTransform.Scale, TimeFraction);
+	}
+}
+
 #undef BOX_INDEX
 #undef SPHERE_INDEX
 #undef CAPSULE_INDEX
\ No newline at end of file
diff --git a/Engine/Plugins/FX/Niagara/Shaders/Private/NiagaraTransformUtils.ush b/Engine/Plugins/FX/Niagara/Shaders/Private/NiagaraTransformUtils.ush
--- a/Engine/Plugins/FX/Niagara/Shaders/Private/NiagaraTransformUtils.ush
+++ b/Engine/Plugins/FX/Niagara/Shaders/Private/NiagaraTransformUtils.ush
@@ -56,3 +56,34 @@ float3x3 TransformToMatrix33(FNiagaraTransform T)
 	const float3x3 Rotation = NiagaraQuatTo3x3(T.Rotation);
 	return float3x3(Rotation[0] * T.Scale.x,  Rotation[1] * T.Scale.y, Rotation[2] * T.Scale.z);
 }
+
+FNiagaraTransform DecomposeMatrix(float4x4 InMatrix)
+{
+	FNiagaraTransform R;
+	R.Translation = InMatrix._m30_m31_m32;
+
+	float3 Col0 = InMatrix._m00_m10_m20;
+	float3 Col1 = InMatrix._m01_m11_m21;
+	float3 Col2 = InMatrix._m02_m12_m23;
+
+	R.Scale = float3(length(Col0), length(Col1), length(Col2));
+	float3x3 NormalizedRotationMatrix = transpose(float3x3(Col0 / R.Scale.x, Col1 / R.Scale.y, Col2 / R.Scale.z));
+	float3 RotationVectors[3] =
+	{
+		NormalizedRotationMatrix[0],
+		NormalizedRotationMatrix[1],
+		NormalizedRotationMatrix[2]
+	};
+
+	R.Rotation = QuatFromMatrix(RotationVectors);
+
+	return R;
+}
+
+// PackedMatrix here is a linear matrix packed for serialization.  Effectively a float4x3 transposed into vec4s
+FNiagaraTransform DecomposePackedMatrix(float3x4 InPackedMatrix)
+{
+	float4x4 Matrix = transpose(float4x4(InPackedMatrix, float4(0.0f, 0.0f, 0.0f, 1.0f)));
+
+	return DecomposeMatrix(Matrix);
+}
\ No newline at end of file
diff --git a/Engine/Plugins/FX/Niagara/Source/Niagara/Private/NiagaraDataInterfaceRigidMeshCollisionQuery.cpp b/Engine/Plugins/FX/Niagara/Source/Niagara/Private/NiagaraDataInterfaceRigidMeshCollisionQuery.cpp
--- a/Engine/Plugins/FX/Niagara/Source/Niagara/Private/NiagaraDataInterfaceRigidMeshCollisionQuery.cpp
+++ b/Engine/Plugins/FX/Niagara/Source/Niagara/Private/NiagaraDataInterfaceRigidMeshCollisionQuery.cpp
@@ -103,6 +103,8 @@ static const FName GetBoxSizeName(TEXT("GetBoxSize"));
 static const FName IsWorldPositionInsideCombinedBoundsName(TEXT("IsWorldPositionInsideCombinedBounds"));
 static const FName GetClosestElementName(TEXT("GetClosestElement"));
 static const FName GetElementPointName(TEXT("GetElementPoint"));
+static const FName GetElementTransformName(TEXT("GetElementTransform"));
+static const FName GetElementInverseTransformName(TEXT("GetElementInverseTransform"));
 static const FName GetElementPointMeshDistanceFieldNoNormalName(TEXT("GetElementPointMeshDistanceFieldNoNormal"));
 static const FName GetElementDistanceName(TEXT("GetElementDistance"));
 static const FName GetClosestPointName(TEXT("GetClosestPoint"));
@@ -119,7 +121,7 @@ static const FText OverlapOriginDescription = IF_WITH_EDITORONLY_DATA(
 );
 
 static const FText OverlapRotationDescription = IF_WITH_EDITORONLY_DATA(
-	LOCTEXT("RigidBodyOVerlapRotationDesciption", "The orientation of the box to be used for hte overlap trace."),
+	LOCTEXT("RigidBodyOVerlapRotationDesciption", "The orientation of the box to be used for the overlap trace."),
 	FText()
 );
 
@@ -1671,7 +1673,7 @@ void UNiagaraDataInterfaceRigidMeshCollisionQuery::GetFunctionsInternal(TArray<F
 	{
 		FNiagaraFunctionSignature Sig;
 		Sig.Name = GetElementPointName;
-		Sig.SetDescription(LOCTEXT("GetClosestElementPointDescription", "Given a world space position and an element index, computes the static mesh's closest point. Also returns normal and velocity for that point."));
+		Sig.SetDescription(LOCTEXT("GetElementPointDescription", "Given a world space position and an element index, computes the static mesh's closest point. Also returns normal and velocity for that point."));
 		Sig.SetFunctionVersion(FNiagaraRigidMeshCollisionDIFunctionVersion::LatestVersion);
 		Sig.bSupportsGPU = true;
 		Sig.bSupportsCPU = false;
@@ -1687,6 +1689,42 @@ void UNiagaraDataInterfaceRigidMeshCollisionQuery::GetFunctionsInternal(TArray<F
 
 		OutFunctions.Add(Sig);
 	}
+	{
+		FNiagaraFunctionSignature Sig;
+		Sig.Name = GetElementTransformName;
+		Sig.SetDescription(LOCTEXT("GetElementTransformDescription", "Gets the World Transform of the specified element."));
+		Sig.SetFunctionVersion(FNiagaraRigidMeshCollisionDIFunctionVersion::LatestVersion);
+		Sig.bSupportsGPU = true;
+		Sig.bSupportsCPU = false;
+		Sig.bMemberFunction = true;
+		Sig.Inputs.Add(FNiagaraVariable(FNiagaraTypeDefinition(GetClass()), TEXT("Collision DI")));
+		Sig.Inputs.Add(FNiagaraVariable(FNiagaraTypeDefinition::GetIntDef(), TEXT("Element Index")));
+		Sig.Inputs.Add(FNiagaraVariable(FNiagaraTypeDefinition::GetFloatDef(), TEXT("Time Fraction")));
+		Sig.Outputs.Add(FNiagaraVariable(FNiagaraTypeDefinition::GetVec3Def(), TEXT("Translation")));
+		Sig.Outputs.Add(FNiagaraVariable(FNiagaraTypeDefinition::GetQuatDef(), TEXT("Rotation")));
+		Sig.Outputs.Add(FNiagaraVariable(FNiagaraTypeDefinition::GetVec3Def(), TEXT("Scale")));
+		Sig.Outputs.Add(FNiagaraVariable(FNiagaraTypeDefinition::GetBoolDef(), TEXT("Is Valid")));
+
+		OutFunctions.Add(Sig);
+	}
+	{
+		FNiagaraFunctionSignature Sig;
+		Sig.Name = GetElementInverseTransformName;
+		Sig.SetDescription(LOCTEXT("GetElementInverseTransformDescription", "Gets the Inverse World Transform of the specified element."));
+		Sig.SetFunctionVersion(FNiagaraRigidMeshCollisionDIFunctionVersion::LatestVersion);
+		Sig.bSupportsGPU = true;
+		Sig.bSupportsCPU = false;
+		Sig.bMemberFunction = true;
+		Sig.Inputs.Add(FNiagaraVariable(FNiagaraTypeDefinition(GetClass()), TEXT("Collision DI")));
+		Sig.Inputs.Add(FNiagaraVariable(FNiagaraTypeDefinition::GetIntDef(), TEXT("Element Index")));
+		Sig.Inputs.Add(FNiagaraVariable(FNiagaraTypeDefinition::GetFloatDef(), TEXT("Time Fraction")));
+		Sig.Outputs.Add(FNiagaraVariable(FNiagaraTypeDefinition::GetVec3Def(), TEXT("Translation")));
+		Sig.Outputs.Add(FNiagaraVariable(FNiagaraTypeDefinition::GetQuatDef(), TEXT("Rotation")));
+		Sig.Outputs.Add(FNiagaraVariable(FNiagaraTypeDefinition::GetVec3Def(), TEXT("Scale")));
+		Sig.Outputs.Add(FNiagaraVariable(FNiagaraTypeDefinition::GetBoolDef(), TEXT("Is Valid")));
+
+		OutFunctions.Add(Sig);
+	}
 	{
 		FNiagaraFunctionSignature Sig;
 		Sig.Name = GetElementPointMeshDistanceFieldNoNormalName;
@@ -1866,6 +1904,8 @@ bool UNiagaraDataInterfaceRigidMeshCollisionQuery::GetFunctionHLSL(const FNiagar
 		(FunctionInfo.DefinitionName == GetClosestPointName) ||
 		(FunctionInfo.DefinitionName == GetClosestPointSimpleName) ||		
 		(FunctionInfo.DefinitionName == GetClosestElementName) ||
+		(FunctionInfo.DefinitionName == GetElementTransformName) ||
+		(FunctionInfo.DefinitionName == GetElementInverseTransformName) ||
 		(FunctionInfo.DefinitionName == GetElementPointName) ||
 		(FunctionInfo.DefinitionName == GetElementPointMeshDistanceFieldNoNormalName) ||
 		(FunctionInfo.DefinitionName == GetElementDistanceName) ||
