- Fix some analysis issues

#rnx

#Synced-CL 50885476

diff --git a/Engine/Plugins/FX/Niagara/Source/Niagara/Internal/DataInterface/NiagaraDataInterfaceArrayDistributionInt.h b/Engine/Plugins/FX/Niagara/Source/Niagara/Internal/DataInterface/NiagaraDataInterfaceArrayDistributionInt.h
--- a/Engine/Plugins/FX/Niagara/Source/Niagara/Internal/DataInterface/NiagaraDataInterfaceArrayDistributionInt.h
+++ b/Engine/Plugins/FX/Niagara/Source/Niagara/Internal/DataInterface/NiagaraDataInterfaceArrayDistributionInt.h
@@ -83,7 +83,7 @@ private:
 	int32 GetBuiltTableLength() const;
 
 	void VMLength(FVectorVMExternalFunctionContext& Context) const;
-	void VMIsLastIndex(FVectorVMExternalFunctionContext& Context) const;
+	void VMIsValidIndex(FVectorVMExternalFunctionContext& Context) const;
 	void VMGetLastIndex(FVectorVMExternalFunctionContext& Context) const;
 	void VMGet(FVectorVMExternalFunctionContext& Context) const;
 	void VMGetProbAlias(FVectorVMExternalFunctionContext& Context) const;
diff --git a/Engine/Plugins/FX/Niagara/Source/Niagara/Private/DataInterface/NDIRenderTargetSimCacheData.cpp b/Engine/Plugins/FX/Niagara/Source/Niagara/Private/DataInterface/NDIRenderTargetSimCacheData.cpp
--- a/Engine/Plugins/FX/Niagara/Source/Niagara/Private/DataInterface/NDIRenderTargetSimCacheData.cpp
+++ b/Engine/Plugins/FX/Niagara/Source/Niagara/Private/DataInterface/NDIRenderTargetSimCacheData.cpp
@@ -100,7 +100,7 @@ bool UNDIRenderTargetSimCacheData::CompareFrame(const UNDIRenderTargetSimCacheDa
 	}
 
 	const EPixelFormat LhsPixelFormat = GetTextureFormat(FrameIndex);
-	const EPixelFormat RhsPixelFormat = GetTextureFormat(FrameIndex);
+	const EPixelFormat RhsPixelFormat = Other->GetTextureFormat(FrameIndex);
 	if (LhsPixelFormat != RhsPixelFormat)
 	{
 		OutErrors = FString::Printf(
@@ -225,7 +225,8 @@ TArray<uint8> UNDIRenderTargetSimCacheData::GetPixelData(int32 FrameIndex) const
 	OutPixelData.AddUninitialized(FrameData.UncompressedSize);
 	if (FrameData.CompressedSize > 0)
 	{
-		FCompression::UncompressMemory(CompressionType, OutPixelData.GetData(), FrameData.UncompressedSize, PixelDataFrames[FrameIndex], FrameData.CompressedSize);
+		const bool bUncompressSucceeded = FCompression::UncompressMemory(CompressionType, OutPixelData.GetData(), FrameData.UncompressedSize, PixelDataFrames[FrameIndex], FrameData.CompressedSize);
+		check(bUncompressSucceeded);
 	}
 	else
 	{
@@ -259,7 +260,8 @@ void UNDIRenderTargetSimCacheData::GetPixelData(int32 FrameIndex, uint8* DestPix
 	{
 		if (FrameData.CompressedSize > 0)
 		{
-			FCompression::UncompressMemory(CompressionType, DestPixelData, FrameData.UncompressedSize, PixelDataFrames[FrameIndex], FrameData.CompressedSize);
+			const bool bUncompressSucceeded = FCompression::UncompressMemory(CompressionType, DestPixelData, FrameData.UncompressedSize, PixelDataFrames[FrameIndex], FrameData.CompressedSize);
+			check(bUncompressSucceeded);
 		}
 		else
 		{
@@ -273,7 +275,8 @@ void UNDIRenderTargetSimCacheData::GetPixelData(int32 FrameIndex, uint8* DestPix
 		if (FrameData.CompressedSize > 0)
 		{
 			Decompressed.AddUninitialized(BlockBytes * FrameData.Size.X * FrameData.Size.Y * FrameData.Size.Z);
-			FCompression::UncompressMemory(CompressionType, Decompressed.GetData(), Decompressed.Num(), PixelDataFrames[FrameIndex], FrameData.CompressedSize);
+			const bool bUncompressSucceeded = FCompression::UncompressMemory(CompressionType, Decompressed.GetData(), Decompressed.Num(), PixelDataFrames[FrameIndex], FrameData.CompressedSize);
+			check(bUncompressSucceeded);
 		}
 
 		const uint8* SrcData = Decompressed.Num() > 0 ? Decompressed.GetData() : PixelDataFrames[FrameIndex];
@@ -327,7 +330,7 @@ void UNDIRenderTargetSimCacheData::SetPixelData(int32 FrameIndex, FIntVector Siz
 		int CompressedSize = FCompression::CompressMemoryBound(CompressionType, PixelData.Num(), CompressionFlags);
 
 		TArray<uint8> CompressedPixelData;
-		CompressedPixelData.Reserve(CompressedSize);
+		CompressedPixelData.AddUninitialized(CompressedSize);
 		if (FCompression::CompressMemory(CompressionType, CompressedPixelData.GetData(), CompressedSize, PixelData.GetData(), PixelData.Num(), CompressionFlags))
 		{
 			FrameData.BulkData.Lock(LOCK_READ_WRITE);
diff --git a/Engine/Plugins/FX/Niagara/Source/Niagara/Private/DataInterface/NiagaraDataInterfaceArrayDistributionInt.cpp b/Engine/Plugins/FX/Niagara/Source/Niagara/Private/DataInterface/NiagaraDataInterfaceArrayDistributionInt.cpp
--- a/Engine/Plugins/FX/Niagara/Source/Niagara/Private/DataInterface/NiagaraDataInterfaceArrayDistributionInt.cpp
+++ b/Engine/Plugins/FX/Niagara/Source/Niagara/Private/DataInterface/NiagaraDataInterfaceArrayDistributionInt.cpp
@@ -220,7 +220,7 @@ void UNiagaraDataInterfaceArrayDistributionInt::GetVMExternalFunction(const FVME
 	}
 	else if (BindingInfo.Name == FNiagaraDataInterfaceArrayImplInternal::Function_IsValidIndexName)
 	{
-		OutFunc = FVMExternalFunction::CreateUObject(this, &UNiagaraDataInterfaceArrayDistributionInt::VMIsLastIndex);
+		OutFunc = FVMExternalFunction::CreateUObject(this, &UNiagaraDataInterfaceArrayDistributionInt::VMIsValidIndex);
 	}
 	else if (BindingInfo.Name == FNiagaraDataInterfaceArrayImplInternal::Function_LastIndexName)
 	{
@@ -366,7 +366,7 @@ void UNiagaraDataInterfaceArrayDistributionInt::VMLength(FVectorVMExternalFuncti
 	}
 }
 
-void UNiagaraDataInterfaceArrayDistributionInt::VMIsLastIndex(FVectorVMExternalFunctionContext& Context) const
+void UNiagaraDataInterfaceArrayDistributionInt::VMIsValidIndex(FVectorVMExternalFunctionContext& Context) const
 {
 	FNDIInputParam<int32> InIndex(Context);
 	FNDIOutputParam<bool> OutValid(Context);
diff --git a/Engine/Plugins/FX/Niagara/Source/Niagara/Private/DataInterface/NiagaraDataInterfaceArrayMesh.cpp b/Engine/Plugins/FX/Niagara/Source/Niagara/Private/DataInterface/NiagaraDataInterfaceArrayMesh.cpp
--- a/Engine/Plugins/FX/Niagara/Source/Niagara/Private/DataInterface/NiagaraDataInterfaceArrayMesh.cpp
+++ b/Engine/Plugins/FX/Niagara/Source/Niagara/Private/DataInterface/NiagaraDataInterfaceArrayMesh.cpp
@@ -163,18 +163,23 @@ void UNiagaraDataInterfaceArrayMesh::ForEachMesh(FNiagaraSystemInstance* SystemI
 	}
 
 	FProxyType::FReadArrayRef ArrayRef(this, PerInstanceData);
+	if (NumMeshesDelegate.IsSet())
+	{
+		NumMeshesDelegate(ArrayRef.GetArray().Num());
+	}
 
-	NumMeshesDelegate(ArrayRef.GetArray().Num());
-
-	FNiagaraMeshRendererMeshProperties MeshProperties;
-	for (const FNiagaraMeshRendererMeshPropertiesBase& MeshEntry : ArrayRef.GetArray())
+	if (IterateDelegate.IsSet())
 	{
-		MeshProperties.Mesh				= MeshEntry.Mesh;
-		MeshProperties.Scale			= MeshEntry.Scale;
-		MeshProperties.Rotation			= MeshEntry.Rotation;
-		MeshProperties.PivotOffset		= MeshEntry.PivotOffset;
-		MeshProperties.PivotOffsetSpace	= MeshEntry.PivotOffsetSpace;
-		IterateDelegate(MeshProperties);
+		FNiagaraMeshRendererMeshProperties MeshProperties;
+		for (const FNiagaraMeshRendererMeshPropertiesBase& MeshEntry : ArrayRef.GetArray())
+		{
+			MeshProperties.Mesh				= MeshEntry.Mesh;
+			MeshProperties.Scale			= MeshEntry.Scale;
+			MeshProperties.Rotation			= MeshEntry.Rotation;
+			MeshProperties.PivotOffset		= MeshEntry.PivotOffset;
+			MeshProperties.PivotOffsetSpace	= MeshEntry.PivotOffsetSpace;
+			IterateDelegate(MeshProperties);
+		}
 	}
 }
 
diff --git a/Engine/Plugins/FX/Niagara/Source/Niagara/Private/DataInterface/NiagaraDataInterfaceDataChannelRead.cpp b/Engine/Plugins/FX/Niagara/Source/Niagara/Private/DataInterface/NiagaraDataInterfaceDataChannelRead.cpp
--- a/Engine/Plugins/FX/Niagara/Source/Niagara/Private/DataInterface/NiagaraDataInterfaceDataChannelRead.cpp
+++ b/Engine/Plugins/FX/Niagara/Source/Niagara/Private/DataInterface/NiagaraDataInterfaceDataChannelRead.cpp
@@ -1400,7 +1400,7 @@ void UNiagaraDataInterfaceDataChannelRead::Num(FVectorVMExternalFunctionContext&
 /** 
 GetNDCSpawnData - Retrieves spawn data about the NDC that spawned a particular particle.
 Uses an indirection table that decomposes each NDC spawn into power of 2 buckets of particles.
-We do this to strike a balance between allowing many particles per NDC entry and allowing many NDC entires to spawn particles.
+We do this to strike a balance between allowing many particles per NDC entry and allowing many NDC entries to spawn particles.
 We have a max of 16 buckets with the highest being for spawns with 1<15 particles or more and the lowest being for individual particles.
 An example with two NDC entries. The first spawning 10 and the second 8.
 The first's spawn count decomposes into an entry in the 8 bucket and 2 bucket.
@@ -1410,7 +1410,7 @@ We have 16 buckets so the bucket counts array looks like
 This means the we have 2 buckets with data so the rest of our buffer is.
 0,1,0,
 
-We have two entires in the 8 bucket. 
+We have two entries in the 8 bucket. 
 As we spawn particles we use our bucket sizes counts and the exec index to see which bucket entry each exec index should use.
 The first 8 particles processed, exec index 0-7 will lookup the first entry and so use NDC 0.
 THe next 8 particles, exec index 8-15 will use the next and so use NDC 1.
@@ -1775,7 +1775,7 @@ void UNiagaraDataInterfaceDataChannelRead::Consume(FVectorVMExternalFunctionCont
 
 		if (OutIndex.IsValid())
 		{
-			FMemory::Memset(OutSuccess.Data.GetDest(), 0xFF, sizeof(int32) * Context.GetNumInstances());
+			FMemory::Memset(OutIndex.Data.GetDest(), 0xFF, sizeof(int32) * Context.GetNumInstances());
 		}
 	}
 }
diff --git a/Engine/Plugins/FX/Niagara/Source/Niagara/Private/DataInterface/NiagaraDataInterfaceDebugDraw.cpp b/Engine/Plugins/FX/Niagara/Source/Niagara/Private/DataInterface/NiagaraDataInterfaceDebugDraw.cpp
--- a/Engine/Plugins/FX/Niagara/Source/Niagara/Private/DataInterface/NiagaraDataInterfaceDebugDraw.cpp
+++ b/Engine/Plugins/FX/Niagara/Source/Niagara/Private/DataInterface/NiagaraDataInterfaceDebugDraw.cpp
@@ -223,6 +223,11 @@ struct FNDIDebugDrawInstanceData_GameThread
 
 	void AddCylinder(const FVector3f& Location, const FQuat4f& Rotation, float Height, float Radius, int32 NumHeightSegments, int32 NumRadiusSegments, const FLinearColor& Color, bool bHemiXDraw, bool bHemiYDraw, const FVector3f& NonUniformScale, bool bCenterVertically)
 	{
+		if (NumHeightSegments == 0 || NumRadiusSegments == 0)
+		{
+			return;
+		}
+
 		const FVector3f HeightVector = FVector3f::ZAxisVector * Height;
 		const FVector3f HalfHeightOffset = bCenterVertically? HeightVector * 0.5f : FVector3f::ZeroVector;
 		const FVector3f AxisStep = FVector3f::ZAxisVector * Height / NumHeightSegments;
diff --git a/Engine/Plugins/FX/Niagara/Source/Niagara/Private/DataInterface/NiagaraDataInterfaceDynamicMesh.cpp b/Engine/Plugins/FX/Niagara/Source/Niagara/Private/DataInterface/NiagaraDataInterfaceDynamicMesh.cpp
--- a/Engine/Plugins/FX/Niagara/Source/Niagara/Private/DataInterface/NiagaraDataInterfaceDynamicMesh.cpp
+++ b/Engine/Plugins/FX/Niagara/Source/Niagara/Private/DataInterface/NiagaraDataInterfaceDynamicMesh.cpp
@@ -1567,7 +1567,7 @@ bool UNiagaraDataInterfaceDynamicMesh::SimCacheCompareFrame(const UObject* LhsSt
 	for (ENiagaraSimTarget SimTarget : {ENiagaraSimTarget::CPUSim, ENiagaraSimTarget::GPUComputeSim})
 	{
 		const FNDIDynamicMeshSimCacheFrame* LhsMeshData = LhsCacheData->GetFrameData(FrameIndex, SimTarget);
-		const FNDIDynamicMeshSimCacheFrame* RhsMeshData = LhsCacheData->GetFrameData(FrameIndex, SimTarget);
+		const FNDIDynamicMeshSimCacheFrame* RhsMeshData = RhsCacheData->GetFrameData(FrameIndex, SimTarget);
 
 		const bool LhsMeshDataValid = LhsMeshData != nullptr;
 		const bool RhsMeshDataValid = RhsMeshData != nullptr;
@@ -1953,7 +1953,7 @@ void UNiagaraDataInterfaceDynamicMesh::VMGetVertexTexCoord(FVectorVMExternalFunc
 	{
 		const int32 VertexIndex = InVertexIndex.GetAndAdvance();
 		const int32 TexCoordIndex = InTexCoordIndex.GetAndAdvance();
-		const bool bValid = VertexIndex >= 0 && VertexIndex < NumVertices && TexCoordIndex > 0 && TexCoordIndex < NumTexCoords;
+		const bool bValid = VertexIndex >= 0 && VertexIndex < NumVertices && TexCoordIndex >= 0 && TexCoordIndex < NumTexCoords;
 		OutValid.SetAndAdvance(bValid);
 		OutTexCoord.SetAndAdvance(bValid ? InstanceData->GetVertexTexCoord(VertexIndex, TexCoordIndex) : FVector2f::ZeroVector);
 	}
@@ -2140,7 +2140,7 @@ void UNiagaraDataInterfaceDynamicMesh::VMSetVertexTexCoord(FVectorVMExternalFunc
 		const int32		TexCoord = InTexCoord.GetAndAdvance();
 		const FVector2f	Coord = InCoord.GetAndAdvance();
 
-		if (bExecute && VertexIndex >= 0 && VertexIndex < NumVertices && TexCoord > 0 && TexCoord < NumTexCoords)
+		if (bExecute && VertexIndex >= 0 && VertexIndex < NumVertices && TexCoord >= 0 && TexCoord < NumTexCoords)
 		{
 			InstanceData->SetVertexTexCoord(VertexIndex, TexCoord, Coord);
 		}
diff --git a/Engine/Plugins/FX/Niagara/Source/Niagara/Private/DataInterface/NiagaraDataInterfaceEmitterProperties.cpp b/Engine/Plugins/FX/Niagara/Source/Niagara/Private/DataInterface/NiagaraDataInterfaceEmitterProperties.cpp
--- a/Engine/Plugins/FX/Niagara/Source/Niagara/Private/DataInterface/NiagaraDataInterfaceEmitterProperties.cpp
+++ b/Engine/Plugins/FX/Niagara/Source/Niagara/Private/DataInterface/NiagaraDataInterfaceEmitterProperties.cpp
@@ -26,8 +26,8 @@ namespace NDIEmitterPropertiesLocal
 
 	struct FInstanceData_GameToRender
 	{
-		FBox	Bounds;
-		FBox	FixedBounds;
+		FBox	Bounds = FBox(EForceInit::ForceInit);
+		FBox	FixedBounds = FBox(EForceInit::ForceInit);
 	};
 
 	struct FInstanceData_GameThread
diff --git a/Engine/Plugins/FX/Niagara/Source/Niagara/Private/DataInterface/NiagaraDataInterfaceMemoryBuffer.cpp b/Engine/Plugins/FX/Niagara/Source/Niagara/Private/DataInterface/NiagaraDataInterfaceMemoryBuffer.cpp
--- a/Engine/Plugins/FX/Niagara/Source/Niagara/Private/DataInterface/NiagaraDataInterfaceMemoryBuffer.cpp
+++ b/Engine/Plugins/FX/Niagara/Source/Niagara/Private/DataInterface/NiagaraDataInterfaceMemoryBuffer.cpp
@@ -323,8 +323,8 @@ namespace NDIMemoryBufferLocal
 			const int32 Offset			= InElementOffset.GetAndAdvance();
 			const int32 ClampedOffset	= FMath::Clamp(Offset, 0, InstanceData->NumElements);
 			const TType Value			= InValue.GetAndAdvance();
-			const bool bSuccess			= bExecute && Offset > 0 && Offset <= ValidEnd;
-			if (bExecute)
+			const bool bSuccess			= bExecute && Offset >= 0 && Offset <= ValidEnd;
+			if (bSuccess)
 			{
 				FMemory::Memcpy(&InstanceData->CpuBuffer[ClampedOffset], &Value, sizeof(TType));
 			}
@@ -444,6 +444,8 @@ namespace NDIMemoryBufferLocal
 				PreviousValue = TOperation::InterlockedOp(reinterpret_cast<int32*>(&InstanceData->CpuBuffer[ClampedOffset]), Value);
 				CurrentValue = TOperation::RawOp(PreviousValue, Value);
 			}
+			OutPreviousValue.SetAndAdvance(PreviousValue);
+			OutCurrentValue.SetAndAdvance(CurrentValue);
 		}
 	}
 
@@ -477,6 +479,8 @@ namespace NDIMemoryBufferLocal
 				PreviousValue = FPlatformAtomics::InterlockedCompareExchange(reinterpret_cast<int32*>(&InstanceData->CpuBuffer[ClampedOffset]), Value, Comperand);
 				CurrentValue = PreviousValue == Comperand ? Value : PreviousValue;
 			}
+			OutPreviousValue.SetAndAdvance(PreviousValue);
+			OutCurrentValue.SetAndAdvance(CurrentValue);
 		}
 	}
 
diff --git a/Engine/Plugins/FX/Niagara/Source/Niagara/Private/NiagaraDataChannel_Map.cpp b/Engine/Plugins/FX/Niagara/Source/Niagara/Private/NiagaraDataChannel_Map.cpp
--- a/Engine/Plugins/FX/Niagara/Source/Niagara/Private/NiagaraDataChannel_Map.cpp
+++ b/Engine/Plugins/FX/Niagara/Source/Niagara/Private/NiagaraDataChannel_Map.cpp
@@ -18,7 +18,7 @@ namespace NDCMapLocal
 	static FAutoConsoleVariableRef CVarEntryPoolInitialSize(TEXT("fx.Niagara.DataChannels.NDCMapBase.InitialSize"), EntryPoolInitialSize, TEXT("The initial size of the entry pool for NDC Map types."), ECVF_Default);
 
 	float EntryPoolUnusedCleanTime = 60.0f;
-	static FAutoConsoleVariableRef CVarEntryPoolUnusedCleanTime(TEXT("fx.Niagara.DataChannels.NDCMapBase.FreeUnusedTime"), EntryPoolUnusedCleanTime, TEXT("Time in seconds we keep unused entires in the pool."), ECVF_Default);
+	static FAutoConsoleVariableRef CVarEntryPoolUnusedCleanTime(TEXT("fx.Niagara.DataChannels.NDCMapBase.FreeUnusedTime"), EntryPoolUnusedCleanTime, TEXT("Time in seconds we keep unused entries in the pool."), ECVF_Default);
 }
 
 //////////////////////////////////////////////////////////////////////////
@@ -222,7 +222,7 @@ FNDCMapEntry& UNiagaraDataChannelHandler_MapBase::FindOrAddEntry(FNDCAccessConte
 		{					
 			Ret = &EntryPool[*ActiveEntry];
 		}
-		else if (!TypedContext.bIsAutoLinkingSystem)//Don't create new entires for auto linking niagra systems.
+		else if (!TypedContext.bIsAutoLinkingSystem)//Don't create new entries for auto linking niagra systems.
 		{
 			int32 NewEntryIdx = INDEX_NONE;
 			if (FreeEntries.Num() > 0)
