- Niagara mesh changed handling reworked to cover user parameters
- Before this change any mesh that was modified which was referenced by a user parameter binding (single mesh or whole array) would be missed and could cause crashing

#Synced-CL 50885476

diff --git a/Engine/Plugins/FX/Niagara/Source/Niagara/Private/NiagaraMeshRendererProperties.cpp b/Engine/Plugins/FX/Niagara/Source/Niagara/Private/NiagaraMeshRendererProperties.cpp
--- a/Engine/Plugins/FX/Niagara/Source/Niagara/Private/NiagaraMeshRendererProperties.cpp
+++ b/Engine/Plugins/FX/Niagara/Source/Niagara/Private/NiagaraMeshRendererProperties.cpp
@@ -41,6 +41,127 @@
 
 TArray<TWeakObjectPtr<UNiagaraMeshRendererProperties>> UNiagaraMeshRendererProperties::MeshRendererPropertiesToDeferredInit;
 
+
+////////////////////////////////////////////////////////////////////////////////////////////////////////////////////////////////
+
+namespace NiagaraMeshRendererPropertiesInternal
+{
+	bool IsStaticMeshSafeToUse(const FNiagaraEmitterInstance* EmitterInstance, UStaticMesh* StaticMesh)
+	{
+		if (StaticMesh)
+		{
+			bool bInvalidDueToCompiling = false;
+#if WITH_EDITOR
+			// During EOF updates we can not test GetRenderData as we will cause a wait / reregister to occur which is invalid while in EOF updates
+			// The assumption is that when we do hit this situation we are waiting on a static mesh build from reimport, etc, so we can just skip until
+			// we get the post build callback, ideally we would have a better way to handle this by not rendering while the build is in progress but
+			// that requires quite a large rework of how we handle reading from the static mesh data, especially around updating MICs in PostLoad.
+			if (StaticMesh->IsCompiling())
+			{
+				FNiagaraSystemInstance* SystemInstance = EmitterInstance ? EmitterInstance->GetParentSystemInstance() : nullptr;
+				UWorld* World = SystemInstance ? SystemInstance->GetWorld() : nullptr;
+				bInvalidDueToCompiling = World && World->bPostTickComponentUpdate;
+			}
+			else
+#endif
+			{
+				const FNiagaraSystemInstance* SystemInstance = EmitterInstance ? EmitterInstance->GetParentSystemInstance() : nullptr;
+				const EShaderPlatform ShaderPlatform = SystemInstance != nullptr ? GShaderPlatformForFeatureLevel[SystemInstance->GetFeatureLevel()] : GMaxRHIShaderPlatform;
+
+				if (StaticMesh->HasValidNaniteData() && !StaticMesh->HasNaniteFallbackMesh(ShaderPlatform))
+				{
+					GEngine->AddOnScreenDebugMessage(uint64(EmitterInstance) ^ uint64(StaticMesh), 10.f, FColor::Yellow,
+						*FString::Printf(TEXT("NiagaraRenderableMesh does not have a fallback mesh. Enable Nanite fallback mesh to fix this issue. System: %s, Mesh: %s"), *GetFullNameSafe(SystemInstance ? SystemInstance->GetSystem() : nullptr), *GetFullNameSafe(StaticMesh)));
+
+					if (FNiagaraUtilities::LogVerboseWarnings())
+					{
+						UE_LOG(LogNiagara, Warning, TEXT("NiagaraRenderableMesh does not have a fallback mesh. Enable Nanite fallback mesh to fix this issue. System: %s, Mesh: %s"), *GetFullNameSafe(SystemInstance ? SystemInstance->GetSystem() : nullptr), *GetFullNameSafe(StaticMesh));
+					}
+
+					return false;
+				}
+			}
+
+			return !bInvalidDueToCompiling && StaticMesh->GetRenderData();
+		}
+		return false;
+	}
+
+	void ResolveRenderableMeshInternal(const FNiagaraMeshRendererMeshProperties& MeshProperties, const FNiagaraEmitterInstance* EmitterInstance, INiagaraRenderableMeshInterface*& OutInterface, UStaticMesh*& OutStaticMesh)
+	{
+		OutInterface = nullptr;
+		OutStaticMesh = nullptr;
+		if (EmitterInstance)
+		{
+			const FNiagaraVariableBase& MeshParameter = MeshProperties.MeshParameterBinding.ResolvedParameter;
+			if (MeshParameter.IsValid())
+			{
+				if (MeshParameter.IsDataInterface())
+				{
+					OutInterface = Cast<INiagaraRenderableMeshInterface>(EmitterInstance->GetRendererBoundVariables().GetDataInterface(MeshParameter));
+					if (OutInterface != nullptr)
+					{
+						return;
+					}
+				}
+
+				UStaticMesh* BoundMesh = Cast<UStaticMesh>(EmitterInstance->GetRendererBoundVariables().GetUObject(MeshParameter));
+				if (IsStaticMeshSafeToUse(EmitterInstance, BoundMesh))
+				{
+					OutStaticMesh = BoundMesh;
+					return;
+				}
+			}
+		}
+		if (IsStaticMeshSafeToUse(EmitterInstance, MeshProperties.Mesh))
+		{
+			OutStaticMesh = MeshProperties.Mesh;
+		}
+	}
+
+#if WITH_EDITOR
+	struct FMeshChangedHandlers
+	{
+		FDelegateHandle OnMeshChanged;
+		FDelegateHandle OnPreMeshBuild;
+		FDelegateHandle OnPostMeshBuild;
+	};
+
+	static TMap<FObjectKey, FMeshChangedHandlers> GOnMeshChangesObjects;
+
+	void AddMeshChangedIfNeeded(UStaticMesh* StaticMesh)
+	{
+		if (GIsEditor && StaticMesh && IsInGameThread() && !GOnMeshChangesObjects.Contains(StaticMesh))
+		{
+			FMeshChangedHandlers& Handlers = GOnMeshChangesObjects.Add(StaticMesh);
+			Handlers.OnMeshChanged = StaticMesh->GetOnMeshChanged().AddLambda([StaticMesh]() { UNiagaraMeshRendererProperties::OnMeshChanged(StaticMesh); });
+			Handlers.OnPreMeshBuild = StaticMesh->OnPreMeshBuild().AddLambda([](UStaticMesh* StaticMesh) { UNiagaraMeshRendererProperties::OnMeshChanged(StaticMesh); });
+			Handlers.OnPostMeshBuild = StaticMesh->OnPostMeshBuild().AddLambda([](UObject* StaticMesh) { UNiagaraMeshRendererProperties::OnMeshChanged(StaticMesh); });
+		}
+	}
+
+	void RemoveMeshChangedHandlers()
+	{
+		if (GIsEditor)
+		{
+			for (auto It=GOnMeshChangesObjects.CreateIterator(); It; ++It)
+			{
+				FObjectKey MeshObjectKey = It.Key();
+				if (UStaticMesh* StaticMesh = Cast<UStaticMesh>(MeshObjectKey.ResolveObjectPtr()))
+				{
+					FMeshChangedHandlers& Handlers = It.Value();
+
+					StaticMesh->GetOnMeshChanged().Remove(Handlers.OnMeshChanged);
+					StaticMesh->OnPreMeshBuild().Remove(Handlers.OnPreMeshBuild);
+					StaticMesh->OnPostMeshBuild().Remove(Handlers.OnPostMeshBuild);
+				}
+			}
+			GOnMeshChangesObjects.Empty();
+		}
+	}
+#endif
+}
+
 ////////////////////////////////////////////////////////////////////////////////////////////////////////////////////////////////
 
 class FNiagaraRenderableStaticMesh : public INiagaraRenderableMesh
@@ -53,6 +174,11 @@ public:
 		MeshMinLOD		= StaticMesh->GetMinLODIdx();
 		MinLOD			= StaticMesh->GetMinLODIdx();
 		LocalBounds		= StaticMesh->GetExtendedBounds().GetBox();
+
+	#if WITH_EDITOR
+		// Const cast as it's just for the editor to simplify the code
+		NiagaraMeshRendererPropertiesInternal::AddMeshChangedIfNeeded(const_cast<UStaticMesh*>(StaticMesh));
+	#endif
 	}
 
 	void SetMinLODBias(int32 MinLODBias) override
@@ -226,84 +352,6 @@ public:
 
 ////////////////////////////////////////////////////////////////////////////////////////////////////////////////////////////////
 
-namespace NiagaraMeshRendererPropertiesInternal
-{
-	bool IsStaticMeshSafeToUse(const FNiagaraEmitterInstance* EmitterInstance, UStaticMesh* StaticMesh)
-	{
-		if (StaticMesh)
-		{
-			bool bInvalidDueToCompiling = false;
-		#if WITH_EDITOR
-			// During EOF updates we can not test GetRenderData as we will cause a wait / reregister to occur which is invalid while in EOF updates
-			// The assumption is that when we do hit this situation we are waiting on a static mesh build from reimport, etc, so we can just skip until
-			// we get the post build callback, ideally we would have a better way to handle this by not rendering while the build is in progress but
-			// that requires quite a large rework of how we handle reading from the static mesh data, especially around updating MICs in PostLoad.
-			if (StaticMesh->IsCompiling())
-			{
-				FNiagaraSystemInstance* SystemInstance = EmitterInstance ? EmitterInstance->GetParentSystemInstance() : nullptr;
-				UWorld* World = SystemInstance ? SystemInstance->GetWorld() : nullptr;
-				bInvalidDueToCompiling = World && World->bPostTickComponentUpdate;
-			}
-			else
-		#endif
-			{
-				const FNiagaraSystemInstance* SystemInstance = EmitterInstance ? EmitterInstance->GetParentSystemInstance() : nullptr;
-				const EShaderPlatform ShaderPlatform = SystemInstance != nullptr ? GShaderPlatformForFeatureLevel[SystemInstance->GetFeatureLevel()] : GMaxRHIShaderPlatform;
-
-				if (StaticMesh->HasValidNaniteData() && !StaticMesh->HasNaniteFallbackMesh(ShaderPlatform))
-				{
-					GEngine->AddOnScreenDebugMessage(uint64(EmitterInstance) ^ uint64(StaticMesh), 10.f, FColor::Yellow,
-						*FString::Printf(TEXT("NiagaraRenderableMesh does not have a fallback mesh. Enable Nanite fallback mesh to fix this issue. System: %s, Mesh: %s"), *GetFullNameSafe(SystemInstance ? SystemInstance->GetSystem() : nullptr), *GetFullNameSafe(StaticMesh)));
-					
-					if (FNiagaraUtilities::LogVerboseWarnings())
-					{
-						UE_LOG(LogNiagara, Warning, TEXT("NiagaraRenderableMesh does not have a fallback mesh. Enable Nanite fallback mesh to fix this issue. System: %s, Mesh: %s"), *GetFullNameSafe(SystemInstance ? SystemInstance->GetSystem() : nullptr), *GetFullNameSafe(StaticMesh));
-					}
-
-					return false;
-				}
-			}
-
-			return !bInvalidDueToCompiling && StaticMesh->GetRenderData();
-		}
-		return false;
-	}
-
-	void ResolveRenderableMeshInternal(const FNiagaraMeshRendererMeshProperties& MeshProperties, const FNiagaraEmitterInstance* EmitterInstance, INiagaraRenderableMeshInterface*& OutInterface, UStaticMesh*& OutStaticMesh)
-	{
-		OutInterface = nullptr;
-		OutStaticMesh = nullptr;
-		if (EmitterInstance)
-		{
-			const FNiagaraVariableBase& MeshParameter = MeshProperties.MeshParameterBinding.ResolvedParameter;
-			if (MeshParameter.IsValid())
-			{
-				if (MeshParameter.IsDataInterface())
-				{
-					OutInterface = Cast<INiagaraRenderableMeshInterface>(EmitterInstance->GetRendererBoundVariables().GetDataInterface(MeshParameter));
-					if (OutInterface != nullptr)
-					{
-						return;
-					}
-				}
-
-				UStaticMesh* BoundMesh = Cast<UStaticMesh>(EmitterInstance->GetRendererBoundVariables().GetUObject(MeshParameter));
-				if (IsStaticMeshSafeToUse(EmitterInstance, BoundMesh))
-				{
-					OutStaticMesh = BoundMesh;
-					return;
-				}
-			}
-		}
-		if (IsStaticMeshSafeToUse(EmitterInstance, MeshProperties.Mesh))
-		{
-			OutStaticMesh = MeshProperties.Mesh;
-		}
-	}
-}
-
-////////////////////////////////////////////////////////////////////////////////////////////////////////////////////////////////
-
 FNiagaraMeshMaterialOverride::FNiagaraMeshMaterialOverride()
 	: ExplicitMat(nullptr)
 	, UserParamBinding(FNiagaraTypeDefinition(UMaterialInterface::StaticClass()))
@@ -755,6 +803,77 @@ void UNiagaraMeshRendererProperties::CacheFromCompiledData(const FNiagaraDataSet
 #endif
 }
 
+#if WITH_EDITOR
+void UNiagaraMeshRendererProperties::OnMeshChanged()
+{
+	FNiagaraSystemUpdateContext ReregisterContext;
+
+	if (UNiagaraSystem* NiagaraSystem = GetTypedOuter<UNiagaraSystem>())
+	{
+		ReregisterContext.Add(NiagaraSystem, true);
+	}
+
+	CheckMaterialUsage();
+	UpdateMICs();
+}
+
+void UNiagaraMeshRendererProperties::OnMeshChanged(UObject* MeshObject)
+{
+	// We need to guard recursive calls that can come from PSO caching requesting the mesh which can call the post build step
+	// This is all very messy as we do not get a consistent callback from any one of the methods for various operations so we have to litter this across various calls
+	static bool GRecursionGuard = false;
+	if (GRecursionGuard)
+	{
+		return;
+	}
+
+	TGuardValue<bool> RecursionGuardValue(GRecursionGuard, true);
+
+	FNiagaraSystemUpdateContext ReregisterContext;
+
+	TSet<FObjectKey> AddedSystems;
+	for (TObjectIterator<UNiagaraMeshRendererProperties> It; It; ++It)
+	{
+		UNiagaraMeshRendererProperties* Properties = *It;
+		UNiagaraSystem* System = ::IsValid(Properties) ? Properties->GetTypedOuter<UNiagaraSystem>() : nullptr;
+		if (::IsValid(System) == false)
+		{
+			continue;
+		}
+
+		// Work out if we need add this system, and if we need to update the MICs / material usage
+		bool bHasExternalBinding = Properties->MeshesBinding.ResolvedParameter.IsValid();
+		bool bHasReference = false;
+		for (const FNiagaraMeshRendererMeshProperties& Mesh : Properties->Meshes)
+		{
+			bHasReference |= Mesh.Mesh == MeshObject;
+			bHasExternalBinding |= Mesh.MeshParameterBinding.ResolvedParameter.IsValid();
+		}
+		if (!bHasReference && !bHasExternalBinding)
+		{
+			continue;
+		}
+
+		if (!AddedSystems.Contains(System))
+		{
+			AddedSystems.Add(System);
+			ReregisterContext.Add(System, true);
+		}
+
+		if (bHasReference)
+		{
+			Properties->CheckMaterialUsage();
+			Properties->UpdateMICs();
+		}
+	}
+}
+
+void UNiagaraMeshRendererProperties::RemoveMeshChangedHandlers()
+{
+	NiagaraMeshRendererPropertiesInternal::RemoveMeshChangedHandlers();
+}
+#endif //WITH_EDITOR
+
 void UNiagaraMeshRendererProperties::UpdateMICs()
 {
 #if WITH_EDITORONLY_DATA
@@ -1090,14 +1209,6 @@ void UNiagaraMeshRendererProperties::PostLoad()
 		if (MeshProperties.Mesh)
 		{
 			MeshProperties.Mesh->ConditionalPostLoad();
-#if WITH_EDITOR
-			if (GIsEditor)
-			{
-				MeshProperties.Mesh->GetOnMeshChanged().AddUObject(this, &UNiagaraMeshRendererProperties::OnMeshChanged);
-				MeshProperties.Mesh->OnPreMeshBuild().AddUObject(this, &UNiagaraMeshRendererProperties::OnMeshPostBuild);
-				MeshProperties.Mesh->OnPostMeshBuild().AddUObject(this, &UNiagaraMeshRendererProperties::OnMeshPostBuild);
-			}
-#endif
 		}
 #if WITH_EDITORONLY_DATA
 		if (MeshProperties.UserParamBinding_DEPRECATED.Parameter.GetName().IsNone() == false)
@@ -1267,43 +1378,6 @@ void UNiagaraMeshRendererProperties::GetRendererMaterialUsageFeedback(TArray<FNi
 	GetMaterialUsageFeedback(MATUSAGE_NiagaraMeshParticles, { MD_Volume }, OutFeedback);
 }
 
-void UNiagaraMeshRendererProperties::BeginDestroy()
-{
-	Super::BeginDestroy();
-#if WITH_EDITOR
-	if (GIsEditor && !HasAnyFlags(RF_ClassDefaultObject))
-	{
-		for (const auto& MeshProperties : Meshes)
-		{
-			if (MeshProperties.Mesh)
-			{
-				MeshProperties.Mesh->GetOnMeshChanged().RemoveAll(this);
-				MeshProperties.Mesh->OnPreMeshBuild().RemoveAll(this);
-				MeshProperties.Mesh->OnPostMeshBuild().RemoveAll(this);
-			}
-		}
-	}
-#endif
-}
-
-void UNiagaraMeshRendererProperties::PreEditChange(class FProperty* PropertyThatWillChange)
-{
-	Super::PreEditChange(PropertyThatWillChange);
-
-	if (ChangeRequiresMeshListRebuild(PropertyThatWillChange))
-	{
-		for (const auto& MeshProperties : Meshes)
-		{
-			if (MeshProperties.Mesh)
-			{
-				MeshProperties.Mesh->GetOnMeshChanged().RemoveAll(this);
-				MeshProperties.Mesh->OnPreMeshBuild().RemoveAll(this);
-				MeshProperties.Mesh->OnPostMeshBuild().RemoveAll(this);
-			}
-		}
-	}
-}
-
 void UNiagaraMeshRendererProperties::PostEditChangeProperty(FPropertyChangedEvent& PropertyChangedEvent)
 {
 	SubImageSize.X = FMath::Max(SubImageSize.X, 1.0);
@@ -1314,20 +1388,6 @@ void UNiagaraMeshRendererProperties::PostEditChangeProperty(FPropertyChangedEven
 	const FName PropertyName = PropertyChangedEvent.GetPropertyName();
 	const FName MemberPropertyName = PropertyChangedEvent.GetMemberPropertyName();
 
-	if (bIsRedirect)
-	{
-		// Do this in case the redirected property is not a mesh (we have no way of knowing b/c the property is nullptr)
-		for (const auto& MeshProperties : Meshes)
-		{
-			if (MeshProperties.Mesh)
-			{
-				MeshProperties.Mesh->GetOnMeshChanged().RemoveAll(this);
-				MeshProperties.Mesh->OnPreMeshBuild().RemoveAll(this);
-				MeshProperties.Mesh->OnPostMeshBuild().RemoveAll(this);
-			}
-		}
-	}
-
 	if (bRebuildMeshList)
 	{
 		if (!IsRunningCommandlet() &&
@@ -1360,21 +1420,6 @@ void UNiagaraMeshRendererProperties::PostEditChangeProperty(FPropertyChangedEven
 		}
 	}
 
-	if (bIsRedirect || bRebuildMeshList)
-	{
-		// We only need to check material usage as we will invalidate any renderers later on
-		CheckMaterialUsage();
-		for (const auto& MeshProperties : Meshes)
-		{
-			if (MeshProperties.Mesh)
-			{
-				MeshProperties.Mesh->GetOnMeshChanged().AddUObject(this, &UNiagaraMeshRendererProperties::OnMeshChanged);
-				MeshProperties.Mesh->OnPreMeshBuild().AddUObject(this, &UNiagaraMeshRendererProperties::OnMeshPostBuild);
-				MeshProperties.Mesh->OnPostMeshBuild().AddUObject(this, &UNiagaraMeshRendererProperties::OnMeshPostBuild);
-			}
-		}
-	}
-
 	// If changing the source mode, we may need to update many of our values.
 	if (PropertyName == TEXT("SourceMode"))
 	{
@@ -1456,36 +1501,6 @@ void UNiagaraMeshRendererProperties::RemoveVariable(const FNiagaraVariableBase&
 #endif
 }
 
-void UNiagaraMeshRendererProperties::OnMeshChanged()
-{
-	FNiagaraSystemUpdateContext ReregisterContext;
-
-	if (UNiagaraSystem* NiagaraSystem = GetTypedOuter<UNiagaraSystem>())
-	{
-		ReregisterContext.Add(NiagaraSystem, true);
-	}
-
-	CheckMaterialUsage();
-	UpdateMICs();
-}
-
-void UNiagaraMeshRendererProperties::OnMeshPostBuild(UStaticMesh*)
-{
-	OnMeshChanged();
-}
-
-void UNiagaraMeshRendererProperties::OnAssetReimported(UObject* Object)
-{
-	for (auto& MeshInfo : Meshes)
-	{
-		if (MeshInfo.Mesh == Object)
-		{
-			OnMeshChanged();
-			break;
-		}
-	}
-}
-
 void UNiagaraMeshRendererProperties::CheckMaterialUsage()
 {
 	for (const auto& MeshProperties : Meshes)
diff --git a/Engine/Plugins/FX/Niagara/Source/Niagara/Private/NiagaraModule.cpp b/Engine/Plugins/FX/Niagara/Source/Niagara/Private/NiagaraModule.cpp
--- a/Engine/Plugins/FX/Niagara/Source/Niagara/Private/NiagaraModule.cpp
+++ b/Engine/Plugins/FX/Niagara/Source/Niagara/Private/NiagaraModule.cpp
@@ -650,6 +650,10 @@ void INiagaraModule::ShutdownRenderingResources()
 
 void INiagaraModule::ShutdownModule()
 {
+#if WITH_EDITOR
+	UNiagaraMeshRendererProperties::RemoveMeshChangedHandlers();
+#endif
+
 	FWorldDelegates::OnWorldBeginTearDown.RemoveAll(this);
 	FWorldDelegates::OnWorldTickStart.RemoveAll(this);
 	FCoreDelegates::OnBeginFrame.RemoveAll(this);
diff --git a/Engine/Plugins/FX/Niagara/Source/Niagara/Public/NiagaraMeshRendererProperties.h b/Engine/Plugins/FX/Niagara/Source/Niagara/Public/NiagaraMeshRendererProperties.h
--- a/Engine/Plugins/FX/Niagara/Source/Niagara/Public/NiagaraMeshRendererProperties.h
+++ b/Engine/Plugins/FX/Niagara/Source/Niagara/Public/NiagaraMeshRendererProperties.h
@@ -136,8 +136,6 @@ public:
 	NIAGARA_API virtual void Serialize(FArchive& Ar) override;
 	NIAGARA_API virtual void GetResourceSizeEx(FResourceSizeEx& CumulativeResourceSize) override;
 #if WITH_EDITORONLY_DATA
-	NIAGARA_API virtual void BeginDestroy() override;
-	NIAGARA_API virtual void PreEditChange(class FProperty* PropertyThatWillChange) override;
 	NIAGARA_API virtual void PostEditChangeProperty(struct FPropertyChangedEvent& PropertyChangedEvent) override;
 	NIAGARA_API virtual void RenameVariable(const FNiagaraVariableBase& OldVariable, const FNiagaraVariableBase& NewVariable, const FVersionedNiagaraEmitterBase& InEmitter) override;
 	NIAGARA_API virtual void RemoveVariable(const FNiagaraVariableBase& OldVariable, const FVersionedNiagaraEmitterBase& InEmitter) override;
@@ -164,14 +162,20 @@ public:
 	NIAGARA_API virtual	void GetRendererTooltipWidgets(const FNiagaraEmitterInstance* InEmitter, TArray<TSharedPtr<SWidget>>& OutWidgets, TSharedPtr<FAssetThumbnailPool> InThumbnailPool) const override;
 	NIAGARA_API virtual void GetRendererFeedback(const FVersionedNiagaraEmitter& InEmitter, TArray<FNiagaraRendererFeedback>& OutErrors, TArray<FNiagaraRendererFeedback>& OutWarnings, TArray<FNiagaraRendererFeedback>& OutInfo) const override;
 	NIAGARA_API virtual void GetRendererMaterialUsageFeedback(TArray<FNiagaraRendererFeedback>& OutFeedback) const override;
-	NIAGARA_API void OnMeshChanged();
-	NIAGARA_API void OnMeshPostBuild(UStaticMesh*);
-	NIAGARA_API void OnAssetReimported(UObject*);
 	NIAGARA_API void CheckMaterialUsage();
 	NIAGARA_API virtual TArray<FNiagaraVariable> GetBoundAttributes() const override;
 #endif // WITH_EDITORONLY_DATA
 	NIAGARA_API virtual void CacheFromCompiledData(const FNiagaraDataSetCompiledData* CompiledData) override;
 
+#if WITH_EDITOR
+	// Old function should not be called by user code but is in some instances
+	NIAGARA_API void OnMeshChanged();
+	// Internal API call for handling mesh changes (reimport / modification / etc)
+	NIAGARA_API static void OnMeshChanged(UObject* MeshObject);
+	// Internal API to remove all of the mesh changed handlers
+	NIAGARA_API static void RemoveMeshChangedHandlers();
+#endif
+
 	NIAGARA_API void UpdateMICs();
 
 	virtual ENiagaraRendererSourceDataMode GetCurrentSourceMode() const override { return SourceMode; }
diff --git a/Engine/Plugins/FX/Niagara/Source/NiagaraEditor/Private/NiagaraEditorModule.cpp b/Engine/Plugins/FX/Niagara/Source/NiagaraEditor/Private/NiagaraEditorModule.cpp
--- a/Engine/Plugins/FX/Niagara/Source/NiagaraEditor/Private/NiagaraEditorModule.cpp
+++ b/Engine/Plugins/FX/Niagara/Source/NiagaraEditor/Private/NiagaraEditorModule.cpp
@@ -1905,13 +1905,7 @@ void FNiagaraEditorModule::OnPostEngineInit()
 			ImportSubsystem->OnAssetReimport.AddLambda(
 				[](UObject* ObjectReimported)
 				{
-					for (TObjectIterator<UNiagaraMeshRendererProperties> MeshRendererIt; MeshRendererIt; ++MeshRendererIt)
-					{
-						if (UNiagaraMeshRendererProperties* MeshRenderer = *MeshRendererIt)
-						{
-							MeshRenderer->OnAssetReimported(ObjectReimported);
-						}
-					}
+					UNiagaraMeshRendererProperties::OnMeshChanged(ObjectReimported);
 				}
 			);
 		}
