- Fix Niagara array atomic min/max being incorrect on GPU

#Synced-CL 50885476

diff --git a/Engine/Plugins/FX/Niagara/Shaders/Private/NiagaraDataInterfaceArrayRWTemplate.ush b/Engine/Plugins/FX/Niagara/Shaders/Private/NiagaraDataInterfaceArrayRWTemplate.ush
--- a/Engine/Plugins/FX/Niagara/Shaders/Private/NiagaraDataInterfaceArrayRWTemplate.ush
+++ b/Engine/Plugins/FX/Niagara/Shaders/Private/NiagaraDataInterfaceArrayRWTemplate.ush
@@ -117,7 +117,7 @@ void AtomicMin_{ParameterName}_UEImpureCall(bool bSkip, int Index, int Value, ou
 	if ( !bSkip && bIsValidIndex )
 	{
 		InterlockedMin({ParameterName}_ArrayRWBuffer[Index], Value, PrevValue);
-		CurrValue = PrevValue + Value;
+		CurrValue = min(PrevValue, Value);
 	}
 	else
 	{
@@ -133,7 +133,7 @@ void AtomicMax_{ParameterName}_UEImpureCall(bool bSkip, int Index, int Value, ou
 	if ( !bSkip && bIsValidIndex )
 	{
 		InterlockedMax({ParameterName}_ArrayRWBuffer[Index], Value, PrevValue);
-		CurrValue = PrevValue + Value;
+		CurrValue = max(PrevValue, Value);
 	}
 	else
 	{
