- Change check to ensure to avoid taking down the editor

#Synced-CL 50885476

diff --git a/Engine/Plugins/FX/Niagara/Source/NiagaraEditor/Public/NiagaraParameterPanelTypes.h b/Engine/Plugins/FX/Niagara/Source/NiagaraEditor/Public/NiagaraParameterPanelTypes.h
--- a/Engine/Plugins/FX/Niagara/Source/NiagaraEditor/Public/NiagaraParameterPanelTypes.h
+++ b/Engine/Plugins/FX/Niagara/Source/NiagaraEditor/Public/NiagaraParameterPanelTypes.h
@@ -71,9 +71,21 @@ public:
 	FOnRequestRename& GetOnRequestRename() const { return OnRequestRenameDelegate; };
 	FOnRequestRenameNamespaceModifier& GetOnRequestRenameNamespaceModifier() const { return OnRequestRenameNamespaceModifierDelegate; };
 
-	void RequestRename() const { check(OnRequestRenameDelegate.IsBound()); OnRequestRenameDelegate.ExecuteIfBound(); };
+	void RequestRename() const
+	{
+		if ( ensure(OnRequestRenameDelegate.IsBound()) )
+		{
+			OnRequestRenameDelegate.Execute();
+		}
+	}
 
-	void RequestRenameNamespaceModifier() const { check(OnRequestRenameNamespaceModifierDelegate.IsBound()); OnRequestRenameNamespaceModifierDelegate.ExecuteIfBound(); };
+	void RequestRenameNamespaceModifier() const
+	{
+		if ( ensure(OnRequestRenameNamespaceModifierDelegate.IsBound()) )
+		{
+			OnRequestRenameNamespaceModifierDelegate.Execute();
+		}
+	}
 
 	void AddToReadCount(const FNiagaraParameterReferencePath& SourcePath)
 	{
