Update NiagaraDumpBytecode commandlet to support dumping more information about the gpu scripts (-gpu), including a new csv with the information about all permutations.

#tests Dumped bytecode/shaders for a few projects

#Synced-CL 50944596

diff --git a/Engine/Plugins/FX/Niagara/Source/NiagaraEditor/Private/Commandlets/NiagaraDumpBytecodeCommandlet.cpp b/Engine/Plugins/FX/Niagara/Source/NiagaraEditor/Private/Commandlets/NiagaraDumpBytecodeCommandlet.cpp
--- a/Engine/Plugins/FX/Niagara/Source/NiagaraEditor/Private/Commandlets/NiagaraDumpBytecodeCommandlet.cpp
+++ b/Engine/Plugins/FX/Niagara/Source/NiagaraEditor/Private/Commandlets/NiagaraDumpBytecodeCommandlet.cpp
@@ -90,6 +90,11 @@ int32 UNiagaraDumpByteCodeCommandlet::Main(const FString& Params)
 		DumpTranslatedHlsl = true;
 	}
 
+	if (Switches.Contains("GPU"))
+	{
+		DumpGpuScripts = true;
+	}
+
 	ProcessNiagaraScripts();
 
 	return 0;
@@ -201,26 +206,40 @@ void UNiagaraDumpByteCodeCommandlet::ProcessBatch(TArray<FAssetData>& BatchAsset
 
 			if (FVersionedNiagaraEmitterData* Emitter = EmitterHandle.GetEmitterData())
 			{
-				if (Emitter->SimTarget == ENiagaraSimTarget::CPUSim)
+				const bool bCpuTarget = Emitter->SimTarget == ENiagaraSimTarget::CPUSim;
+
+				if (bCpuTarget || DumpGpuScripts)
 				{
 					FString EmitterName = EmitterHandle.GetUniqueInstanceName();
 
-					// trim the emitter name down to 16 char at most
-					const int32 MaxEmitterNameLength = 16;
+					// trim the emitter name down to 64 char at most
+					const int32 MaxEmitterNameLength = 64;
 					if (EmitterName.Len() > MaxEmitterNameLength)
 					{
 						UE_LOG(LogNiagaraDumpBytecodeCommandlet, Warning, TEXT("Emitter encountered with excessively long unique name: System: %s - Unique Emitter: %s"), *NiagaraSystem->GetFullName(), *EmitterName);
 						EmitterName.MidInline(0, MaxEmitterNameLength);
 					}
 
-					TArray<UNiagaraScript*> EmitterScripts;
-					Emitter->GetScripts(EmitterScripts);
+					IFileManager::Get().MakeDirectory(*(AuditOutputFolder / HashedPathName / EmitterName));
 
-					IFileManager::Get().MakeDirectory(*(HashedPathName / EmitterName));
+					if (bCpuTarget)
+					{
+						TArray<UNiagaraScript*> EmitterScripts;
+						Emitter->GetScripts(EmitterScripts);
 
-					for (const auto* EmitterScript : EmitterScripts)
+						for (const auto* EmitterScript : EmitterScripts)
+						{
+							DumpByteCode(EmitterScript, SystemPathName, HashedPathName, EmitterName / UsageEnum->GetNameStringByValue(static_cast<int64>(EmitterScript->GetUsage())));
+						}
+					}
+					else
 					{
-						DumpByteCode(EmitterScript, SystemPathName, HashedPathName, EmitterName / UsageEnum->GetNameStringByValue(static_cast<int64>(EmitterScript->GetUsage())));
+						if (const UNiagaraScript* GPUComputeScript = Emitter->GetGPUComputeScript())
+						{
+							const FString UsageName = UsageEnum->GetNameStringByValue(static_cast<int64>(GPUComputeScript->GetUsage()));
+
+							DumpShaderScript(GPUComputeScript, SystemPathName, HashedPathName, EmitterName / UsageName);
+						}
 					}
 				}
 			}
@@ -271,9 +290,10 @@ void UNiagaraDumpByteCodeCommandlet::ProcessNiagaraScripts()
 			continue;
 		}
 
-		if (BatchAssets.Num() == BatchPackageCount)
+		if (BatchAssets.Num() >= BatchPackageCount)
 		{
 			ProcessBatch(BatchAssets);
+			BatchAssets.Reset();
 		}
 
 		BatchAssets.Add(AssetIt);
@@ -282,6 +302,7 @@ void UNiagaraDumpByteCodeCommandlet::ProcessNiagaraScripts()
 	if (!BatchAssets.IsEmpty())
 	{
 		ProcessBatch(BatchAssets);
+		BatchAssets.Reset();
 	}
 
 	// sort the data alphabetically (based on MetaData.FullName
@@ -294,40 +315,78 @@ void UNiagaraDumpByteCodeCommandlet::ProcessNiagaraScripts()
 	{
 		const FString MetaDataFileName = AuditOutputFolder / TEXT("NiagaraScripts.xml");
 		TUniquePtr<FArchive> FileArchive(IFileManager::Get().CreateDebugFileWriter(*MetaDataFileName));
-		TUniquePtr<FOutputDeviceArchiveWrapper> OutputStream(new FOutputDeviceArchiveWrapper(FileArchive.Get()));
-
-		OutputStream->Log(TEXT("<?xml version='1.0' ?>"));
-		OutputStream->Log(TEXT("<Scripts>"));
-		for (const auto& MetaData : ScriptMetaData)
+		if (FileArchive)
 		{
-			OutputStream->Log(TEXT("\t<Script>"));
-			OutputStream->Logf(TEXT("\t\t<Hash>%s</Hash>"), *MetaData.SystemHash);
-			OutputStream->Logf(TEXT("\t\t<Name>%s</Name>"), *MetaData.FullName);
-			OutputStream->Logf(TEXT("\t\t<OpCount>%d</OpCount>"), MetaData.OpCount);
-			OutputStream->Logf(TEXT("\t\t<RegisterCount>%d</RegisterCount>"), MetaData.RegisterCount);
-			OutputStream->Logf(TEXT("\t\t<ConstantCount>%d</ConstantCount>"), MetaData.ConstantCount);
-			OutputStream->Logf(TEXT("\t\t<AttributeCount>%d</AttributeCount>"), MetaData.AttributeCount);
-			OutputStream->Log(TEXT("\t</Script>"));
+			TUniquePtr<FOutputDeviceArchiveWrapper> OutputStream(new FOutputDeviceArchiveWrapper(FileArchive.Get()));
+			if(OutputStream)
+			{
+				OutputStream->Log(TEXT("<?xml version='1.0' ?>"));
+				OutputStream->Log(TEXT("<Scripts>"));
+				for (const auto& MetaData : ScriptMetaData)
+				{
+					OutputStream->Log(TEXT("\t<Script>"));
+					OutputStream->Logf(TEXT("\t\t<Hash>%s</Hash>"), *MetaData.SystemHash);
+					OutputStream->Logf(TEXT("\t\t<Name>%s</Name>"), *MetaData.FullName);
+					OutputStream->Logf(TEXT("\t\t<OpCount>%d</OpCount>"), MetaData.OpCount);
+					OutputStream->Logf(TEXT("\t\t<RegisterCount>%d</RegisterCount>"), MetaData.RegisterCount);
+					OutputStream->Logf(TEXT("\t\t<ConstantCount>%d</ConstantCount>"), MetaData.ConstantCount);
+					OutputStream->Logf(TEXT("\t\t<AttributeCount>%d</AttributeCount>"), MetaData.AttributeCount);
+					OutputStream->Log(TEXT("\t</Script>"));
+				}
+				OutputStream->Log(TEXT("</Scripts>"));
+			}
 		}
-		OutputStream->Log(TEXT("</Scripts>"));
 	}
 
 	// create the csv for the meta data
 	{
 		const FString MetaDataFileName = AuditOutputFolder / TEXT("NiagaraScripts.csv");
 		TUniquePtr<FArchive> FileArchive(IFileManager::Get().CreateDebugFileWriter(*MetaDataFileName));
-		TUniquePtr<FOutputDeviceArchiveWrapper> OutputStream(new FOutputDeviceArchiveWrapper(FileArchive.Get()));
+		if (FileArchive)
+		{
+			TUniquePtr<FOutputDeviceArchiveWrapper> OutputStream(new FOutputDeviceArchiveWrapper(FileArchive.Get()));
+			if (OutputStream)
+			{
+				OutputStream->Log(TEXT("Hash, Name, OpCount, RegisteredCount, ConstantCount, AttributeCount"));
+				for (const auto& MetaData : ScriptMetaData)
+				{
+					OutputStream->Logf(TEXT("%s, %s, %d, %d, %d, %d"),
+						*MetaData.SystemHash,
+						*MetaData.FullName,
+						MetaData.OpCount,
+						MetaData.RegisterCount,
+						MetaData.ConstantCount,
+						MetaData.AttributeCount);
+				}
+			}
+		}
+	}
 
-		OutputStream->Log(TEXT("Hash, Name, OpCount, RegisteredCount, ConstantCount, AttributeCount"));
-		for (const auto& MetaData : ScriptMetaData)
+	ShaderMetaData.StableSort([](const FShaderMetaData& Lhs, const FShaderMetaData& Rhs)
+	{
+		return Lhs.FullName.Compare(Rhs.FullName, ESearchCase::IgnoreCase) < 0;
+	});
+
+	// create the csv for the meta data
+	{
+		const FString MetaDataFileName = AuditOutputFolder / TEXT("NiagaraShaders.csv");
+		TUniquePtr<FArchive> FileArchive(IFileManager::Get().CreateDebugFileWriter(*MetaDataFileName));
+		if (FileArchive)
 		{
-			OutputStream->Logf(TEXT("%s, %s, %d, %d, %d, %d"),
-				*MetaData.SystemHash,
-				*MetaData.FullName,
-				MetaData.OpCount,
-				MetaData.RegisterCount,
-				MetaData.ConstantCount,
-				MetaData.AttributeCount);
+			TUniquePtr<FOutputDeviceArchiveWrapper> OutputStream(new FOutputDeviceArchiveWrapper(FileArchive.Get()));
+			if (OutputStream)
+			{
+				OutputStream->Log(TEXT("Hash, Name, PermutationIndex, ContentHash, Size(Bytes)"));
+				for (const auto& MetaData : ShaderMetaData)
+				{
+					OutputStream->Logf(TEXT("%s, %s, %d, %s, %d"),
+						*MetaData.SystemHash,
+						*MetaData.FullName,
+						MetaData.PermutationIndex,
+						*MetaData.ShaderHash,
+						MetaData.ByteCount);
+				}
+			}
 		}
 	}
 
@@ -414,3 +473,56 @@ void UNiagaraDumpByteCodeCommandlet::DumpByteCode(const UNiagaraScript* Script,
 		FFileHelper::SaveStringToFile(Script->GetVMExecutableData().LastHlslTranslation, *FullFilePath);
 	}
 }
+
+void UNiagaraDumpByteCodeCommandlet::DumpShaderScript(const UNiagaraScript* GPUComputeScript, const FString& PathName, const FString& HashName, const FString& FilePath)
+{
+	if (!GPUComputeScript)
+	{
+		return;
+	}
+
+	FScriptMetaData& MetaData = ScriptMetaData.AddZeroed_GetRef();
+	MetaData.SystemHash = HashName;
+	MetaData.FullName = PathName / FilePath;
+	MetaData.RegisterCount = 0;
+	MetaData.OpCount = 0;
+	MetaData.ConstantCount = 0;
+	MetaData.AttributeCount = 0;
+
+	if (GPUComputeScript && DumpTranslatedHlsl)
+	{
+		const FString FullFilePath = AuditOutputFolder / HashName / FilePath + TEXT(".usf");
+
+		FFileHelper::SaveStringToFile(GPUComputeScript->GetVMExecutableData().LastHlslTranslationGPU, *FullFilePath);
+	}
+
+
+	if (const FNiagaraShaderScript* ShaderScript = GPUComputeScript->GetRenderThreadScript())
+	{
+		const FNiagaraShaderMap* ShaderMap = ShaderScript->GetGameThreadShaderMap();
+		const FNiagaraShaderMapContent* ShaderMapContent = ShaderMap->GetContent();
+
+		FShaderListReport ShaderListReport;
+		ShaderListReport.BindLambda([&HashName, FullName = PathName / FilePath, this](
+				const TCHAR* FriendlyName,
+				const TCHAR* OwnerName,
+				const TCHAR* VertexFactory,
+				const TCHAR* ShaderType,
+				uint32 Permutation,
+				uint32 Bytes,
+				const TCHAR* Hash,
+				void* Ptr)
+		{
+			FShaderMetaData& PermutationData = ShaderMetaData.AddDefaulted_GetRef();
+			PermutationData.SystemHash = HashName;
+			PermutationData.FullName = FullName;
+			PermutationData.PermutationIndex = Permutation;
+			PermutationData.VertexFactory = VertexFactory;
+			PermutationData.ShaderType = ShaderType;
+			PermutationData.ByteCount = Bytes;
+			PermutationData.ShaderHash = Hash;
+		});
+
+		ShaderMapContent->DumpShaderList(*ShaderMap, ShaderListReport);
+	}
+}
diff --git a/Engine/Plugins/FX/Niagara/Source/NiagaraEditor/Public/Commandlets/NiagaraDumpBytecodeCommandlet.h b/Engine/Plugins/FX/Niagara/Source/NiagaraEditor/Public/Commandlets/NiagaraDumpBytecodeCommandlet.h
--- a/Engine/Plugins/FX/Niagara/Source/NiagaraEditor/Public/Commandlets/NiagaraDumpBytecodeCommandlet.h
+++ b/Engine/Plugins/FX/Niagara/Source/NiagaraEditor/Public/Commandlets/NiagaraDumpBytecodeCommandlet.h
@@ -8,6 +8,7 @@
 #include "NiagaraDumpBytecodeCommandlet.generated.h"
 
 class UNiagaraScript;
+class FNiagaraShaderScript;
 
 UCLASS(config=Editor)
 class UNiagaraDumpByteCodeCommandlet : public UCommandlet
@@ -34,6 +35,9 @@ class UNiagaraDumpByteCodeCommandlet : public UCommandlet
 	/** Whether the translated hlsl should also be dumped */
 	bool DumpTranslatedHlsl = false;
 
+	/** Include GPU scripts */
+	bool DumpGpuScripts = false;
+
 	/** */
 	bool IncludeDeveloperFolder = false;
 
@@ -41,6 +45,7 @@ private:
 	void ProcessNiagaraScripts();
 	void ProcessBatch(TArray<FAssetData>& BatchAssets);
 	void DumpByteCode(const UNiagaraScript* Script, const FString& PathName, const FString& HashName, const FString& FilePath);
+	void DumpShaderScript(const UNiagaraScript* GPUComputeScript, const FString& PathName, const FString& HashName, const FString& FilePath);
 
 	struct FScriptMetaData
 	{
@@ -53,4 +58,16 @@ private:
 	};
 
 	TArray<FScriptMetaData> ScriptMetaData;
+
+	struct FShaderMetaData
+	{
+		FString SystemHash;
+		FString FullName;
+		int32 PermutationIndex;
+		FString VertexFactory;
+		FString ShaderType;
+		uint32 ByteCount;
+		FString ShaderHash;
+	};
+	TArray<FShaderMetaData> ShaderMetaData;
 };
\ No newline at end of file
