Fix for DumpNiagaraByteCode commandlet for shipping editor configuration
#rnx

#Synced-CL 50944596

diff --git a/Engine/Plugins/FX/Niagara/Source/NiagaraEditor/Private/Commandlets/NiagaraDumpBytecodeCommandlet.cpp b/Engine/Plugins/FX/Niagara/Source/NiagaraEditor/Private/Commandlets/NiagaraDumpBytecodeCommandlet.cpp
--- a/Engine/Plugins/FX/Niagara/Source/NiagaraEditor/Private/Commandlets/NiagaraDumpBytecodeCommandlet.cpp
+++ b/Engine/Plugins/FX/Niagara/Source/NiagaraEditor/Private/Commandlets/NiagaraDumpBytecodeCommandlet.cpp
@@ -496,7 +496,7 @@ void UNiagaraDumpByteCodeCommandlet::DumpShaderScript(const UNiagaraScript* GPUC
 		FFileHelper::SaveStringToFile(GPUComputeScript->GetVMExecutableData().LastHlslTranslationGPU, *FullFilePath);
 	}
 
-
+#if !UE_BUILD_SHIPPING
 	if (const FNiagaraShaderScript* ShaderScript = GPUComputeScript->GetRenderThreadScript())
 	{
 		const FNiagaraShaderMap* ShaderMap = ShaderScript->GetGameThreadShaderMap();
@@ -525,4 +525,5 @@ void UNiagaraDumpByteCodeCommandlet::DumpShaderScript(const UNiagaraScript* GPUC
 
 		ShaderMapContent->DumpShaderList(*ShaderMap, ShaderListReport);
 	}
+#endif // !UE_BUILD_SHIPPING
 }
