- Improve debug information when we have a null node

#rnx

#Synced-CL 50947237

diff --git a/Engine/Plugins/FX/Niagara/Source/NiagaraEditor/Private/NiagaraNode.cpp b/Engine/Plugins/FX/Niagara/Source/NiagaraEditor/Private/NiagaraNode.cpp
--- a/Engine/Plugins/FX/Niagara/Source/NiagaraEditor/Private/NiagaraNode.cpp
+++ b/Engine/Plugins/FX/Niagara/Source/NiagaraEditor/Private/NiagaraNode.cpp
@@ -1,23 +1,26 @@
 // Copyright Epic Games, Inc. All Rights Reserved.
 
 #include "NiagaraNode.h"
+#include "NiagaraConstants.h"
 #include "NiagaraCompileHashVisitor.h"
+#include "NiagaraEditorUtilities.h"
 #include "NiagaraGraph.h"
-#include "EdGraphSchema_Niagara.h"
 #include "NiagaraHlslTranslator.h"
+#include "NiagaraScriptSource.h"
+
+#include "EdGraphSchema_Niagara.h"
 #include "GraphEditAction.h"
 #include "Widgets/SNiagaraGraphNode.h"
 #include "Misc/SecureHash.h"
 #include "ToolMenuSection.h"
 #include "ToolMenu.h"
 #include "GraphEditorActions.h"
-#include "NiagaraConstants.h"
 #include "Serialization/PropertyLocalizationDataGathering.h"
-#include "NiagaraEditorUtilities.h"
 #include "ScopedTransaction.h"
 #include "Misc/Guid.h"
 #include "UObject/TextProperty.h"
 
+
 #include UE_INLINE_GENERATED_CPP_BY_NAME(NiagaraNode)
 
 #define LOCTEXT_NAMESPACE "NiagaraNode"
@@ -58,10 +61,32 @@ void UNiagaraNode::PostLoad()
 				if(LinkedPin == nullptr)
 				{
 					UNiagaraGraph* OwningGraph = GetNiagaraGraph();
-					FString AssetPath = OwningGraph ? OwningGraph->GetPathName() : "Undetermined";
+
+					FString ScriptUsageName = TEXT("Unknown Usage");
+					FString AssetPath = "Undetermined";
 					FString PinName = Pin->PinName.ToString();
-					FString NodeName = Pin->GetOwningNode() ? Pin->GetOwningNode()->GetNodeTitle(ENodeTitleType::FullTitle).ToString() : "Undetermined";
-					UE_LOG(LogNiagaraEditor, Log, TEXT("Connected pin of Pin %s of node %s in asset %s is nullptr!"), *PinName, *NodeName, *AssetPath);
+					FString NodeName = "Undetermined";
+
+					if (UNiagaraScriptSource* ScriptSource = GetSource())
+					{
+						if (UNiagaraScript* Script = ScriptSource->GetTypedOuter<UNiagaraScript>())
+						{
+							const UEnum* UsageEnum = StaticEnum<ENiagaraScriptUsage>();
+							ScriptUsageName = UsageEnum->GetDisplayNameTextByValue(int64(Script->GetUsage())).ToString();
+						}
+					}
+
+					if ( OwningGraph )
+					{
+						AssetPath = OwningGraph->GetPathName();
+					}
+
+					if (UEdGraphNode* OwningNode = Pin->GetOwningNode())
+					{
+						NodeName = FString::Printf(TEXT("%s:%s"), *GetNameSafe(OwningNode->GetClass()), *OwningNode->GetNodeTitle(ENodeTitleType::FullTitle).ToString());
+					}
+
+					UE_LOG(LogNiagaraEditor, Log, TEXT("Connected pin of Pin '%s' in '%s' of node '%s' in asset '%s' is nullptr!"), *PinName, *ScriptUsageName, *NodeName, *AssetPath);
 				}
 			}
 		}
