- Fix for potential nullptr access
- Unclear if this will fix a seen once creash or not due to lacking information, but better to be safe

#Synced-CL 51095106

diff --git a/Engine/Plugins/FX/Niagara/Source/Niagara/Private/NiagaraParameterCollectionResolver.cpp b/Engine/Plugins/FX/Niagara/Source/Niagara/Private/NiagaraParameterCollectionResolver.cpp
--- a/Engine/Plugins/FX/Niagara/Source/Niagara/Private/NiagaraParameterCollectionResolver.cpp
+++ b/Engine/Plugins/FX/Niagara/Source/Niagara/Private/NiagaraParameterCollectionResolver.cpp
@@ -59,12 +59,18 @@ void FNiagaraParameterCollectionResolver::UnbindCollections(TConstArrayView<TObj
 
 void FNiagaraParameterCollectionResolver::BindCollections(UNiagaraScript* Script, FNiagaraParameterStore& DestinationStore) const
 {
-	BindCollections(Script->GetCachedParameterCollectionReferences(), DestinationStore);
+	if (Script)
+	{
+		BindCollections(Script->GetCachedParameterCollectionReferences(), DestinationStore);
+	}
 }
 
 void FNiagaraParameterCollectionResolver::UnbindCollections(UNiagaraScript* Script, FNiagaraParameterStore& DestinationStore) const
 {
-	UnbindCollections(Script->GetCachedParameterCollectionReferences(), DestinationStore);
+	if (Script)
+	{
+		UnbindCollections(Script->GetCachedParameterCollectionReferences(), DestinationStore);
+	}
 }
 
 UNiagaraParameterCollectionInstance* FNiagaraParameterCollectionResolver::ResolveInstance(UNiagaraParameterCollection* Collection) const
