Avoid calling RegisterDynamicCustomizations when each module is loaded as it iterates over each class, and is not required for cooking

#Synced-CL 51078227

diff --git a/Engine/Plugins/FX/Niagara/Source/NiagaraEditor/Private/NiagaraEditorModule.cpp b/Engine/Plugins/FX/Niagara/Source/NiagaraEditor/Private/NiagaraEditorModule.cpp
--- a/Engine/Plugins/FX/Niagara/Source/NiagaraEditor/Private/NiagaraEditorModule.cpp
+++ b/Engine/Plugins/FX/Niagara/Source/NiagaraEditor/Private/NiagaraEditorModule.cpp
@@ -2738,7 +2738,11 @@ void FNiagaraEditorModule::OnModulesChanged(FName Module, EModuleChangeReason Re
 	//Refresh any dynamic customizations for potentially newly loaded types.
 	if (Reason == EModuleChangeReason::ModuleLoaded)
 	{
-		RegisterDynamicCustomizations();
+		// if we are running a cook commandlet we dont need to setup the NDC types or refresh them
+		if (!IsRunningCookCommandlet())
+		{
+			RegisterDynamicCustomizations();
+		}
 	}
 }
 
