[NiagaraEditor] Add some ensures to check for null to catch the null dereference we keep seeing in GetStackNodeGroups()   Should at least not crash (there...) and will help us narrow down the true culprit.

#Synced-CL 51114630

diff --git a/Engine/Plugins/FX/Niagara/Source/NiagaraEditor/Private/ViewModels/Stack/NiagaraStackGraphUtilities.cpp b/Engine/Plugins/FX/Niagara/Source/NiagaraEditor/Private/ViewModels/Stack/NiagaraStackGraphUtilities.cpp
--- a/Engine/Plugins/FX/Niagara/Source/NiagaraEditor/Private/ViewModels/Stack/NiagaraStackGraphUtilities.cpp
+++ b/Engine/Plugins/FX/Niagara/Source/NiagaraEditor/Private/ViewModels/Stack/NiagaraStackGraphUtilities.cpp
@@ -883,21 +883,59 @@ void FNiagaraStackGraphUtilities::GetStackNodeGroups(UNiagaraNode& StackNode, TA
 			GetOrderedModuleNodes(*OutputNode, ModuleNodes);
 			for (UNiagaraNodeFunctionCall* ModuleNode : ModuleNodes)
 			{
+				UNiagaraNode& LastEndNode = *OutStackNodeGroups.Last().EndNode;
+				UEdGraphPin* PreviousOutputPin = GetParameterMapOutputPin(LastEndNode);
+				if (!ensureMsgf(PreviousOutputPin, TEXT("Could not find a valid parameter map output pin for %s"), *LastEndNode.GetPathName()))
+				{
+					OutStackNodeGroups.Empty();
+					return; // unrecoverable, sorry.
+				}
 				FStackNodeGroup ModuleGroup;
-				UEdGraphPin* PreviousOutputPin = GetParameterMapOutputPin(*OutStackNodeGroups.Last().EndNode);
 				for (UEdGraphPin* PreviousOutputPinLinkedPin : PreviousOutputPin->LinkedTo)
 				{
-					ModuleGroup.StartNodes.Add(CastChecked<UNiagaraNode>(PreviousOutputPinLinkedPin->GetOwningNode()));
+					if (PreviousOutputPinLinkedPin)
+					{
+						ModuleGroup.StartNodes.Add(CastChecked<UNiagaraNode>(PreviousOutputPinLinkedPin->GetOwningNode()));
+					}
+					else
+					{
+						UE_LOG(LogNiagaraEditor, Log, TEXT("Invalid PreviousOutputPinLinkedPin in %s"), *ModuleNode->GetPathName());
+					}
+				}
+				// previous tests could make this malformed.
+				if (!ensureMsgf(!ModuleGroup.StartNodes.IsEmpty(), TEXT("ModuleGroup for %s has no StartNodes"), *ModuleNode->GetPathName()))
+				{
+					OutStackNodeGroups.Empty();
+					return;
 				}
 				ModuleGroup.EndNode = ModuleNode;
 				OutStackNodeGroups.Add(ModuleGroup);
 			}
 
+			UNiagaraNode& LastEndNode = *OutStackNodeGroups.Last().EndNode;
+			UEdGraphPin* PreviousOutputPin = GetParameterMapOutputPin(LastEndNode);
+			if (!ensureMsgf(PreviousOutputPin, TEXT("Could not find a valid parameter map output pin for %s"), *LastEndNode.GetPathName()))
+			{
+				OutStackNodeGroups.Empty();
+				return; // unrecoverable, sorry.
+			}
 			FStackNodeGroup OutputGroup;
-			UEdGraphPin* PreviousOutputPin = GetParameterMapOutputPin(*OutStackNodeGroups.Last().EndNode);
 			for (UEdGraphPin* PreviousOutputPinLinkedPin : PreviousOutputPin->LinkedTo)
 			{
-				OutputGroup.StartNodes.Add(CastChecked<UNiagaraNode>(PreviousOutputPinLinkedPin->GetOwningNode()));
+				if (PreviousOutputPinLinkedPin)
+				{
+					OutputGroup.StartNodes.Add(CastChecked<UNiagaraNode>(PreviousOutputPinLinkedPin->GetOwningNode()));
+				}
+				else
+				{
+					UE_LOG(LogNiagaraEditor, Log, TEXT("Invalid PreviousOutputPinLinkedPin in %s"), *StackNode.GetPathName());
+				}
+			}
+			// previous tests could make this malformed.
+			if (!ensureMsgf(!OutputGroup.StartNodes.IsEmpty(), TEXT("OutputGroup for %s has no StartNodes"), *StackNode.GetPathName()))
+			{
+				OutStackNodeGroups.Empty();
+				return;
 			}
 			OutputGroup.EndNode = OutputNode;
 			OutStackNodeGroups.Add(OutputGroup);
