In NiagaraEditorModule avoid loading NDC Assets if we are running a CookCommandlet

#Synced-CL 50966172
#Synced-CL 50966172

diff --git a/Engine/Plugins/FX/Niagara/Source/NiagaraEditor/Private/NiagaraEditorModule.cpp b/Engine/Plugins/FX/Niagara/Source/NiagaraEditor/Private/NiagaraEditorModule.cpp
--- a/Engine/Plugins/FX/Niagara/Source/NiagaraEditor/Private/NiagaraEditorModule.cpp
+++ b/Engine/Plugins/FX/Niagara/Source/NiagaraEditor/Private/NiagaraEditorModule.cpp
@@ -2704,16 +2704,20 @@ void FNiagaraEditorModule::OnAssetRegistryLoadComplete()
 
 	check(AssetRegistry.IsLoadingAssets() == false);
 
-	IAssetTools& AssetTools = IAssetTools::Get();
-
-	//Ensure All Data Channel Assets are loaded and available for use in editor.
-	TArray<FAssetData> AllDataChannels;
-	AssetRegistry.GetAssetsByClass(UNiagaraDataChannelAsset::StaticClass()->GetClassPathName(), AllDataChannels);
-	for (const FAssetData& DataChannelAsset : AllDataChannels)
+	// NDC is only for when running in the Editor with UI
+	if (!IsRunningCookCommandlet())
 	{
-		if (AssetTools.IsAssetVisible(DataChannelAsset, true))
+		IAssetTools& AssetTools = IAssetTools::Get();
+
+		//Ensure All Data Channel Assets are loaded and available for use in editor.
+		TArray<FAssetData> AllDataChannels;
+		AssetRegistry.GetAssetsByClass(UNiagaraDataChannelAsset::StaticClass()->GetClassPathName(), AllDataChannels);
+		for (const FAssetData& DataChannelAsset : AllDataChannels)
 		{
-			DataChannelAsset.GetAsset();
+			if (AssetTools.IsAssetVisible(DataChannelAsset, true))
+			{
+				DataChannelAsset.GetAsset();
+			}
 		}
 	}
 }
