Add ensure and skip for null node in niagara graph

#Synced-CL 51167185

diff --git a/Engine/Plugins/FX/Niagara/Source/NiagaraEditor/Private/NiagaraGraph.cpp b/Engine/Plugins/FX/Niagara/Source/NiagaraEditor/Private/NiagaraGraph.cpp
--- a/Engine/Plugins/FX/Niagara/Source/NiagaraEditor/Private/NiagaraGraph.cpp
+++ b/Engine/Plugins/FX/Niagara/Source/NiagaraEditor/Private/NiagaraGraph.cpp
@@ -676,7 +676,10 @@ void UNiagaraGraph::PostLoad()
 
 		for (UEdGraphNode* Node : Nodes)
 		{
-			Node->ConditionalPostLoad();
+			if (ensureMsgf(Node, TEXT("Unexpected null node in niagara graph %s"), *GetPathName()))
+			{
+				Node->ConditionalPostLoad();
+			}
 		}
 		
 		for (UEdGraphNode* Node : Nodes)
