Change NDC dynamic structs to process only loaded modules rather than all structs.

#Synced-CL 50966172

diff --git a/Engine/Plugins/FX/Niagara/Source/NiagaraEditor/Private/NiagaraEditorModule.cpp b/Engine/Plugins/FX/Niagara/Source/NiagaraEditor/Private/NiagaraEditorModule.cpp
--- a/Engine/Plugins/FX/Niagara/Source/NiagaraEditor/Private/NiagaraEditorModule.cpp
+++ b/Engine/Plugins/FX/Niagara/Source/NiagaraEditor/Private/NiagaraEditorModule.cpp
@@ -1775,6 +1775,8 @@ void FNiagaraEditorModule::ShutdownModule()
 		LevelSequenceModule.OnNewActorTrackAdded().Remove(DefaultTrackHandle);
 	}
 
+	UnregisterDynamicCustomizations();
+
 	FNiagaraEditorStyle::Unregister();
 
 	UnregisterSettings();
@@ -2737,44 +2739,87 @@ void FNiagaraEditorModule::OnAssetsPreDelete(const TArray<UObject*>& Objects)
 	}
 }
 
+void FNiagaraEditorModule::ConditionalRegisterDynamicStruct(const UScriptStruct* Struct)
+{
+	//Register customizations for Niagara Data Channel Access Context types.
+	if (Struct->IsChildOf(FNDCAccessContextBase::StaticStruct()))
+	{
+		const FName StructName = Struct->GetFName();
+
+		//Ensure we don't register the struct again.
+		//If the module is unloaded / reloaded we'll still have it's correct customizer in place.
+		if (RegisteredDynamicStructs.Contains(StructName))
+		{
+			return;
+		}
+
+		FPropertyEditorModule& PropertyModule = FModuleManager::LoadModuleChecked<FPropertyEditorModule>("PropertyEditor");
+
+		PropertyModule.RegisterCustomPropertyTypeLayout(
+			StructName,
+			FOnGetPropertyTypeCustomizationInstance::CreateStatic(&FNDCAccessContextDetailsCustomization::MakeInstance));
+
+		RegisteredDynamicStructs.Add(StructName);
+	}
+}
+
 void FNiagaraEditorModule::OnModulesChanged(FName Module, EModuleChangeReason Reason)
-{	
+{
 	//Refresh any dynamic customizations for potentially newly loaded types.
 	if (Reason == EModuleChangeReason::ModuleLoaded)
 	{
-		// if we are running a cook commandlet we dont need to setup the NDC types or refresh them
+		// if we are running a cook commandlet we don't need to setup the NDC types or refresh them
 		if (!IsRunningCookCommandlet())
 		{
-			RegisterDynamicCustomizations();
+			if (FModuleManager::Get().DoesLoadedModuleHaveUObjects(Module))
+			{
+				FNameBuilder ScriptPackageName(FPackageName::GetModuleScriptPackageName(Module));
+				if (const UPackage* ScriptPackage = FindPackage(nullptr, ScriptPackageName.ToString()))
+				{
+					ForEachObjectWithPackage(ScriptPackage, [this](UObject* Object) -> bool
+					{
+						if (const UScriptStruct* Struct = Cast<const UScriptStruct>(Object))
+						{
+							ConditionalRegisterDynamicStruct(Struct);
+						}
+
+						return true;
+					});
+				}
+			}
 		}
 	}
 }
 
 void FNiagaraEditorModule::RegisterDynamicCustomizations()
 {
+	// if we are running a cook commandlet we don't need to setup the NDC types or refresh them
+	if (IsRunningCookCommandlet())
+	{
+		return;
+	}
+
 	//Register any dynamic/discovered class customizations.
-	FPropertyEditorModule& PropertyModule = FModuleManager::LoadModuleChecked<FPropertyEditorModule>("PropertyEditor");
-	
 	for (TObjectIterator<UScriptStruct> It; It; ++It)
 	{
-		UScriptStruct* Struct = *It;
-		
-		//Ensure we don't register the struct again.
-		//If the module is unloaded / reloaded we'll still have it's correct customizer in place.
-		if(RegisteredDynamicStructs.Contains(Struct->GetFName()))
-		{
-			continue;
-		}
+		ConditionalRegisterDynamicStruct(*It);
+	}
+}
 
-		//Register customizations for Niagara Data Channel Access Context types.
-		if (Struct->IsChildOf(FNDCAccessContextBase::StaticStruct()))
+void FNiagaraEditorModule::UnregisterDynamicCustomizations()
+{
+	if (!RegisteredDynamicStructs.IsEmpty())
+	{
+		if (FModuleManager::Get().IsModuleLoaded("PropertyEditor"))
 		{
-			PropertyModule.RegisterCustomPropertyTypeLayout(
-				Struct->GetFName(),
-				FOnGetPropertyTypeCustomizationInstance::CreateStatic(&FNDCAccessContextDetailsCustomization::MakeInstance));
-
-			RegisteredDynamicStructs.Add(Struct->GetFName());
+			FPropertyEditorModule& PropertyModule = FModuleManager::LoadModuleChecked<FPropertyEditorModule>("PropertyEditor");
+			for (const FName& StructName : RegisteredDynamicStructs)
+			{
+				PropertyModule.UnregisterCustomPropertyTypeLayout(StructName);
+			}
 		}
+
+		RegisteredDynamicStructs.Reset();
 	}
 }
 
diff --git a/Engine/Plugins/FX/Niagara/Source/NiagaraEditor/Public/NiagaraEditorModule.h b/Engine/Plugins/FX/Niagara/Source/NiagaraEditor/Public/NiagaraEditorModule.h
--- a/Engine/Plugins/FX/Niagara/Source/NiagaraEditor/Public/NiagaraEditorModule.h
+++ b/Engine/Plugins/FX/Niagara/Source/NiagaraEditor/Public/NiagaraEditorModule.h
@@ -467,11 +467,14 @@ private:
 	void OnModulesChanged(FName Module, EModuleChangeReason Reason);
 
 	void RegisterDynamicCustomizations();
+	void UnregisterDynamicCustomizations();
 	
 	void RegisterMenuExtensions();
 	
 	void RegisterMenuExtension_MigrateToUserAssetTags();
 
+	void ConditionalRegisterDynamicStruct(const UScriptStruct* Struct);
+
 private:
 	TSharedPtr<FExtensibilityManager> MenuExtensibilityManager;
 	TSharedPtr<FExtensibilityManager> ToolBarExtensibilityManager;
