PR #14391: Fix spherical normals not affected by WPO

#Synced-CL 51164657

diff --git a/Engine/Plugins/FX/Niagara/Shaders/Private/NiagaraSpriteVertexFactory.ush b/Engine/Plugins/FX/Niagara/Shaders/Private/NiagaraSpriteVertexFactory.ush
--- a/Engine/Plugins/FX/Niagara/Shaders/Private/NiagaraSpriteVertexFactory.ush
+++ b/Engine/Plugins/FX/Niagara/Shaders/Private/NiagaraSpriteVertexFactory.ush
@@ -156,6 +156,9 @@ struct FVertexFactoryInterpolantsVSToPS
 #if USE_PARTICLE_POSITION
 	/** Cam-relative (translated) particle center and radius */
 	nointerpolation float4 ParticleTranslatedWorldPositionAndSize	: PARTICLE_POSITION;
+	#if GENERATE_SPHERICAL_PARTICLE_NORMALS
+		float3 DisplacedParticleCenterForNormal : PARTICLE_POSITION_FSN;
+	#endif
 #endif
 
 #if USE_PARTICLE_VELOCITY
@@ -336,10 +339,10 @@ FMaterialPixelParameters GetMaterialPixelParameters(FVertexFactoryInterpolantsVS
 
 	#if GENERATE_SPHERICAL_PARTICLE_NORMALS && USE_PARTICLE_POSITION
 		{	
-			// can be optimized
-			float4 ScreenPosition = SvPositionToResolvedScreenPosition(SvPosition);
+			// Use displaced position (SvPosition already includes WPO)
 			float3 TranslatedWorldPosition = SvPositionToResolvedTranslatedWorld(SvPosition);
-			Result.TangentToWorld = GetSphericalParticleNormal(TranslatedWorldPosition, Interpolants.ParticleTranslatedWorldPositionAndSize.xyz, Interpolants.ParticleTranslatedWorldPositionAndSize.w);
+			float3 SphereCenter = Interpolants.DisplacedParticleCenterForNormal;
+			Result.TangentToWorld = GetSphericalParticleNormal(TranslatedWorldPosition, SphereCenter, Interpolants.ParticleTranslatedWorldPositionAndSize.w);
 		}
 	#else
 		Result.TangentToWorld = AssembleTangentToWorld(TangentToWorld0.xyz, TangentToWorld2);
@@ -1194,6 +1197,10 @@ FVertexFactoryInterpolantsVSToPS VertexFactoryGetInterpolantsVSToPS(FVertexFacto
 
 #if USE_PARTICLE_POSITION
 	Interpolants.ParticleTranslatedWorldPositionAndSize = Intermediates.TranslatedWorldPositionAndSize;
+	#if GENERATE_SPHERICAL_PARTICLE_NORMALS
+		// Center for spherical normal must include WPO so shading follows displaced vertices
+		Interpolants.DisplacedParticleCenterForNormal = Intermediates.TranslatedWorldPositionAndSize.xyz + GetMaterialWorldPositionOffset(VertexParameters);
+	#endif
 #endif
 
 #if USE_PARTICLE_VELOCITY
diff --git a/Engine/Shaders/Private/ParticleSpriteVertexFactory.ush b/Engine/Shaders/Private/ParticleSpriteVertexFactory.ush
--- a/Engine/Shaders/Private/ParticleSpriteVertexFactory.ush
+++ b/Engine/Shaders/Private/ParticleSpriteVertexFactory.ush
@@ -69,6 +69,9 @@ struct FVertexFactoryInterpolantsVSToPS
 #if USE_PARTICLE_POSITION
 	/** Cam-relative (translated) particle center and radius */
 	float4 ParticleTranslatedWorldPositionAndSize	: PARTICLE_POSITION;
+	#if GENERATE_SPHERICAL_PARTICLE_NORMALS
+		float3 DisplacedParticleCenterForNormal : PARTICLE_POSITION_FSN;
+	#endif
 #endif
 
 #if USE_PARTICLE_VELOCITY
@@ -192,10 +195,10 @@ FMaterialPixelParameters GetMaterialPixelParameters(FVertexFactoryInterpolantsVS
 
 	#if GENERATE_SPHERICAL_PARTICLE_NORMALS && USE_PARTICLE_POSITION
 		{	
-			// can be optimized
-			float4 ScreenPosition = SvPositionToResolvedScreenPosition(SvPosition);
+			// Use displaced position (SvPosition already includes WPO)
 			float3 TranslatedWorldPosition = SvPositionToResolvedTranslatedWorld(SvPosition);
-			Result.TangentToWorld = GetSphericalParticleNormal(TranslatedWorldPosition, Interpolants.ParticleTranslatedWorldPositionAndSize.xyz, Interpolants.ParticleTranslatedWorldPositionAndSize.w);
+			float3 SphereCenter = Interpolants.DisplacedParticleCenterForNormal;
+			Result.TangentToWorld = GetSphericalParticleNormal(TranslatedWorldPosition, SphereCenter, Interpolants.ParticleTranslatedWorldPositionAndSize.w);
 		}
 	#else
 		Result.TangentToWorld = AssembleTangentToWorld(TangentToWorld0.xyz, TangentToWorld2);
@@ -703,6 +706,10 @@ FVertexFactoryInterpolantsVSToPS VertexFactoryGetInterpolantsVSToPS(FVertexFacto
 
 #if USE_PARTICLE_POSITION
 	Interpolants.ParticleTranslatedWorldPositionAndSize = Intermediates.TranslatedWorldPositionAndSize;
+	#if GENERATE_SPHERICAL_PARTICLE_NORMALS
+		// Center for spherical normal must include WPO so shading follows displaced vertices
+		Interpolants.DisplacedParticleCenterForNormal = Intermediates.TranslatedWorldPositionAndSize.xyz + GetMaterialWorldPositionOffset(VertexParameters);
+	#endif
 #endif
 
 #if USE_PARTICLE_VELOCITY
