move warning to log for unnamed grids in grid2d

#Synced-CL 51190646

diff --git a/Engine/Plugins/FX/Niagara/Source/Niagara/Private/NiagaraDataInterfaceGrid2DCollection.cpp b/Engine/Plugins/FX/Niagara/Source/Niagara/Private/NiagaraDataInterfaceGrid2DCollection.cpp
--- a/Engine/Plugins/FX/Niagara/Source/Niagara/Private/NiagaraDataInterfaceGrid2DCollection.cpp
+++ b/Engine/Plugins/FX/Niagara/Source/Niagara/Private/NiagaraDataInterfaceGrid2DCollection.cpp
@@ -1959,7 +1959,7 @@ bool UNiagaraDataInterfaceGrid2DCollection::InitPerInstanceData(void* PerInstanc
 	// give an error if there are any unnamed attributes
 	if (NumAttributes > 0 && NumNamedAttribChannelsFound > 0) //-V1051
 	{
-		UE_LOG(LogNiagara, Warning, TEXT("%s: Unnamed attributes should not be used with named.  This is a deprecated workflow"), *Proxy->SourceDIName.ToString());
+		UE_LOG(LogNiagara, Log, TEXT("%s: Unnamed attributes should not be used with named.  This is a deprecated workflow"), *Proxy->SourceDIName.ToString());
 	}
 
 	InstanceData->WorldBBoxSize = WorldBBoxSize;
