- Add optional reactivate to PostApplyToComponent
- When inside the BP window activating a component that isn't activated by default in the parent blueprint will go inactivate after the first frame due to PostToApply
- The new reactivation can be disabled with fx.Niagara.Rollback.AllowReactivateInPostApplyToComponent in case we run into issues with parameter stores and dirty data

#rnx

#Synced-CL 51214333

diff --git a/Engine/Plugins/FX/Niagara/Source/Niagara/Private/NiagaraComponent.cpp b/Engine/Plugins/FX/Niagara/Source/Niagara/Private/NiagaraComponent.cpp
--- a/Engine/Plugins/FX/Niagara/Source/Niagara/Private/NiagaraComponent.cpp
+++ b/Engine/Plugins/FX/Niagara/Source/Niagara/Private/NiagaraComponent.cpp
@@ -66,6 +66,14 @@ namespace NiagaraComponentPrivate
 		ECVF_Default
 	);
 
+	bool bAllowReactivateInPostApplyToComponent = true;
+	static FAutoConsoleVariableRef CVarRollback_AllowReactivateInPostApplyToComponent(
+		TEXT("fx.Niagara.Rollback.AllowReactivateInPostApplyToComponent"),
+		bAllowReactivateInPostApplyToComponent,
+		TEXT("When enabled activate & registered components will be reactivated at the end of PostApplyToComponent."),
+		ECVF_Default
+	);
+
 	static bool GParallelGDMEEnabled = true;
 	static FAutoConsoleVariableRef CVarNiagaraParallelGDME(
 		TEXT("fx.Niagara.ParallelGDME"),
@@ -2130,6 +2138,8 @@ void UNiagaraComponent::PostApplyToComponent()
 	Super::PostApplyToComponent();
 
 #if WITH_EDITOR
+	const bool bNeedsReactivate = IsRegistered() && IsActive() && NiagaraComponentPrivate::bAllowReactivateInPostApplyToComponent;
+
 	// Changes to the component can be broad ranging and include things like manipulating the instance parameters.  In particular if the
 	// data interfaces are changed out, we may need to rebind VM functions, so be sure to recreate the instance.
 	DestroyInstanceNotComponent();
@@ -2142,6 +2152,11 @@ void UNiagaraComponent::PostApplyToComponent()
 	}
 	ApplyOverridesToParameterStore();
 	FixDataInterfaceOuters();
+
+	if (bNeedsReactivate)
+	{
+		Activate();
+	}
 #endif
 }
 
