[Backout] - CL51167169
#rnx
Original CL Desc
-----------------------------------------------------------------
In NiagaraEditorModule avoid loading NDC Assets if we are running a CookCommandlet

diff --git a/Engine/Plugins/FX/Niagara/Source/NiagaraEditor/Private/NiagaraEditorModule.cpp b/Engine/Plugins/FX/Niagara/Source/NiagaraEditor/Private/NiagaraEditorModule.cpp
--- a/Engine/Plugins/FX/Niagara/Source/NiagaraEditor/Private/NiagaraEditorModule.cpp
+++ b/Engine/Plugins/FX/Niagara/Source/NiagaraEditor/Private/NiagaraEditorModule.cpp
@@ -2706,20 +2706,16 @@ void FNiagaraEditorModule::OnAssetRegistryLoadComplete()
 
 	check(AssetRegistry.IsLoadingAssets() == false);
 
-	// NDC is only for when running in the Editor with UI
-	if (!IsRunningCookCommandlet())
-	{
-		IAssetTools& AssetTools = IAssetTools::Get();
+	IAssetTools& AssetTools = IAssetTools::Get();
 
-		//Ensure All Data Channel Assets are loaded and available for use in editor.
-		TArray<FAssetData> AllDataChannels;
-		AssetRegistry.GetAssetsByClass(UNiagaraDataChannelAsset::StaticClass()->GetClassPathName(), AllDataChannels);
-		for (const FAssetData& DataChannelAsset : AllDataChannels)
+	//Ensure All Data Channel Assets are loaded and available for use in editor.
+	TArray<FAssetData> AllDataChannels;
+	AssetRegistry.GetAssetsByClass(UNiagaraDataChannelAsset::StaticClass()->GetClassPathName(), AllDataChannels);
+	for (const FAssetData& DataChannelAsset : AllDataChannels)
+	{
+		if (AssetTools.IsAssetVisible(DataChannelAsset, true))
 		{
-			if (AssetTools.IsAssetVisible(DataChannelAsset, true))
-			{
-				DataChannelAsset.GetAsset();
-			}
+			DataChannelAsset.GetAsset();
 		}
 	}
 }
