- Various analysis fixes

#Synced-CL 51379665

diff --git a/Engine/Plugins/FX/Niagara/Source/Niagara/Private/DataInterface/NiagaraDataInterfaceArrayDistributionInt.cpp b/Engine/Plugins/FX/Niagara/Source/Niagara/Private/DataInterface/NiagaraDataInterfaceArrayDistributionInt.cpp
--- a/Engine/Plugins/FX/Niagara/Source/Niagara/Private/DataInterface/NiagaraDataInterfaceArrayDistributionInt.cpp
+++ b/Engine/Plugins/FX/Niagara/Source/Niagara/Private/DataInterface/NiagaraDataInterfaceArrayDistributionInt.cpp
@@ -38,10 +38,10 @@ namespace NDIArrayDistributionIntPrivate
 		const int32 Offset = Index * EntrySize;
 		check(TableData.Num() >= Offset + EntrySize);
 		uint8* Data = TableData.GetData() + Offset;
-		reinterpret_cast<float*>(Data)[0] = Prob;
-		reinterpret_cast<int32*>(Data)[1] = Alias;
-		reinterpret_cast<int32*>(Data)[2] = Value;
-		reinterpret_cast<float*>(Data)[3] = Weight;
+		FMemory::Memcpy(Data +  0, &Prob, sizeof(float));
+		FMemory::Memcpy(Data +  4, &Alias, sizeof(int32));
+		FMemory::Memcpy(Data +  8, &Value, sizeof(int32));
+		FMemory::Memcpy(Data + 12, &Weight, sizeof(float));
 	}
 
 	void DecodeTableDirect(const TArray<uint8>& TableData, int32 Index, int32& OutValue, float& OutWeight)
@@ -50,8 +50,8 @@ namespace NDIArrayDistributionIntPrivate
 		const int32 Offset = Index * EntrySize;
 		check(TableData.Num() >= Offset + EntrySize);
 		const uint8* Data = TableData.GetData() + Offset;
-		OutValue = reinterpret_cast<const int32*>(Data)[2];
-		OutWeight = reinterpret_cast<const float*>(Data)[3];
+		FMemory::Memcpy(&OutValue, Data + 8, sizeof(int32));
+		FMemory::Memcpy(&OutWeight, Data + 12, sizeof(float));
 	}
 
 	void DecodeTableProbAlias(const TArray<uint8>& TableData, int32 Index, float& OutProb, int32& OutAlias)
@@ -60,8 +60,8 @@ namespace NDIArrayDistributionIntPrivate
 		const int32 Offset = Index * EntrySize;
 		check(TableData.Num() >= Offset + EntrySize);
 		const uint8* Data = TableData.GetData() + Offset;
-		OutProb = reinterpret_cast<const float*>(Data)[0];
-		OutAlias = reinterpret_cast<const int32*>(Data)[1];
+		FMemory::Memcpy(&OutProb, Data + 0, sizeof(float));
+		FMemory::Memcpy(&OutAlias, Data + 4, sizeof(int32));
 	}
 
 	void DecodeTableIndirect(const TArray<uint8>& TableData, int32 Index, float Probability, int32& OutValue, float& OutWeight)
@@ -71,8 +71,10 @@ namespace NDIArrayDistributionIntPrivate
 		check(TableData.Num() >= Offset + EntrySize);
 		const uint8* Data = TableData.GetData() + Offset;
 
-		const float EntryProb	= reinterpret_cast<const float*>(Data)[0];
-		const float EntryAlias	= reinterpret_cast<const int32*>(Data)[1];
+		float EntryProb;
+		int32 EntryAlias;
+		FMemory::Memcpy(&EntryProb, Data + 0, sizeof(float));
+		FMemory::Memcpy(&EntryAlias, Data + 4, sizeof(int32));
 		Index = Probability > EntryProb ? EntryAlias : Index;
 
 		DecodeTableDirect(TableData, Index, OutValue, OutWeight);
diff --git a/Engine/Plugins/FX/Niagara/Source/Niagara/Private/DataInterface/NiagaraDataInterfaceDataChannelCommon.cpp b/Engine/Plugins/FX/Niagara/Source/Niagara/Private/DataInterface/NiagaraDataInterfaceDataChannelCommon.cpp
--- a/Engine/Plugins/FX/Niagara/Source/Niagara/Private/DataInterface/NiagaraDataInterfaceDataChannelCommon.cpp
+++ b/Engine/Plugins/FX/Niagara/Source/Niagara/Private/DataInterface/NiagaraDataInterfaceDataChannelCommon.cpp
@@ -73,7 +73,7 @@ void FNDIDataChannel_FunctionToDataSetBinding::GenVMBindings(const FNiagaraVaria
 		}
 		else if (Property->IsA(FUInt16Property::StaticClass()))
 		{
-			VMRegisterBindings.Emplace(FuncHalfRegister++, DataSetHalfRegister == INDEX_NONE ? INDEX_NONE : DataSetIntRegister++, ENiagaraBaseTypes::Half);
+			VMRegisterBindings.Emplace(FuncHalfRegister++, DataSetHalfRegister == INDEX_NONE ? INDEX_NONE : DataSetHalfRegister++, ENiagaraBaseTypes::Half);
 		}
 		else if (Property->IsA(FIntProperty::StaticClass()))
 		{
@@ -169,17 +169,8 @@ FNDIDataChannel_FuncToDataSetBindingPtr FNDIDataChannelLayoutManager::GetLayoutI
 	//No valid existing layout so generate a new one.
 	FRWScopeLock WriteLock(FunctionToDataSetMapLock, SLT_Write);
 	FNDIDataChannel_FuncToDataSetBindingPtr FuncLayout = MakeShared<FNDIDataChannel_FunctionToDataSetBinding, ESPMode::ThreadSafe>(FunctionInfo, DataSetLayout, OutMissingParams);
-	if (FuncLayout.IsValid() == false)
-	{
-		FunctionToDataSetLayoutMap.Add(Key) = FuncLayout;
-	}
-#if DEBUG_NDI_DATACHANNEL
-	else
-	{
-		checkf(FuncLayout->DebugFunctionInfo.CheckHashConflict(FunctionInfo) == false, TEXT("Key conflict. Function Information does not match that already placed at this key."));
-		checkf(FuncLayout->DebugCompiledData.CheckHashConflict(DataSetLayout) == false, TEXT("Key conflict. DataSet Compiled Information does not match that already placed at this key."));
-	}
-#endif
+	FunctionToDataSetLayoutMap.Add(Key) = FuncLayout;
+
 	return FuncLayout;
 }
 
diff --git a/Engine/Plugins/FX/Niagara/Source/Niagara/Private/DataInterface/NiagaraDataInterfaceDataChannelRead.cpp b/Engine/Plugins/FX/Niagara/Source/Niagara/Private/DataInterface/NiagaraDataInterfaceDataChannelRead.cpp
--- a/Engine/Plugins/FX/Niagara/Source/Niagara/Private/DataInterface/NiagaraDataInterfaceDataChannelRead.cpp
+++ b/Engine/Plugins/FX/Niagara/Source/Niagara/Private/DataInterface/NiagaraDataInterfaceDataChannelRead.cpp
@@ -1775,7 +1775,7 @@ void UNiagaraDataInterfaceDataChannelRead::Consume(FVectorVMExternalFunctionCont
 
 		if (OutIndex.IsValid())
 		{
-			FMemory::Memset(OutIndex.Data.GetDest(), 0xFF, sizeof(int32) * Context.GetNumInstances());
+			FMemory::Memset(OutIndex.Data.GetDest(), -1, sizeof(int32) * Context.GetNumInstances());
 		}
 	}
 }
diff --git a/Engine/Plugins/FX/Niagara/Source/Niagara/Private/DataInterface/NiagaraDataInterfaceDataTable.cpp b/Engine/Plugins/FX/Niagara/Source/Niagara/Private/DataInterface/NiagaraDataInterfaceDataTable.cpp
--- a/Engine/Plugins/FX/Niagara/Source/Niagara/Private/DataInterface/NiagaraDataInterfaceDataTable.cpp
+++ b/Engine/Plugins/FX/Niagara/Source/Niagara/Private/DataInterface/NiagaraDataInterfaceDataTable.cpp
@@ -441,7 +441,7 @@ namespace NDIDataTableLocal
 			OutIsValid.SetAndAdvance(bValidRowIndex);
 			for (int32 iOutput = 0; iOutput < VariadicReadOffsets.Num(); ++iOutput)
 			{
-				const int32 VariableReadOffset = VariadicReadOffsets[iOutput] == INDEX_NONE ? InvalidRowReadOffset : RowReadOffset + VariadicReadOffsets[iOutput];
+				const uint32 VariableReadOffset = VariadicReadOffsets[iOutput] == INDEX_NONE ? InvalidRowReadOffset : RowReadOffset + VariadicReadOffsets[iOutput];
 				const int32* Value = reinterpret_cast<const int32*>(RawTableData + VariableReadOffset);
 				OutVariadics[iOutput].SetAndAdvance(*Value);
 			}
@@ -473,7 +473,7 @@ namespace NDIDataTableLocal
 			OutIsValid.SetAndAdvance(bValidRowIndex);
 			for (int32 iOutput = 0; iOutput < VariadicReadOffsets.Num(); ++iOutput)
 			{
-				const int32 VariableReadOffset = VariadicReadOffsets[iOutput] == INDEX_NONE ? InvalidRowReadOffset : RowReadOffset + VariadicReadOffsets[iOutput];
+				const uint32 VariableReadOffset = VariadicReadOffsets[iOutput] == INDEX_NONE ? InvalidRowReadOffset : RowReadOffset + VariadicReadOffsets[iOutput];
 				const int32* Value = reinterpret_cast<const int32*>(RawTableData + VariableReadOffset);
 				OutVariadics[iOutput].SetAndAdvance(*Value);
 			}
diff --git a/Engine/Plugins/FX/Niagara/Source/Niagara/Private/MovieScene/MovieSceneNiagaraSystemTrackTemplate.cpp b/Engine/Plugins/FX/Niagara/Source/Niagara/Private/MovieScene/MovieSceneNiagaraSystemTrackTemplate.cpp
--- a/Engine/Plugins/FX/Niagara/Source/Niagara/Private/MovieScene/MovieSceneNiagaraSystemTrackTemplate.cpp
+++ b/Engine/Plugins/FX/Niagara/Source/Niagara/Private/MovieScene/MovieSceneNiagaraSystemTrackTemplate.cpp
@@ -129,7 +129,7 @@ struct FNiagaraSystemUpdateDesiredAgeExecutionToken : IMovieSceneExecutionToken
 			UNiagaraComponent* NiagaraComponent = Cast<UNiagaraComponent>(Object);
 			if (!NiagaraComponent)
 			{
-				return;
+				continue;
 			}
 
 			{
diff --git a/Engine/Plugins/FX/Niagara/Source/Niagara/Private/MovieScene/Parameters/MovieSceneNiagaraBoolParameterSectionTemplate.cpp b/Engine/Plugins/FX/Niagara/Source/Niagara/Private/MovieScene/Parameters/MovieSceneNiagaraBoolParameterSectionTemplate.cpp
--- a/Engine/Plugins/FX/Niagara/Source/Niagara/Private/MovieScene/Parameters/MovieSceneNiagaraBoolParameterSectionTemplate.cpp
+++ b/Engine/Plugins/FX/Niagara/Source/Niagara/Private/MovieScene/Parameters/MovieSceneNiagaraBoolParameterSectionTemplate.cpp
@@ -17,8 +17,8 @@ FMovieSceneNiagaraBoolParameterSectionTemplate::FMovieSceneNiagaraBoolParameterS
 
 void FMovieSceneNiagaraBoolParameterSectionTemplate::GetAnimatedParameterValue(FFrameTime InTime, const FNiagaraVariableBase& InTargetParameter, const TArray<uint8>& InCurrentValueData, TArray<uint8>& OutAnimatedValueData) const
 {
-	FNiagaraBool const* CurrentValue = (FNiagaraBool const*)InCurrentValueData.GetData();
-	FNiagaraBool AnimatedNiagaraValue = *CurrentValue;
+	FNiagaraBool AnimatedNiagaraValue;
+	FMemory::Memcpy(&AnimatedNiagaraValue, InCurrentValueData.GetData(), sizeof(FNiagaraBool));
 
 	bool AnimatedValue;
 	if (BoolChannel.Evaluate(InTime, AnimatedValue))
diff --git a/Engine/Plugins/FX/Niagara/Source/Niagara/Private/MovieScene/Parameters/MovieSceneNiagaraBoolParameterTrack.cpp b/Engine/Plugins/FX/Niagara/Source/Niagara/Private/MovieScene/Parameters/MovieSceneNiagaraBoolParameterTrack.cpp
--- a/Engine/Plugins/FX/Niagara/Source/Niagara/Private/MovieScene/Parameters/MovieSceneNiagaraBoolParameterTrack.cpp
+++ b/Engine/Plugins/FX/Niagara/Source/Niagara/Private/MovieScene/Parameters/MovieSceneNiagaraBoolParameterTrack.cpp
@@ -25,7 +25,8 @@ void UMovieSceneNiagaraBoolParameterTrack::SetSectionChannelDefaults(UMovieScene
 	if (ensureMsgf(BoolSection != nullptr, TEXT("Section must be a bool section.")) && ensureMsgf(DefaultValueData.Num() == sizeof(FNiagaraBool), TEXT("DefaultValueData must be a FNiagaraBool.")))
 	{
 		FMovieSceneChannelProxy& BoolChannelProxy = BoolSection->GetChannelProxy();
-		FNiagaraBool DefaultValue = *((FNiagaraBool*)DefaultValueData.GetData());
+		FNiagaraBool DefaultValue;
+		FMemory::Memcpy(&DefaultValue, DefaultValueData.GetData(), sizeof(FNiagaraBool));
 		SetChannelDefault(BoolChannelProxy, BoolSection->GetChannel(), DefaultValue.GetValue());
 	}
 }
diff --git a/Engine/Plugins/FX/Niagara/Source/Niagara/Private/MovieScene/Parameters/MovieSceneNiagaraColorParameterSectionTemplate.cpp b/Engine/Plugins/FX/Niagara/Source/Niagara/Private/MovieScene/Parameters/MovieSceneNiagaraColorParameterSectionTemplate.cpp
--- a/Engine/Plugins/FX/Niagara/Source/Niagara/Private/MovieScene/Parameters/MovieSceneNiagaraColorParameterSectionTemplate.cpp
+++ b/Engine/Plugins/FX/Niagara/Source/Niagara/Private/MovieScene/Parameters/MovieSceneNiagaraColorParameterSectionTemplate.cpp
@@ -20,8 +20,8 @@ FMovieSceneNiagaraColorParameterSectionTemplate::FMovieSceneNiagaraColorParamete
 
 void FMovieSceneNiagaraColorParameterSectionTemplate::GetAnimatedParameterValue(FFrameTime InTime, const FNiagaraVariableBase& InTargetParameter, const TArray<uint8>& InCurrentValueData, TArray<uint8>& OutAnimatedValueData) const
 {
-	FLinearColor const* CurrentValue = (FLinearColor const*)InCurrentValueData.GetData();
-	FLinearColor AnimatedValue = *CurrentValue;
+	FLinearColor AnimatedValue;
+	FMemory::Memcpy(&AnimatedValue, InCurrentValueData.GetData(), sizeof(FLinearColor));
 
 	RedChannel.Evaluate(InTime, AnimatedValue.R);
 	GreenChannel.Evaluate(InTime, AnimatedValue.G);
diff --git a/Engine/Plugins/FX/Niagara/Source/Niagara/Private/MovieScene/Parameters/MovieSceneNiagaraColorParameterTrack.cpp b/Engine/Plugins/FX/Niagara/Source/Niagara/Private/MovieScene/Parameters/MovieSceneNiagaraColorParameterTrack.cpp
--- a/Engine/Plugins/FX/Niagara/Source/Niagara/Private/MovieScene/Parameters/MovieSceneNiagaraColorParameterTrack.cpp
+++ b/Engine/Plugins/FX/Niagara/Source/Niagara/Private/MovieScene/Parameters/MovieSceneNiagaraColorParameterTrack.cpp
@@ -23,7 +23,9 @@ void UMovieSceneNiagaraColorParameterTrack::SetSectionChannelDefaults(UMovieScen
 	UMovieSceneColorSection* ColorSection = Cast<UMovieSceneColorSection>(Section);
 	if (ensureMsgf(ColorSection != nullptr, TEXT("Section must be a color section.")) && ensureMsgf(DefaultValueData.Num() == sizeof(FLinearColor), TEXT("DefaultValueData must be a FLinearColor.")))
 	{
-		FLinearColor DefaultValue = *((FLinearColor*)DefaultValueData.GetData());
+		FLinearColor DefaultValue;
+		FMemory::Memcpy(&DefaultValue, DefaultValueData.GetData(), sizeof(FLinearColor));
+
 		FMovieSceneChannelProxy& ColorChannelProxy = ColorSection->GetChannelProxy();
 		SetChannelDefault(ColorChannelProxy, ColorSection->GetRedChannel(), DefaultValue.R);
 		SetChannelDefault(ColorChannelProxy, ColorSection->GetGreenChannel(), DefaultValue.G);
diff --git a/Engine/Plugins/FX/Niagara/Source/Niagara/Private/MovieScene/Parameters/MovieSceneNiagaraFloatParameterSectionTemplate.cpp b/Engine/Plugins/FX/Niagara/Source/Niagara/Private/MovieScene/Parameters/MovieSceneNiagaraFloatParameterSectionTemplate.cpp
--- a/Engine/Plugins/FX/Niagara/Source/Niagara/Private/MovieScene/Parameters/MovieSceneNiagaraFloatParameterSectionTemplate.cpp
+++ b/Engine/Plugins/FX/Niagara/Source/Niagara/Private/MovieScene/Parameters/MovieSceneNiagaraFloatParameterSectionTemplate.cpp
@@ -17,8 +17,8 @@ FMovieSceneNiagaraFloatParameterSectionTemplate::FMovieSceneNiagaraFloatParamete
 
 void FMovieSceneNiagaraFloatParameterSectionTemplate::GetAnimatedParameterValue(FFrameTime InTime, const FNiagaraVariableBase& InTargetParameter, const TArray<uint8>& InCurrentValueData, TArray<uint8>& OutAnimatedValueData) const
 {
-	FNiagaraFloat const* CurrentValue = (FNiagaraFloat const*)InCurrentValueData.GetData();
-	FNiagaraFloat AnimatedValue = *CurrentValue;
+	FNiagaraFloat AnimatedValue;
+	FMemory::Memcpy(&AnimatedValue, InCurrentValueData.GetData(), sizeof(FNiagaraFloat));
 
 	FloatChannel.Evaluate(InTime, AnimatedValue.Value);
 
diff --git a/Engine/Plugins/FX/Niagara/Source/Niagara/Private/MovieScene/Parameters/MovieSceneNiagaraFloatParameterTrack.cpp b/Engine/Plugins/FX/Niagara/Source/Niagara/Private/MovieScene/Parameters/MovieSceneNiagaraFloatParameterTrack.cpp
--- a/Engine/Plugins/FX/Niagara/Source/Niagara/Private/MovieScene/Parameters/MovieSceneNiagaraFloatParameterTrack.cpp
+++ b/Engine/Plugins/FX/Niagara/Source/Niagara/Private/MovieScene/Parameters/MovieSceneNiagaraFloatParameterTrack.cpp
@@ -25,7 +25,9 @@ void UMovieSceneNiagaraFloatParameterTrack::SetSectionChannelDefaults(UMovieScen
 	if (ensureMsgf(FloatSection != nullptr, TEXT("Section must be a float section.")) && ensureMsgf(DefaultValueData.Num() == sizeof(float), TEXT("DefaultValueData must be a float.")))
 	{
 		FMovieSceneChannelProxy& FloatChannelProxy = FloatSection->GetChannelProxy();
-		float DefaultValue = *((float*)DefaultValueData.GetData());
+		float DefaultValue;
+		FMemory::Memcpy(&DefaultValue, DefaultValueData.GetData(), sizeof(float));
+
 		SetChannelDefault(FloatChannelProxy, FloatSection->GetChannel(), DefaultValue);
 	}
 }
diff --git a/Engine/Plugins/FX/Niagara/Source/Niagara/Private/MovieScene/Parameters/MovieSceneNiagaraIntegerParameterSectionTemplate.cpp b/Engine/Plugins/FX/Niagara/Source/Niagara/Private/MovieScene/Parameters/MovieSceneNiagaraIntegerParameterSectionTemplate.cpp
--- a/Engine/Plugins/FX/Niagara/Source/Niagara/Private/MovieScene/Parameters/MovieSceneNiagaraIntegerParameterSectionTemplate.cpp
+++ b/Engine/Plugins/FX/Niagara/Source/Niagara/Private/MovieScene/Parameters/MovieSceneNiagaraIntegerParameterSectionTemplate.cpp
@@ -17,8 +17,8 @@ FMovieSceneNiagaraIntegerParameterSectionTemplate::FMovieSceneNiagaraIntegerPara
 
 void FMovieSceneNiagaraIntegerParameterSectionTemplate::GetAnimatedParameterValue(FFrameTime InTime, const FNiagaraVariableBase& InTargetParameter, const TArray<uint8>& InCurrentValueData, TArray<uint8>& OutAnimatedValueData) const
 {
-	FNiagaraInt32 const* CurrentValue = (FNiagaraInt32 const*)InCurrentValueData.GetData();
-	FNiagaraInt32 AnimatedValue = *CurrentValue;
+	FNiagaraInt32 AnimatedValue;
+	FMemory::Memcpy(&AnimatedValue, InCurrentValueData.GetData(), sizeof(FNiagaraInt32));
 
 	IntegerChannel.Evaluate(InTime, AnimatedValue.Value);
 
diff --git a/Engine/Plugins/FX/Niagara/Source/Niagara/Private/MovieScene/Parameters/MovieSceneNiagaraIntegerParameterTrack.cpp b/Engine/Plugins/FX/Niagara/Source/Niagara/Private/MovieScene/Parameters/MovieSceneNiagaraIntegerParameterTrack.cpp
--- a/Engine/Plugins/FX/Niagara/Source/Niagara/Private/MovieScene/Parameters/MovieSceneNiagaraIntegerParameterTrack.cpp
+++ b/Engine/Plugins/FX/Niagara/Source/Niagara/Private/MovieScene/Parameters/MovieSceneNiagaraIntegerParameterTrack.cpp
@@ -24,7 +24,9 @@ void UMovieSceneNiagaraIntegerParameterTrack::SetSectionChannelDefaults(UMovieSc
 	if (ensureMsgf(IntegerSection != nullptr, TEXT("Section must be an integer section.")) && ensureMsgf(DefaultValueData.Num() == sizeof(int32), TEXT("DefaultValueData must be an int32.")))
 	{
 		FMovieSceneChannelProxy& IntegerChannelProxy = IntegerSection->GetChannelProxy();
-		int32 DefaultValue = *((int32*)DefaultValueData.GetData());
+		int32 DefaultValue;
+		FMemory::Memcpy(&DefaultValue, DefaultValueData.GetData(), sizeof(int32));
+
 		SetChannelDefault(IntegerChannelProxy, IntegerSection->GetChannel(), DefaultValue);
 	}
 }
diff --git a/Engine/Plugins/FX/Niagara/Source/Niagara/Private/MovieScene/Parameters/MovieSceneNiagaraVectorParameterSectionTemplate.cpp b/Engine/Plugins/FX/Niagara/Source/Niagara/Private/MovieScene/Parameters/MovieSceneNiagaraVectorParameterSectionTemplate.cpp
--- a/Engine/Plugins/FX/Niagara/Source/Niagara/Private/MovieScene/Parameters/MovieSceneNiagaraVectorParameterSectionTemplate.cpp
+++ b/Engine/Plugins/FX/Niagara/Source/Niagara/Private/MovieScene/Parameters/MovieSceneNiagaraVectorParameterSectionTemplate.cpp
@@ -23,8 +23,8 @@ void FMovieSceneNiagaraVectorParameterSectionTemplate::GetAnimatedParameterValue
 {
 	if (ChannelsUsed == 2)
 	{
-		FVector2f const* CurrentValue = (FVector2f const*)InCurrentValueData.GetData();
-		FVector2f AnimatedValue = *CurrentValue;
+		FVector2f AnimatedValue;
+		FMemory::Memcpy(&AnimatedValue, InCurrentValueData.GetData(), sizeof(FVector2f));
 
 		VectorChannels[0].Evaluate(InTime, AnimatedValue.X);
 		VectorChannels[1].Evaluate(InTime, AnimatedValue.Y);
@@ -34,8 +34,8 @@ void FMovieSceneNiagaraVectorParameterSectionTemplate::GetAnimatedParameterValue
 	}
 	else if (ChannelsUsed == 3)
 	{
-		FVector3f const* CurrentValue = (FVector3f const*)InCurrentValueData.GetData();
-		FVector3f AnimatedValue = *CurrentValue;
+		FVector3f AnimatedValue;
+		FMemory::Memcpy(&AnimatedValue, InCurrentValueData.GetData(), sizeof(FVector3f));
 
 		VectorChannels[0].Evaluate(InTime, AnimatedValue.X);
 		VectorChannels[1].Evaluate(InTime, AnimatedValue.Y);
@@ -57,8 +57,8 @@ void FMovieSceneNiagaraVectorParameterSectionTemplate::GetAnimatedParameterValue
 	}
 	else if (ChannelsUsed == 4)
 	{
-		FVector4f const* CurrentValue = (FVector4f const*)InCurrentValueData.GetData();
-		FVector4f AnimatedValue = *CurrentValue;
+		FVector4f AnimatedValue;
+		FMemory::Memcpy(&AnimatedValue, InCurrentValueData.GetData(), sizeof(FVector4f));
 
 		VectorChannels[0].Evaluate(InTime, AnimatedValue.X);
 		VectorChannels[1].Evaluate(InTime, AnimatedValue.Y);
diff --git a/Engine/Plugins/FX/Niagara/Source/Niagara/Private/MovieScene/Parameters/MovieSceneNiagaraVectorParameterTrack.cpp b/Engine/Plugins/FX/Niagara/Source/Niagara/Private/MovieScene/Parameters/MovieSceneNiagaraVectorParameterTrack.cpp
--- a/Engine/Plugins/FX/Niagara/Source/Niagara/Private/MovieScene/Parameters/MovieSceneNiagaraVectorParameterTrack.cpp
+++ b/Engine/Plugins/FX/Niagara/Source/Niagara/Private/MovieScene/Parameters/MovieSceneNiagaraVectorParameterTrack.cpp
@@ -29,20 +29,26 @@ void UMovieSceneNiagaraVectorParameterTrack::SetSectionChannelDefaults(UMovieSce
 
 		if (VectorSection->GetChannelsUsed() == 2 && ensureMsgf(DefaultValueData.Num() == sizeof(FVector2f), TEXT("DefaultValueData must be a FVector2f when channels used is 2")))
 		{
-			FVector2f DefaultValue = *((FVector2f*)DefaultValueData.GetData());
+			FVector2f DefaultValue;
+			FMemory::Memcpy(&DefaultValue, DefaultValueData.GetData(), sizeof(FVector2f));
+
 			SetChannelDefault(VectorChannelProxy, VectorSection->GetChannel(0), DefaultValue.X);
 			SetChannelDefault(VectorChannelProxy, VectorSection->GetChannel(1), DefaultValue.Y);
 		}
 		else if (VectorSection->GetChannelsUsed() == 3 && ensureMsgf(DefaultValueData.Num() == sizeof(FVector3f), TEXT("DefaultValueData must be a FVector3f when channels used is 3")))
 		{
-			FVector3f DefaultValue = *((FVector3f*)DefaultValueData.GetData());
+			FVector3f DefaultValue;
+			FMemory::Memcpy(&DefaultValue, DefaultValueData.GetData(), sizeof(FVector3f));
+
 			SetChannelDefault(VectorChannelProxy, VectorSection->GetChannel(0), DefaultValue.X);
 			SetChannelDefault(VectorChannelProxy, VectorSection->GetChannel(1), DefaultValue.Y);
 			SetChannelDefault(VectorChannelProxy, VectorSection->GetChannel(2), DefaultValue.Z);
 		}
 		else if (VectorSection->GetChannelsUsed() == 4 && ensureMsgf(DefaultValueData.Num() == sizeof(FVector4f), TEXT("DefaultValueData must be a FVector4f when channels used is 4")))
 		{
-			FVector4f DefaultValue = *((FVector4f*)DefaultValueData.GetData());
+			FVector4f DefaultValue;
+			FMemory::Memcpy(&DefaultValue, DefaultValueData.GetData(), sizeof(FVector4f));
+
 			SetChannelDefault(VectorChannelProxy, VectorSection->GetChannel(0), DefaultValue.X);
 			SetChannelDefault(VectorChannelProxy, VectorSection->GetChannel(1), DefaultValue.Y);
 			SetChannelDefault(VectorChannelProxy, VectorSection->GetChannel(2), DefaultValue.Z);
diff --git a/Engine/Plugins/FX/Niagara/Source/Niagara/Private/NiagaraActor.cpp b/Engine/Plugins/FX/Niagara/Source/Niagara/Private/NiagaraActor.cpp
--- a/Engine/Plugins/FX/Niagara/Source/Niagara/Private/NiagaraActor.cpp
+++ b/Engine/Plugins/FX/Niagara/Source/Niagara/Private/NiagaraActor.cpp
@@ -110,9 +110,12 @@ bool ANiagaraActor::GetReferencedContentObjects(TArray<UObject*>& Objects) const
 {
 	Super::GetReferencedContentObjects(Objects);
 
-	if (UNiagaraSystem* System = NiagaraComponent->GetAsset())
+	if (NiagaraComponent)
 	{
-		Objects.Add(System);
+		if (UNiagaraSystem* System = NiagaraComponent->GetAsset())
+		{
+			Objects.Add(System);
+		}
 	}
 
 	return true;
diff --git a/Engine/Plugins/FX/Niagara/Source/Niagara/Private/NiagaraDataChannel_Islands.cpp b/Engine/Plugins/FX/Niagara/Source/Niagara/Private/NiagaraDataChannel_Islands.cpp
--- a/Engine/Plugins/FX/Niagara/Source/Niagara/Private/NiagaraDataChannel_Islands.cpp
+++ b/Engine/Plugins/FX/Niagara/Source/Niagara/Private/NiagaraDataChannel_Islands.cpp
@@ -21,7 +21,7 @@
 namespace NDCIslands_Local
 {
 	int32 AllowAsyncLoad = 1;
-	static FAutoConsoleVariableRef CVarAllowAsyncLoad(TEXT("fx.Niagara.DataChannels.AllowAsyncLoad"), AllowAsyncLoad, TEXT("True if we should attempt to load systems etc asynchronosly."), ECVF_Default);
+	static FAutoConsoleVariableRef CVarAllowAsyncLoad(TEXT("fx.Niagara.DataChannels.AllowAsyncLoad"), AllowAsyncLoad, TEXT("True if we should attempt to load systems etc asynchronously."), ECVF_Default);
 
 	int32 BlockAsyncLoadOnUse = 1;
 	static FAutoConsoleVariableRef CVarBlockAsyncLoadOnUse(TEXT("fx.Niagara.DataChannels.BlockAsyncLoadOnUse"), BlockAsyncLoadOnUse, TEXT("True if we should block on any pending async loads when those assets are used."), ECVF_Default);
@@ -220,10 +220,6 @@ FNDCIsland* UNiagaraDataChannelHandler_Islands::FindOrCreateIsland(const FNiagar
 		}
 	}
 
-	UWorld* World = GetWorld();
-	check(World);
-	double WorldTime = World->GetTimeSeconds();
-
 	//Find the first island that could contain this point.
 	//For now we do a linear search of active islands.
 	//Assuming that overall active island count will be low.
@@ -256,7 +252,7 @@ FNDCIsland* UNiagaraDataChannelHandler_Islands::FindOrCreateIsland(const FNiagar
 		}
 	}
 
-	//Failing that, get a init a new island from the pool.
+	//Failing that, get and init a new island from the pool.
 	FNDCIsland* IslandToUse = nullptr;
 	int32 NewIslandIndex = ActivateNewIsland(Location);
 	if (ensure(NewIslandIndex != INDEX_NONE))
@@ -363,7 +359,7 @@ void FNDCIsland::OnAcquired(FVector Location)
 		if(Sys)
 		{
 			FFXSystemSpawnParameters SpawnParams;
-			SpawnParams.bAutoActivate = false;//We must activate AFTER adding this component to our handle system array.
+			SpawnParams.bAutoActivate = false;//We must activate AFTER adding this component to our handler system array.
 			SpawnParams.bAutoDestroy = false;
 			SpawnParams.Location = Bounds.Origin;
 			SpawnParams.PoolingMethod = EPSCPoolMethod::ManualRelease;
@@ -390,29 +386,43 @@ void FNDCIsland::OnReleased()
 	Data = nullptr;
 	for (UNiagaraComponent* Comp : SpawnedComponents)
 	{
-		Comp->ReleaseToPool();
+		if (::IsValid(Comp))
+		{
+			Comp->ReleaseToPool();
+		}
 	}
 	SpawnedComponents.Reset();
 }
 
 void FNDCIsland::BeginFrame()
 {
-	Data->BeginFrame(Owner);
+	if (Data)
+	{
+		Data->BeginFrame(Owner);
+	}
 }
 
 void FNDCIsland::EndFrame()
 {
-	Data->EndFrame(Owner);
+	if (Data)
+	{
+		Data->EndFrame(Owner);
+	}
 }
 
 void FNDCIsland::Tick(const ETickingGroup& TickGroup)
 {
+	if (Data == nullptr)
+	{
+		return;
+	}
+
 	int32 AddedData = Data->ConsumePublishRequests(Owner, TickGroup);
 	if (AddedData > 0)
 	{
 		for (UNiagaraComponent* Comp : SpawnedComponents)
 		{
-			if (Comp->IsComplete())
+			if (::IsValid(Comp) && Comp->IsComplete())
 			{
 				Comp->Activate();
 			}
@@ -447,9 +457,8 @@ bool FNDCIsland::IsHandlerSystem(USceneComponent* Component)
 
 bool FNDCIsland::IsBeingUsed()const
 {
-	bool bSystemsComplete = SpawnedComponents.ContainsByPredicate([](const TObjectPtr<UNiagaraComponent>& Comp) { return !IsValid(Comp) || Comp->IsComplete(); });
-
-	return Data.IsValid() && (!Data.IsUnique() || !bSystemsComplete);
+	const bool bSystemsIncomplete = SpawnedComponents.ContainsByPredicate([](const TObjectPtr<UNiagaraComponent>& Comp) { return IsValid(Comp) && !Comp->IsComplete(); });
+	return Data.IsValid() && (!Data.IsUnique() || bSystemsIncomplete);
 }
 
 bool FNDCIsland::TryGrow(FVector Point, FVector PerElementExtents, FVector MaxIslandExtents)
@@ -481,7 +490,7 @@ bool FNDCIsland::TryGrow(FVector Point, FVector PerElementExtents, FVector MaxIs
 	Bounds = LocalBounds;
 	Bounds.Origin += OldOrigin;
 
-	//If new bounds are valid, grow our bounds and update the 
+	//If new bounds are valid, grow our bounds and update the system bounds
 	FBox HandlerSystemBounds(-Bounds.BoxExtent, Bounds.BoxExtent);
 	for (auto It = SpawnedComponents.CreateIterator(); It; ++It)
 	{
@@ -502,5 +511,8 @@ bool FNDCIsland::TryGrow(FVector Point, FVector PerElementExtents, FVector MaxIs
 
 void FNDCIsland::DebugDrawBounds()
 {
-	DrawDebugBox(Owner->GetWorld(), Bounds.Origin, Bounds.BoxExtent, FQuat::Identity, FColor::Red);
+	if (Owner && Owner->GetWorld())
+	{
+		DrawDebugBox(Owner->GetWorld(), Bounds.Origin, Bounds.BoxExtent, FQuat::Identity, FColor::Red);
+	}
 }
diff --git a/Engine/Plugins/FX/Niagara/Source/Niagara/Private/NiagaraDataChannel_Map.cpp b/Engine/Plugins/FX/Niagara/Source/Niagara/Private/NiagaraDataChannel_Map.cpp
--- a/Engine/Plugins/FX/Niagara/Source/Niagara/Private/NiagaraDataChannel_Map.cpp
+++ b/Engine/Plugins/FX/Niagara/Source/Niagara/Private/NiagaraDataChannel_Map.cpp
@@ -6,7 +6,6 @@
 #include "NiagaraDataChannelCommon.h"
 #include "NiagaraWorldManager.h"
 #include "NiagaraComponent.h"
-#include "NiagaraDataChannelCommon.h"
 #include "NiagaraDataChannelData.h"
 #include "NiagaraFunctionLibrary.h"
 
@@ -75,7 +74,7 @@ void UNiagaraDataChannelHandler_MapBase::BeginFrame(float DeltaTime)
  		{
  			//This entry is done with so reset it and return to the free entries.
  
- 			//Fist clear our it's spawned components.
+ 			// First clear our it's spawned components.
  			for(TObjectPtr<UNiagaraComponent> SpawnedSystem : Entry.Get().SpawnedComponents)
  			{
  				if(SpawnedSystem)
@@ -222,7 +221,7 @@ FNDCMapEntry& UNiagaraDataChannelHandler_MapBase::FindOrAddEntry(FNDCAccessConte
 		{					
 			Ret = &EntryPool[*ActiveEntry];
 		}
-		else if (!TypedContext.bIsAutoLinkingSystem)//Don't create new entries for auto linking niagra systems.
+		else if (!TypedContext.bIsAutoLinkingSystem)//Don't create new entries for auto linking Niagara systems.
 		{
 			int32 NewEntryIdx = INDEX_NONE;
 			if (FreeEntries.Num() > 0)
@@ -318,6 +317,7 @@ void FNDCMapEntryBase::Cleanup()
 
 bool FNDCMapEntryBase::BeginFrame(float DeltaTime, FNiagaraWorldManager* OwningWorld, const FNDCMapKey& Key)
 {
+	check(OwningWorld && OwningWorld->GetWorld());
 	LastUsedTime = OwningWorld->GetWorld()->GetRealTimeSeconds();
 	bool bInUse = false;
 	if (Data)
@@ -359,7 +359,7 @@ void FNDCMapEntryBase::Tick(float DeltaTime, ETickingGroup TickGroup, FNiagaraWo
 	}
 }
 
-UNiagaraComponent* FNDCMapEntryBase::SpawnSystem(FNDCAccessContext_MapBase& AccessContext, UObject* System, USceneComponent* AttacheParent, FVector Location, FVector BoundExtents, bool bActivate)
+UNiagaraComponent* FNDCMapEntryBase::SpawnSystem(FNDCAccessContext_MapBase& AccessContext, UObject* System, USceneComponent* AttachParent, FVector Location, FVector BoundExtents, bool bActivate)
 {
 	UNiagaraComponent* NewComp = nullptr;
 	auto SpawnSys = [&](UNiagaraSystem* Sys)
@@ -374,9 +374,9 @@ UNiagaraComponent* FNDCMapEntryBase::SpawnSystem(FNDCAccessContext_MapBase& Acce
 			SpawnParams.SystemTemplate = Sys;
 			SpawnParams.WorldContextObject = Owner->GetWorld();
 
-			if (AttacheParent)
+			if (AttachParent)
 			{
-				SpawnParams.AttachToComponent = AttacheParent;
+				SpawnParams.AttachToComponent = AttachParent;
 				SpawnParams.LocationType = EAttachLocation::SnapToTarget;
 				NewComp = UNiagaraFunctionLibrary::SpawnSystemAttachedWithParams(SpawnParams);
 			}
@@ -407,9 +407,9 @@ UNiagaraComponent* FNDCMapEntryBase::SpawnSystem(FNDCAccessContext_MapBase& Acce
 	{
 		SpawnSys(NiagaraSystem);
 	}
-	else if(UNiagaraSystemCollection* SystemCollecton = Cast<UNiagaraSystemCollection>(System))
+	else if(UNiagaraSystemCollection* SystemCollection = Cast<UNiagaraSystemCollection>(System))
 	{
-		for(UNiagaraSystem* Sys : SystemCollecton->GetSystems())
+		for(UNiagaraSystem* Sys : SystemCollection->GetSystems())
 		{
 			SpawnSys(Sys);
 		}
