Add scopes around some NiagaraEditorModule areas when dealing with the AR

#Synced-CL 51322414

diff --git a/Engine/Plugins/FX/Niagara/Source/NiagaraEditor/Private/NiagaraEditorModule.cpp b/Engine/Plugins/FX/Niagara/Source/NiagaraEditor/Private/NiagaraEditorModule.cpp
--- a/Engine/Plugins/FX/Niagara/Source/NiagaraEditor/Private/NiagaraEditorModule.cpp
+++ b/Engine/Plugins/FX/Niagara/Source/NiagaraEditor/Private/NiagaraEditorModule.cpp
@@ -1214,11 +1214,14 @@ void FNiagaraEditorModule::StartupModule()
 		FAssetRegistryModule& AssetRegistryModule = FModuleManager::LoadModuleChecked<FAssetRegistryModule>("AssetRegistry");
 		AssetRegistryModule.Get().OnFilesLoaded().AddLambda([this]()
 		{
+			TRACE_CPUPROFILER_EVENT_SCOPE(FNiagaraEditorModule::OnFilesLoaded);
 			ParameterCollectionAssetCache.RefreshCache(true /*bAllowLoading*/);
 			ParameterDefinitionsAssetCache.RefreshCache(true /*bAllowLoading*/);
 		});
 		AssetRegistryModule.Get().OnAssetsAdded().AddLambda([this](TConstArrayView<FAssetData> InAssets)
 		{
+			TRACE_CPUPROFILER_EVENT_SCOPE(FNiagaraEditorModule::OnAssetsAdded);
+
 			bool FoundParameterCollection = false;
 			bool FoundParameterDefinitions = false;
 			for (const FAssetData& Asset : InAssets)
@@ -2699,6 +2702,7 @@ void FNiagaraEditorModule::RegisterDefaultStackObjectCustomizations()
 
 void FNiagaraEditorModule::OnAssetRegistryLoadComplete()
 {
+	TRACE_CPUPROFILER_EVENT_SCOPE(FNiagaraEditorModule::OnAssetRegistryLoadComplete);
 	FAssetRegistryModule& AssetRegistryModule = FModuleManager::LoadModuleChecked<FAssetRegistryModule>("AssetRegistry");
 	IAssetRegistry& AssetRegistry = AssetRegistryModule.Get();
 
