- Make sure solo mode is correctly set when bAwaitingActivationDueToNotReady is true

#Synced-CL 51363360

diff --git a/Engine/Plugins/FX/Niagara/Source/Niagara/Private/NiagaraComponent.cpp b/Engine/Plugins/FX/Niagara/Source/Niagara/Private/NiagaraComponent.cpp
--- a/Engine/Plugins/FX/Niagara/Source/Niagara/Private/NiagaraComponent.cpp
+++ b/Engine/Plugins/FX/Niagara/Source/Niagara/Private/NiagaraComponent.cpp
@@ -918,6 +918,7 @@ void UNiagaraComponent::TickComponent(float DeltaSeconds, enum ELevelTick TickTy
 
 	if (bAwaitingActivationDueToNotReady)
 	{
+		UpdateInstanceSoloMode();
 		Activate(bActivateShouldResetWhenReady);
 		return;
 	}
@@ -1367,6 +1368,7 @@ void UNiagaraComponent::ActivateInternal(bool bReset /* = false */, bool bIsScal
 		bAwaitingActivationDueToNotReady = true;
 		bActivateShouldResetWhenReady = bReset;
 		ScopedComponentTickEnabled.bTickEnabled = true;
+		UpdateInstanceSoloMode();
 		return;
 	}
 
@@ -1502,6 +1504,7 @@ void UNiagaraComponent::ActivateInternal(bool bReset /* = false */, bool bIsScal
 			bAwaitingActivationDueToNotReady = true;
 			bActivateShouldResetWhenReady = bReset;
 			ScopedComponentTickEnabled.bTickEnabled = true;
+			UpdateInstanceSoloMode();
 			return;
 		}
 	}
@@ -4705,7 +4708,7 @@ void UNiagaraComponent::SetForceSolo(bool bInForceSolo)
 bool UNiagaraComponent::RequiresSoloMode() const
 {
 	//-OPT: If parameter collection overrides are only bound into the particle scripts we can avoid forcing solo mode
-	return bForceSolo || SimCache != nullptr || !FMath::IsNearlyEqual(CustomTimeDilation, 1.0f) || (Asset && Asset->HasFixedTickDelta()) || (ParameterCollectionOverrides.Num() > 0);
+	return bForceSolo || SimCache != nullptr || !FMath::IsNearlyEqual(CustomTimeDilation, 1.0f) || (Asset && Asset->HasFixedTickDelta()) || (ParameterCollectionOverrides.Num() > 0) || bAwaitingActivationDueToNotReady;
 }
 
 void UNiagaraComponent::UpdateInstanceSoloMode()
