Avoid loading NDC content for a Cook Commandlet

#Synced-CL 51322414

diff --git a/Engine/Plugins/FX/Niagara/Source/NiagaraEditor/Private/NiagaraEditorModule.cpp b/Engine/Plugins/FX/Niagara/Source/NiagaraEditor/Private/NiagaraEditorModule.cpp
--- a/Engine/Plugins/FX/Niagara/Source/NiagaraEditor/Private/NiagaraEditorModule.cpp
+++ b/Engine/Plugins/FX/Niagara/Source/NiagaraEditor/Private/NiagaraEditorModule.cpp
@@ -2709,17 +2709,20 @@ void FNiagaraEditorModule::OnAssetRegistryLoadComplete()
 	AssetRegistry.OnFilesLoaded().Remove(AssetRegistryOnLoadCompleteHandle);
 
 	check(AssetRegistry.IsLoadingAssets() == false);
-
-	IAssetTools& AssetTools = IAssetTools::Get();
-
-	//Ensure All Data Channel Assets are loaded and available for use in editor.
-	TArray<FAssetData> AllDataChannels;
-	AssetRegistry.GetAssetsByClass(UNiagaraDataChannelAsset::StaticClass()->GetClassPathName(), AllDataChannels);
-	for (const FAssetData& DataChannelAsset : AllDataChannels)
+	// NDC is only for when running in the Editor with UI
+	if (!IsRunningCookCommandlet())
 	{
-		if (AssetTools.IsAssetVisible(DataChannelAsset, true))
+		IAssetTools& AssetTools = IAssetTools::Get();
+
+		//Ensure All Data Channel Assets are loaded and available for use in editor.
+		TArray<FAssetData> AllDataChannels;
+		AssetRegistry.GetAssetsByClass(UNiagaraDataChannelAsset::StaticClass()->GetClassPathName(), AllDataChannels);
+		for (const FAssetData& DataChannelAsset : AllDataChannels)
 		{
-			DataChannelAsset.GetAsset();
+			if (AssetTools.IsAssetVisible(DataChannelAsset, true))
+			{
+				DataChannelAsset.GetAsset();
+			}
 		}
 	}
 }
