- Improve messaging when failing to find an appropriate format to allow tracking down the user

#Synced-CL 51414939

diff --git a/Engine/Plugins/FX/Niagara/Source/Niagara/Classes/NiagaraDataInterfaceGrid2DCollection.h b/Engine/Plugins/FX/Niagara/Source/Niagara/Classes/NiagaraDataInterfaceGrid2DCollection.h
--- a/Engine/Plugins/FX/Niagara/Source/Niagara/Classes/NiagaraDataInterfaceGrid2DCollection.h
+++ b/Engine/Plugins/FX/Niagara/Source/Niagara/Classes/NiagaraDataInterfaceGrid2DCollection.h
@@ -47,7 +47,7 @@ struct FGrid2DCollectionRWInstanceData_GameThread
 
 	int32 FindAttributeIndexByName(const FName& InName, int32 NumChannels);
 
-	bool UpdateTargetTexture(ENiagaraGpuBufferFormat BufferFormat);
+	bool UpdateTargetTexture(ENiagaraGpuBufferFormat BufferFormat, UObject* Owner);
 };
 
 struct FGrid2DCollectionRWInstanceData_RenderThread
diff --git a/Engine/Plugins/FX/Niagara/Source/Niagara/Classes/NiagaraDataInterfaceGrid3DCollection.h b/Engine/Plugins/FX/Niagara/Source/Niagara/Classes/NiagaraDataInterfaceGrid3DCollection.h
--- a/Engine/Plugins/FX/Niagara/Source/Niagara/Classes/NiagaraDataInterfaceGrid3DCollection.h
+++ b/Engine/Plugins/FX/Niagara/Source/Niagara/Classes/NiagaraDataInterfaceGrid3DCollection.h
@@ -49,7 +49,7 @@ struct FGrid3DCollectionRWInstanceData_GameThread
 	FGrid3DCollectionRWInstanceData_GameThread* OtherInstanceData = nullptr;
 
 	int32 FindAttributeIndexByName(const FName& InName, int32 NumChannels);
-	bool UpdateTargetTexture(ENiagaraGpuBufferFormat BufferFormat);
+	bool UpdateTargetTexture(ENiagaraGpuBufferFormat BufferFormat, UObject* Owner);
 };
 
 struct FGrid3DCollectionRWInstanceData_RenderThread
diff --git a/Engine/Plugins/FX/Niagara/Source/Niagara/Private/NiagaraCommon.cpp b/Engine/Plugins/FX/Niagara/Source/Niagara/Private/NiagaraCommon.cpp
--- a/Engine/Plugins/FX/Niagara/Source/Niagara/Private/NiagaraCommon.cpp
+++ b/Engine/Plugins/FX/Niagara/Source/Niagara/Private/NiagaraCommon.cpp
@@ -1290,7 +1290,7 @@ FString FNiagaraUtilities::SystemInstanceIDToString(FNiagaraSystemInstanceID ID)
 	return FString(Buffer);
 }
 
-TOptional<EPixelFormat> FNiagaraUtilities::BufferFormatToPixelFormat(ENiagaraGpuBufferFormat NiagaraFormat, EPixelFormatCapabilities RequiredCapabilities, int NumberOfChannels)
+TOptional<EPixelFormat> FNiagaraUtilities::BufferFormatToPixelFormat(ENiagaraGpuBufferFormat NiagaraFormat, EPixelFormatCapabilities RequiredCapabilities, int NumberOfChannels, UObject* RequesterForDebugging)
 {
 	if (!ensure(NumberOfChannels >= 1 && NumberOfChannels <= 4))
 	{
@@ -1347,11 +1347,11 @@ TOptional<EPixelFormat> FNiagaraUtilities::BufferFormatToPixelFormat(ENiagaraGpu
 		}
 	};
 
-	UE_LOG(LogNiagara, Error, TEXT("Could not find a valid EPixelFormat for NiagaraFormat(%d) Channels(%d) Caps(%08x)"), int(NiagaraFormat), NumberOfChannels, int(RequiredCapabilities));
+	UE_LOG(LogNiagara, Error, TEXT("Could not find a valid EPixelFormat for NiagaraFormat(%s) Channels(%d) Caps(%08x) for (%s)"), *StaticEnum<ENiagaraGpuBufferFormat>()->GetNameStringByValue(int64(NiagaraFormat)), NumberOfChannels, int(RequiredCapabilities), *GetFullNameSafe(RequesterForDebugging));
 	return TOptional<EPixelFormat>();
 }
 
-TOptional<ETextureRenderTargetFormat> FNiagaraUtilities::BufferFormatToRenderTargetFormat(ENiagaraGpuBufferFormat NiagaraFormat, EPixelFormatCapabilities RequiredCapabilities)
+TOptional<ETextureRenderTargetFormat> FNiagaraUtilities::BufferFormatToRenderTargetFormat(ENiagaraGpuBufferFormat NiagaraFormat, EPixelFormatCapabilities RequiredCapabilities, UObject* RequesterForDebugging)
 {
 	switch (NiagaraFormat)
 	{
@@ -1377,19 +1377,19 @@ TOptional<ETextureRenderTargetFormat> FNiagaraUtilities::BufferFormatToRenderTar
 			// Intentional fall through to failure
 
 		default:
-			UE_LOG(LogNiagara, Error, TEXT("Could not find a valid ETextureRenderTargetFormat for NiagaraFormat(%d) Caps(%08x)"), int(NiagaraFormat), int(RequiredCapabilities));
+			UE_LOG(LogNiagara, Error, TEXT("Could not find a valid ETextureRenderTargetFormat for NiagaraFormat(%s) Caps(%08x) for (%s)"), *StaticEnum<ENiagaraGpuBufferFormat>()->GetNameStringByValue(int64(NiagaraFormat)), int(RequiredCapabilities), *GetFullNameSafe(RequesterForDebugging));
 			return TOptional<ETextureRenderTargetFormat>();
 	}
 }
 
-EPixelFormat FNiagaraUtilities::BufferFormatToPixelFormat(ENiagaraGpuBufferFormat NiagaraFormat)
+TOptional<EPixelFormat> FNiagaraUtilities::BufferFormatToPixelFormat(ENiagaraGpuBufferFormat NiagaraFormat, EPixelFormatCapabilities RequiredCapabilities, int NumberOfChannels)
 {
-	return BufferFormatToPixelFormat(NiagaraFormat, EPixelFormatCapabilities::TypedUAVLoad | EPixelFormatCapabilities::TypedUAVStore).Get(EPixelFormat::PF_R32_FLOAT);
+	return BufferFormatToPixelFormat(NiagaraFormat, RequiredCapabilities, NumberOfChannels, nullptr);
 }
 
-ETextureRenderTargetFormat FNiagaraUtilities::BufferFormatToRenderTargetFormat(ENiagaraGpuBufferFormat NiagaraFormat)
+TOptional<ETextureRenderTargetFormat> FNiagaraUtilities::BufferFormatToRenderTargetFormat(ENiagaraGpuBufferFormat NiagaraFormat, EPixelFormatCapabilities RequiredCapabilities)
 {
-	return BufferFormatToRenderTargetFormat(NiagaraFormat, EPixelFormatCapabilities::TypedUAVLoad | EPixelFormatCapabilities::TypedUAVStore).Get(ETextureRenderTargetFormat::RTF_R32f);
+	return BufferFormatToRenderTargetFormat(NiagaraFormat, RequiredCapabilities, nullptr);
 }
 
 FString FNiagaraUtilities::SanitizeNameForObjectsAndPackages(const FString& InName)
diff --git a/Engine/Plugins/FX/Niagara/Source/Niagara/Private/NiagaraDataInterfaceGrid2DCollection.cpp b/Engine/Plugins/FX/Niagara/Source/Niagara/Private/NiagaraDataInterfaceGrid2DCollection.cpp
--- a/Engine/Plugins/FX/Niagara/Source/Niagara/Private/NiagaraDataInterfaceGrid2DCollection.cpp
+++ b/Engine/Plugins/FX/Niagara/Source/Niagara/Private/NiagaraDataInterfaceGrid2DCollection.cpp
@@ -1970,7 +1970,7 @@ bool UNiagaraDataInterfaceGrid2DCollection::InitPerInstanceData(void* PerInstanc
 		BufferFormat = ENiagaraGpuBufferFormat(GNiagaraGrid2DOverrideFormat);
 	}
 
-	InstanceData->PixelFormat = FNiagaraUtilities::BufferFormatToPixelFormat(BufferFormat, GNDIGrid2DFomatCaps);
+	InstanceData->PixelFormat = FNiagaraUtilities::BufferFormatToPixelFormat(BufferFormat, GNDIGrid2DFomatCaps, 1, this);
 
 	if (!FMath::IsNearlyEqual(GNiagaraGrid2DResolutionMultiplier, 1.0f))
 	{
@@ -2006,7 +2006,7 @@ bool UNiagaraDataInterfaceGrid2DCollection::InitPerInstanceData(void* PerInstanc
 	// Initialize target texture
 	InstanceData->TargetTexture = nullptr;
 	InstanceData->RTUserParamBinding.Init(SystemInstance->GetInstanceParameters(), RenderTargetUserParameter.Parameter);
-	InstanceData->UpdateTargetTexture(BufferFormat);
+	InstanceData->UpdateTargetTexture(BufferFormat, this);
 
 #if WITH_EDITORONLY_DATA
 	InstanceData->bPreviewGrid = bPreviewGrid;
@@ -2120,7 +2120,7 @@ bool UNiagaraDataInterfaceGrid2DCollection::PerInstanceTick(void* PerInstanceDat
 		BufferFormat = ENiagaraGpuBufferFormat(GNiagaraGrid2DOverrideFormat);
 	}
 
-	bool NeedsReset = InstanceData->UpdateTargetTexture(BufferFormat);
+	bool NeedsReset = InstanceData->UpdateTargetTexture(BufferFormat, this);
 
 	FNiagaraDataInterfaceProxyGrid2DCollectionProxy* RT_Proxy = GetProxyAs<FNiagaraDataInterfaceProxyGrid2DCollectionProxy>();
 	ENQUEUE_RENDER_COMMAND(FUpdateData)(
@@ -2541,7 +2541,7 @@ bool UNiagaraDataInterfaceGrid2DCollection::PerInstanceTickPostSimulate(void* Pe
 				BufferFormat = ENiagaraGpuBufferFormat(GNiagaraGrid2DOverrideFormat);
 			}
 
-			bNeedsReset = InstanceData->UpdateTargetTexture(BufferFormat);
+			bNeedsReset = InstanceData->UpdateTargetTexture(BufferFormat, this);
 		}
 
 		// Push Updates to Proxy.
@@ -2613,7 +2613,7 @@ int32 FGrid2DCollectionRWInstanceData_GameThread::FindAttributeIndexByName(const
 	return INDEX_NONE;
 }
 
-bool FGrid2DCollectionRWInstanceData_GameThread::UpdateTargetTexture(ENiagaraGpuBufferFormat BufferFormat)
+bool FGrid2DCollectionRWInstanceData_GameThread::UpdateTargetTexture(ENiagaraGpuBufferFormat BufferFormat, UObject* Owner)
 {
 	// Pull value from user parameter
 	if (UObject* UserParamObject = RTUserParamBinding.GetValue())
@@ -2633,7 +2633,7 @@ bool FGrid2DCollectionRWInstanceData_GameThread::UpdateTargetTexture(ENiagaraGpu
 	{
 		if ( UTextureRenderTarget2DArray* TargetTextureArray = Cast<UTextureRenderTarget2DArray>(TargetTexture) )
 		{
-			const TOptional<EPixelFormat> RenderTargetFormat = FNiagaraUtilities::BufferFormatToPixelFormat(BufferFormat, GNDIGrid2DFomatCaps);
+			const TOptional<EPixelFormat> RenderTargetFormat = FNiagaraUtilities::BufferFormatToPixelFormat(BufferFormat, GNDIGrid2DFomatCaps, 1, Owner);
 			if (RenderTargetFormat.IsSet())
 			{
 				if (TargetTextureArray->SizeX != NumCells.X || TargetTextureArray->SizeY != NumCells.Y || TargetTextureArray->Slices != NumAttributes || TargetTextureArray->OverrideFormat != RenderTargetFormat.GetValue())
@@ -2660,7 +2660,7 @@ bool FGrid2DCollectionRWInstanceData_GameThread::UpdateTargetTexture(ENiagaraGpu
 			{
 				const FNiagaraGrid2DLegacyTiled2DInfo Tiled2DInfo(NumCells, NumAttributes);
 
-				const TOptional<ETextureRenderTargetFormat> RenderTargetFormat  = FNiagaraUtilities::BufferFormatToRenderTargetFormat(BufferFormat, GNDIGrid2DFomatCaps);
+				const TOptional<ETextureRenderTargetFormat> RenderTargetFormat  = FNiagaraUtilities::BufferFormatToRenderTargetFormat(BufferFormat, GNDIGrid2DFomatCaps, Owner);
 				if ( RenderTargetFormat.IsSet() )
 				{
 					if (TargetTexture2D->SizeX != Tiled2DInfo.Size.X || TargetTexture2D->SizeY != Tiled2DInfo.Size.Y || TargetTexture2D->RenderTargetFormat != RenderTargetFormat.GetValue())
diff --git a/Engine/Plugins/FX/Niagara/Source/Niagara/Private/NiagaraDataInterfaceGrid3DCollection.cpp b/Engine/Plugins/FX/Niagara/Source/Niagara/Private/NiagaraDataInterfaceGrid3DCollection.cpp
--- a/Engine/Plugins/FX/Niagara/Source/Niagara/Private/NiagaraDataInterfaceGrid3DCollection.cpp
+++ b/Engine/Plugins/FX/Niagara/Source/Niagara/Private/NiagaraDataInterfaceGrid3DCollection.cpp
@@ -3274,7 +3274,7 @@ bool UNiagaraDataInterfaceGrid3DCollection::InitPerInstanceData(void* PerInstanc
 		InstanceData->NumTiles.Y = 1;
 		InstanceData->NumTiles.Z = 1;
 		
-		InstanceData->PixelFormat = FNiagaraUtilities::BufferFormatToPixelFormat(BufferFormat, GNDIGrid3DFomatCaps, NumAttribChannelsFound);
+		InstanceData->PixelFormat = FNiagaraUtilities::BufferFormatToPixelFormat(BufferFormat, GNDIGrid3DFomatCaps, NumAttribChannelsFound, this);
 
 		InstanceData->UseRGBATexture = true;
 	}
@@ -3315,12 +3315,12 @@ bool UNiagaraDataInterfaceGrid3DCollection::InitPerInstanceData(void* PerInstanc
 		check(InstanceData->NumTiles.Z > 0);
 
 		InstanceData->UseRGBATexture = false;
-		InstanceData->PixelFormat = FNiagaraUtilities::BufferFormatToPixelFormat(BufferFormat, GNDIGrid3DFomatCaps);
+		InstanceData->PixelFormat = FNiagaraUtilities::BufferFormatToPixelFormat(BufferFormat, GNDIGrid3DFomatCaps, 1, this);
 	}
 
 	InstanceData->TargetTexture = nullptr;
 	InstanceData->RTUserParamBinding.Init(SystemInstance->GetInstanceParameters(), RenderTargetUserParameter.Parameter);
-	InstanceData->UpdateTargetTexture(BufferFormat);
+	InstanceData->UpdateTargetTexture(BufferFormat, this);
 
 #if WITH_EDITOR
 	InstanceData->bPreviewGrid = bPreviewGrid;
@@ -3433,7 +3433,7 @@ bool UNiagaraDataInterfaceGrid3DCollection::PerInstanceTick(void* PerInstanceDat
 	{
 		BufferFormat = ENiagaraGpuBufferFormat(GNiagaraGrid3DOverrideFormat);
 	}
-	bool NeedsReset = InstanceData->UpdateTargetTexture(BufferFormat);
+	bool NeedsReset = InstanceData->UpdateTargetTexture(BufferFormat, this);
 
 	FNiagaraDataInterfaceProxyGrid3DCollectionProxy* RT_Proxy = GetProxyAs<FNiagaraDataInterfaceProxyGrid3DCollectionProxy>();
 	ENQUEUE_RENDER_COMMAND(FUpdateData)(
@@ -3886,7 +3886,7 @@ int32 FGrid3DCollectionRWInstanceData_GameThread::FindAttributeIndexByName(const
 	return INDEX_NONE;
 }
 
-bool FGrid3DCollectionRWInstanceData_GameThread::UpdateTargetTexture(ENiagaraGpuBufferFormat BufferFormat)
+bool FGrid3DCollectionRWInstanceData_GameThread::UpdateTargetTexture(ENiagaraGpuBufferFormat BufferFormat, UObject* Owner)
 {
 	// Pull value from user parameter
 	if (UObject* UserParamObject = RTUserParamBinding.GetValue())
@@ -3903,7 +3903,7 @@ bool FGrid3DCollectionRWInstanceData_GameThread::UpdateTargetTexture(ENiagaraGpu
 	if (TargetTexture != nullptr)
 	{
 		const FIntVector RTSize(NumCells.X * NumTiles.X, NumCells.Y * NumTiles.Y, NumCells.Z * NumTiles.Z);
-		const TOptional<EPixelFormat> RenderTargetFormat = FNiagaraUtilities::BufferFormatToPixelFormat(BufferFormat, GNDIGrid3DFomatCaps);
+		const TOptional<EPixelFormat> RenderTargetFormat = FNiagaraUtilities::BufferFormatToPixelFormat(BufferFormat, GNDIGrid3DFomatCaps, 1, Owner);
 		if (RenderTargetFormat.IsSet())
 		{
 			if (TargetTexture->SizeX != RTSize.X || TargetTexture->SizeY != RTSize.Y || TargetTexture->SizeZ != RTSize.Z || TargetTexture->OverrideFormat != RenderTargetFormat.GetValue())
diff --git a/Engine/Plugins/FX/Niagara/Source/Niagara/Public/NiagaraCommon.h b/Engine/Plugins/FX/Niagara/Source/Niagara/Public/NiagaraCommon.h
--- a/Engine/Plugins/FX/Niagara/Source/Niagara/Public/NiagaraCommon.h
+++ b/Engine/Plugins/FX/Niagara/Source/Niagara/Public/NiagaraCommon.h
@@ -1680,22 +1680,22 @@ namespace FNiagaraUtilities
 	Converts a Niagara format into a EPixelFormat, will look for fallback options if the requested format can not be satisfied based on the RequiredCapabilities.
 	If the return value optional is invalid that means no format was found.
 	*/
-	NIAGARA_API TOptional<EPixelFormat> BufferFormatToPixelFormat(ENiagaraGpuBufferFormat NiagaraFormat, EPixelFormatCapabilities RequiredCapabilities, int NumberOfChannels = 1);
+	NIAGARA_API TOptional<EPixelFormat> BufferFormatToPixelFormat(ENiagaraGpuBufferFormat NiagaraFormat, EPixelFormatCapabilities RequiredCapabilities, int NumberOfChannels, UObject* RequesterForDebugging);
 
 	/**
 	Converts a Niagara format into a ETextureRenderTargetFormat, will look for fallback options if the requested format can not be satisfied based on the RequiredCapabilities.
 	If the return value optional is invalid that means no format was found.
 	*/
-	NIAGARA_API TOptional<ETextureRenderTargetFormat> BufferFormatToRenderTargetFormat(ENiagaraGpuBufferFormat NiagaraFormat, EPixelFormatCapabilities RequiredCapabilities);
+	NIAGARA_API TOptional<ETextureRenderTargetFormat> BufferFormatToRenderTargetFormat(ENiagaraGpuBufferFormat NiagaraFormat, EPixelFormatCapabilities RequiredCapabilities, UObject* RequesterForDebugging);
 
 	NIAGARA_API FString SanitizeNameForObjectsAndPackages(const FString& InName);
 
 	NIAGARA_API FNiagaraVariable ResolveAliases(const FNiagaraVariable& InVar, const FNiagaraAliasContext& InContext);
 
-	UE_DEPRECATED(5.3, "Please update your code to include EPixelFormatCapabilities.")
-	NIAGARA_API EPixelFormat BufferFormatToPixelFormat(ENiagaraGpuBufferFormat NiagaraFormat);
-	UE_DEPRECATED(5.3, "Please update your code to include EPixelFormatCapabilities.")
-	NIAGARA_API ETextureRenderTargetFormat BufferFormatToRenderTargetFormat(ENiagaraGpuBufferFormat NiagaraFormat);
+	UE_DEPRECATED(5.8, "Please update your code to include the caller information.")
+	NIAGARA_API TOptional<EPixelFormat> BufferFormatToPixelFormat(ENiagaraGpuBufferFormat NiagaraFormat, EPixelFormatCapabilities RequiredCapabilities, int NumberOfChannels = 1);
+	UE_DEPRECATED(5.8, "Please update your code to include the caller information.")
+	NIAGARA_API TOptional<ETextureRenderTargetFormat> BufferFormatToRenderTargetFormat(ENiagaraGpuBufferFormat NiagaraFormat, EPixelFormatCapabilities RequiredCapabilities);
 };
 
 USTRUCT()
