Niagara - System adapter classes and automated testing.

#virtualized
#Synced-CL 50784797

diff --git a/Engine/Plugins/FX/Niagara/Source/NiagaraEditor/Private/SystemAdapter/EmitterAdapter.cpp b/Engine/Plugins/FX/Niagara/Source/NiagaraEditor/Private/SystemAdapter/EmitterAdapter.cpp
new file mode 100644
--- /dev/null
+++ b/Engine/Plugins/FX/Niagara/Source/NiagaraEditor/Private/SystemAdapter/EmitterAdapter.cpp
@@ -0,0 +1,194 @@
+// Copyright Epic Games, Inc. All Rights Reserved.
+
+#include "SystemAdapter/EmitterAdapter.h"
+
+#include "NiagaraConstants.h"
+#include "NiagaraEmitter.h"
+#include "NiagaraEmitterHandle.h"
+#include "NiagaraSystem.h"
+#include "SystemAdapter/EmitterAdapter.h"
+#include "SystemAdapter/StackScriptAdapter.h"
+
+namespace UE::Niagara
+{
+	class FEmitterScriptAdapterFactory : public TAdapterPtr<UNiagaraSystem, FStackScriptAdapter>::IFactory
+	{
+	public:
+		FEmitterScriptAdapterFactory(const FGuid& InEmitterHandleId, ENiagaraScriptUsage InUsage)
+			: EmitterHandleId(InEmitterHandleId)
+			, Usage(InUsage)
+		{
+		}
+
+		virtual TSharedRef<FStackScriptAdapter> CreateAdapter(UNiagaraSystem* InSourceObject) const override
+		{
+			return CreateAdapterInternal(InSourceObject);
+		}
+
+		virtual TSharedRef<FStackScriptAdapter> CreateAdapter(const UNiagaraSystem* InSourceObject) const override
+		{
+			return CreateAdapterInternal(InSourceObject);
+		}
+
+	private:
+		template<typename TSystem>
+		TSharedRef<FStackScriptAdapter> CreateAdapterInternal(TSystem OwningSystem) const
+		{
+			if (OwningSystem != nullptr)
+			{
+				const FNiagaraEmitterHandle* EmitterHandle = OwningSystem->GetEmitterHandles().FindByPredicate([this](const FNiagaraEmitterHandle& EmitterHandle)
+					{ return EmitterHandle.GetId() == EmitterHandleId; });
+				if (EmitterHandle != nullptr)
+				{
+					FVersionedNiagaraEmitterData* EmitterData = EmitterHandle->GetEmitterData();
+					switch (Usage)
+					{
+					case ENiagaraScriptUsage::EmitterSpawnScript:
+						return FStackScriptAdapter::Create(OwningSystem, EmitterHandleId, EmitterData->EmitterSpawnScriptProps.Script);
+					case ENiagaraScriptUsage::EmitterUpdateScript:
+						return FStackScriptAdapter::Create(OwningSystem, EmitterHandleId, EmitterData->EmitterUpdateScriptProps.Script);
+					case ENiagaraScriptUsage::ParticleSpawnScript:
+					case ENiagaraScriptUsage::ParticleSpawnScriptInterpolated:
+						return FStackScriptAdapter::Create(OwningSystem, EmitterHandleId, EmitterData->SpawnScriptProps.Script);
+					case ENiagaraScriptUsage::ParticleUpdateScript:
+						return FStackScriptAdapter::Create(OwningSystem, EmitterHandleId, EmitterData->UpdateScriptProps.Script);
+					}
+				}
+			}
+			return MakeShared<FStackScriptAdapter>();
+		}
+
+	private:
+		FGuid EmitterHandleId;
+		ENiagaraScriptUsage Usage;
+	};
+
+	TSharedRef<FEmitterAdapter> FEmitterAdapter::Create(UNiagaraSystem* InOwningSystem, const FNiagaraEmitterHandle& InEmitterHandle)
+	{
+		TSharedRef<FEmitterAdapter> EmitterAdapter = MakeShared<FEmitterAdapter>();
+		EmitterAdapter->Initialize(InOwningSystem, InEmitterHandle);
+		return EmitterAdapter;
+	}
+
+	TSharedRef<FEmitterAdapter> FEmitterAdapter::Create(const UNiagaraSystem* InOwningSystem, const FNiagaraEmitterHandle& InEmitterHandle)
+	{
+		TSharedRef<FEmitterAdapter> EmitterAdapter = MakeShared<FEmitterAdapter>();
+		EmitterAdapter->Initialize(InOwningSystem, InEmitterHandle);
+		return EmitterAdapter;
+	}
+
+	TSharedRef<const FEmitterAdapter> FEmitterAdapter::CreateConst(const UNiagaraSystem* InOwningSystem, const FNiagaraEmitterHandle& InEmitterHandle)
+	{
+		TSharedRef<FEmitterAdapter> EmitterAdapter = MakeShared<FEmitterAdapter>();
+		EmitterAdapter->Initialize(InOwningSystem, InEmitterHandle);
+		return EmitterAdapter;
+	}
+
+	bool FEmitterAdapter::IsValidAdapter() const
+	{
+		return OwningSystemConstWeak.Get() != nullptr && EmitterHandleId.IsValid();
+	}
+
+	bool FEmitterAdapter::IsValidWriteAdapter() const
+	{
+		return OwningSystemWeak.IsSet() && OwningSystemWeak.GetValue().Get() != nullptr && EmitterHandleId.IsValid();
+	}
+
+	TSharedRef<FStackScriptAdapter> FEmitterAdapter::GetEmitterSpawnScript()
+	{
+		return EmitterSpawnScript.Get();
+	}
+
+	TSharedRef<const FStackScriptAdapter> FEmitterAdapter::GetEmitterSpawnScript() const
+	{
+		return EmitterSpawnScript.Get();
+	}
+
+	TSharedRef<FStackScriptAdapter> FEmitterAdapter::GetEmitterUpdateScript()
+	{
+		return EmitterUpdateScript.Get();
+	}
+
+	TSharedRef<const FStackScriptAdapter> FEmitterAdapter::GetEmitterUpdateScript() const
+	{
+		return EmitterUpdateScript.Get();
+	}
+
+	TSharedRef<FStackScriptAdapter> FEmitterAdapter::GetParticleSpawnScript()
+	{
+		return ParticleSpawnScript.Get();
+	}
+
+	TSharedRef<const FStackScriptAdapter> FEmitterAdapter::GetParticleSpawnScript() const
+	{
+		return ParticleSpawnScript.Get();
+	}
+
+	TSharedRef<FStackScriptAdapter> FEmitterAdapter::GetParticleUpdateScript()
+	{
+		return ParticleUpdateScript.Get();
+	}
+
+	TSharedRef<const FStackScriptAdapter> FEmitterAdapter::GetParticleUpdateScript() const
+	{
+		return ParticleUpdateScript.Get();
+	}
+
+
+	int32 FEmitterAdapter::GetNumScripts() const
+	{
+		return 4;
+	}
+
+	TSharedRef<FStackScriptAdapter> FEmitterAdapter::GetScriptAt(int32 Index)
+	{
+		checkf(Index > INDEX_NONE && Index < GetNumScripts(), TEXT("Index out of range"));
+		switch (Index)
+		{
+		case 0:
+			return GetEmitterSpawnScript();
+		case 1:
+			return GetEmitterUpdateScript();
+		case 2:
+			return GetParticleSpawnScript();
+		default: // case 3:
+			return GetParticleUpdateScript();
+		}
+	}
+
+	FScriptSharedRefCollection FEmitterAdapter::GetScripts()
+	{
+		return FScriptSharedRefCollection(&StackScriptAdapterCollection);
+	}
+
+	FScriptConstSharedRefCollection FEmitterAdapter::GetScripts() const
+	{
+		return FScriptConstSharedRefCollection(&StackScriptAdapterCollection);
+	}
+
+	void FEmitterAdapter::Initialize(UNiagaraSystem* InOwningSystem, const FNiagaraEmitterHandle& InEmitterHandle)
+	{
+		OwningSystemWeak = InOwningSystem;
+		OwningSystemConstWeak = InOwningSystem;
+		EmitterName = *InEmitterHandle.GetUniqueInstanceName();
+		EmitterHandleId = InEmitterHandle.GetId();
+		EmitterSpawnScript.Initialize(InOwningSystem, MakeShared<FEmitterScriptAdapterFactory>(EmitterHandleId, ENiagaraScriptUsage::EmitterSpawnScript));
+		EmitterUpdateScript.Initialize(InOwningSystem, MakeShared<FEmitterScriptAdapterFactory>(EmitterHandleId, ENiagaraScriptUsage::EmitterUpdateScript));
+		ParticleSpawnScript.Initialize(InOwningSystem, MakeShared<FEmitterScriptAdapterFactory>(EmitterHandleId, ENiagaraScriptUsage::ParticleSpawnScript));
+		ParticleUpdateScript.Initialize(InOwningSystem, MakeShared<FEmitterScriptAdapterFactory>(EmitterHandleId, ENiagaraScriptUsage::ParticleUpdateScript));
+		StackScriptAdapterCollection.Initialize(this->AsShared());
+	}
+
+	void FEmitterAdapter::Initialize(const UNiagaraSystem* InOwningSystem, const FNiagaraEmitterHandle& InEmitterHandle)
+	{
+		OwningSystemConstWeak = InOwningSystem;
+		EmitterName = *InEmitterHandle.GetUniqueInstanceName();
+		EmitterHandleId = InEmitterHandle.GetId();
+		EmitterSpawnScript.Initialize(InOwningSystem, MakeShared<FEmitterScriptAdapterFactory>(EmitterHandleId, ENiagaraScriptUsage::EmitterSpawnScript));
+		EmitterUpdateScript.Initialize(InOwningSystem, MakeShared<FEmitterScriptAdapterFactory>(EmitterHandleId, ENiagaraScriptUsage::EmitterUpdateScript));
+		ParticleSpawnScript.Initialize(InOwningSystem, MakeShared<FEmitterScriptAdapterFactory>(EmitterHandleId, ENiagaraScriptUsage::ParticleSpawnScript));
+		ParticleUpdateScript.Initialize(InOwningSystem, MakeShared<FEmitterScriptAdapterFactory>(EmitterHandleId, ENiagaraScriptUsage::ParticleUpdateScript));
+		StackScriptAdapterCollection.Initialize(this->AsShared());
+	}
+
+} // namespace UE::Niagara
\ No newline at end of file
diff --git a/Engine/Plugins/FX/Niagara/Source/NiagaraEditor/Private/SystemAdapter/StackFunctionAdapter.cpp b/Engine/Plugins/FX/Niagara/Source/NiagaraEditor/Private/SystemAdapter/StackFunctionAdapter.cpp
new file mode 100644
--- /dev/null
+++ b/Engine/Plugins/FX/Niagara/Source/NiagaraEditor/Private/SystemAdapter/StackFunctionAdapter.cpp
@@ -0,0 +1,164 @@
+// Copyright Epic Games, Inc. All Rights Reserved.
+
+#include "SystemAdapter/StackFunctionAdapter.h"
+
+#include "NiagaraEmitterHandle.h"
+#include "NiagaraNodeAssignment.h"
+#include "NiagaraNodeFunctionCall.h"
+#include "NiagaraSystem.h"
+#include "SystemAdapter/StackFunctionInputAdapter.h"
+#include "TraversalCache/StaticContextBuilder.h"
+#include "TraversalCache/TraversalCache.h"
+#include "TraversalCache/TraversalShared.h"
+#include "ViewModels/Stack/NiagaraStackGraphUtilities.h"
+
+namespace UE::Niagara
+{
+	class FStackFunctionInputAdapterFactory : public TAdapterRefListThreeSources<UNiagaraSystem, UNiagaraScript, UNiagaraNodeFunctionCall, FStackFunctionInputAdapter>::IFactory
+	{
+	public:
+		FStackFunctionInputAdapterFactory(FGuid InOwningEmitterHandleId)
+			: OwningEmitterHandleId(InOwningEmitterHandleId)
+		{
+		}
+
+		virtual void CreateAdapters(UNiagaraSystem* OwningSystem, UNiagaraScript* OwningScript, UNiagaraNodeFunctionCall* OwningFunctionCallNode, TArray<TSharedRef<FStackFunctionInputAdapter>>& OutAdapters) const override
+		{
+			CreateAdaptersInternal(OwningSystem, OwningScript, OwningFunctionCallNode, OutAdapters);
+		}
+
+		virtual void CreateAdapters(const UNiagaraSystem* OwningSystem, const UNiagaraScript* OwningScript, const UNiagaraNodeFunctionCall* OwningFunctionCallNode, TArray<TSharedRef<FStackFunctionInputAdapter>>& OutAdapters) const override
+		{
+			CreateAdaptersInternal(OwningSystem, OwningScript, OwningFunctionCallNode, OutAdapters);
+		}
+
+	private:
+		template<typename TSystem, typename TScript, typename TFunction>
+		void CreateAdaptersInternal(TSystem* OwningSystem, TScript* OwningScript, TFunction* OwningFunctionCallNode, TArray<TSharedRef<FStackFunctionInputAdapter>>& OutAdapters) const
+		{
+			using namespace UE::Niagara::TraversalCache;
+			FVersionedNiagaraEmitterData* OwningEmitterData = nullptr;
+			if (OwningEmitterHandleId.IsValid())
+			{
+				const FNiagaraEmitterHandle* OwningEmitterHandle = OwningSystem->GetEmitterHandles().FindByPredicate([this](const FNiagaraEmitterHandle& EmitterHandle)
+					{ return EmitterHandle.GetId() == OwningEmitterHandleId; });
+				if (OwningEmitterHandle == nullptr)
+				{
+					return;
+				}
+				OwningEmitterData = OwningEmitterHandle->GetEmitterData();
+			}
+
+			FTopLevelScriptStaticContext StaticContext;
+			FStaticContextBuilder::CreateTopLevelScriptContext(*OwningSystem, OwningEmitterData, *OwningScript, StaticContext);
+			TArray<FNiagaraVariable> InputReads;
+			TSet<FNiagaraVariable> HiddenInputReads;
+			FTraversalCache::GetStackFunctionReads(StaticContext, *OwningFunctionCallNode, InputReads, HiddenInputReads, TraversalCache::EStackFunctionReadFilterFlags::InputsOnly);
+
+			for (const FNiagaraVariable& InputRead : InputReads)
+			{
+				OutAdapters.Add(FStackFunctionInputAdapter::Create(OwningFunctionCallNode, InputRead));
+			}
+		}
+
+	private:
+		FGuid OwningEmitterHandleId;
+	};
+
+	TSharedRef<FStackFunctionAdapter> FStackFunctionAdapter::Create(UNiagaraSystem* InOwningSystem, FGuid InOwningEmitterHandleId, UNiagaraScript* InOwningScript, UNiagaraNodeFunctionCall* InFunctionCallNode)
+	{
+		TSharedRef<FStackFunctionAdapter> FunctionAdapter = MakeShared<FStackFunctionAdapter>();
+		FunctionAdapter->Initialize(InOwningSystem, InOwningEmitterHandleId, InOwningScript, InFunctionCallNode);
+		return FunctionAdapter;
+	}
+
+	TSharedRef<FStackFunctionAdapter> FStackFunctionAdapter::Create(const UNiagaraSystem* InOwningSystem, FGuid InOwningEmitterHandleId, const UNiagaraScript* InOwningScript, const UNiagaraNodeFunctionCall* InFunctionCallNode)
+	{
+		TSharedRef<FStackFunctionAdapter> FunctionAdapter = MakeShared<FStackFunctionAdapter>();
+		FunctionAdapter->Initialize(InOwningSystem, InOwningEmitterHandleId, InOwningScript, InFunctionCallNode);
+		return FunctionAdapter;
+	}
+
+	TSharedRef<const FStackFunctionAdapter> FStackFunctionAdapter::CreateConst(const UNiagaraSystem* InOwningSystem, FGuid InOwningEmitterHandleId, const UNiagaraScript* InOwningScript, const UNiagaraNodeFunctionCall* InFunctionCallNode)
+	{
+		TSharedRef<FStackFunctionAdapter> FunctionAdapter = MakeShared<FStackFunctionAdapter>();
+		FunctionAdapter->Initialize(InOwningSystem, InOwningEmitterHandleId, InOwningScript, InFunctionCallNode);
+		return FunctionAdapter;
+	}
+
+	bool FStackFunctionAdapter::IsValidAdapter() const
+	{
+		return FunctionCallNodeConstWeak.Get() != nullptr;
+	}
+
+	bool FStackFunctionAdapter::IsValidWriteAdapter() const
+	{
+		return FunctionCallNodeWeak.IsSet() && FunctionCallNodeWeak.GetValue().Get() != nullptr;
+	}
+
+	FString FStackFunctionAdapter::GetFunctionName() const
+	{
+		const UNiagaraNodeFunctionCall* FunctionCallNode = FunctionCallNodeConstWeak.Get();
+		if (FunctionCallNode != nullptr)
+		{
+			return FunctionCallNode->GetFunctionName();
+		}
+		return FString();
+	}
+
+	void FStackFunctionAdapter::GetSetParameters(TArray<FNiagaraVariableBase>& OutParameters) const
+	{
+		const UNiagaraNodeAssignment* AssignmentNode = Cast<const UNiagaraNodeAssignment>(FunctionCallNodeConstWeak.Get());
+		if (AssignmentNode != nullptr)
+		{
+			for (const FNiagaraVariable& AssignmentTarget : AssignmentNode->GetAssignmentTargets())
+			{
+				OutParameters.Add(AssignmentTarget);
+			}
+		}
+	}
+
+	bool FStackFunctionAdapter::SetsParameter(const FNiagaraVariableBase& InParameter) const
+	{
+		const UNiagaraNodeAssignment* AssignmentNode = Cast<const UNiagaraNodeAssignment>(FunctionCallNodeConstWeak.Get());
+		return AssignmentNode != nullptr && AssignmentNode->GetAssignmentTargets().Contains(InParameter);
+	}
+
+	TSharedRefCollection<FStackFunctionInputAdapter> FStackFunctionAdapter::GetInputs()
+	{
+		return Inputs.Get();
+	}
+
+	TConstSharedRefCollection<FStackFunctionInputAdapter> FStackFunctionAdapter::GetInputs() const
+	{
+		return Inputs.Get();
+	}
+
+	TSharedPtr<FStackFunctionInputAdapter> FStackFunctionAdapter::FindInputByName(FName InputName)
+	{
+		return Inputs.Get().FindByPredicate([&InputName](const TSharedRef<FStackFunctionInputAdapter>& Input)
+			{ return Input->GetName() == InputName; });
+	}
+
+	TSharedPtr<const FStackFunctionInputAdapter> FStackFunctionAdapter::FindInputByName(FName InputName) const
+	{
+		return Inputs.Get().FindByPredicate([&InputName](const TSharedRef<const FStackFunctionInputAdapter>& Input)
+			{ return Input->GetName() == InputName; });
+	}
+
+	void FStackFunctionAdapter::Initialize(UNiagaraSystem* InOwningSystem, FGuid InOwningEmitterHandleId, UNiagaraScript* InOwningScript, UNiagaraNodeFunctionCall* InFunctionCallNode)
+	{
+		FunctionCallNodeWeak = InFunctionCallNode;
+		FunctionCallNodeConstWeak = InFunctionCallNode;
+		Type = InFunctionCallNode->IsA<UNiagaraNodeAssignment>() ? EFunctionType::SetParametersModule : EFunctionType::ScriptModule;
+		Inputs.Initialize(InOwningSystem, InOwningScript, InFunctionCallNode, MakeShared<FStackFunctionInputAdapterFactory>(InOwningEmitterHandleId));
+	}
+
+	void FStackFunctionAdapter::Initialize(const UNiagaraSystem* InOwningSystem, FGuid InOwningEmitterHandleId, const UNiagaraScript* InOwningScript, const UNiagaraNodeFunctionCall* InFunctionCallNode)
+	{
+		FunctionCallNodeConstWeak = InFunctionCallNode;
+		Type = InFunctionCallNode->IsA<UNiagaraNodeAssignment>() ? EFunctionType::SetParametersModule : EFunctionType::ScriptModule;
+		Inputs.Initialize(InOwningSystem, InOwningScript, InFunctionCallNode, MakeShared<FStackFunctionInputAdapterFactory>(InOwningEmitterHandleId));
+	}
+
+} // namespace UE::Niagara
\ No newline at end of file
diff --git a/Engine/Plugins/FX/Niagara/Source/NiagaraEditor/Private/SystemAdapter/StackFunctionInputAdapter.cpp b/Engine/Plugins/FX/Niagara/Source/NiagaraEditor/Private/SystemAdapter/StackFunctionInputAdapter.cpp
new file mode 100644
--- /dev/null
+++ b/Engine/Plugins/FX/Niagara/Source/NiagaraEditor/Private/SystemAdapter/StackFunctionInputAdapter.cpp
@@ -0,0 +1,121 @@
+// Copyright Epic Games, Inc. All Rights Reserved.
+
+#include "SystemAdapter/StackFunctionInputAdapter.h"
+
+#include "NiagaraNodeFunctionCall.h"
+#include "ViewModels/Stack/NiagaraParameterHandle.h"
+
+namespace UE::Niagara
+{
+	TSharedRef<FStackFunctionInputAdapter> FStackFunctionInputAdapter::Create(UNiagaraNodeFunctionCall* InFunctionCallNode, const FNiagaraVariableBase& InInputParameter)
+	{
+		TSharedRef<FStackFunctionInputAdapter> FunctionInputAdapter = MakeShared<FStackFunctionInputAdapter>();
+		FunctionInputAdapter->Initialize(InFunctionCallNode, InInputParameter);
+		return FunctionInputAdapter;
+	}
+
+	TSharedRef<FStackFunctionInputAdapter> FStackFunctionInputAdapter::Create(const UNiagaraNodeFunctionCall* InFunctionCallNode, const FNiagaraVariableBase& InInputParameter)
+	{
+		TSharedRef<FStackFunctionInputAdapter> FunctionInputAdapter = MakeShared<FStackFunctionInputAdapter>();
+		FunctionInputAdapter->Initialize(InFunctionCallNode, InInputParameter);
+		return FunctionInputAdapter;
+	}
+
+	TSharedRef<const FStackFunctionInputAdapter> FStackFunctionInputAdapter::CreateConst(const UNiagaraNodeFunctionCall* InFunctionCallNode, const FNiagaraVariableBase& InInputParameter)
+	{
+		TSharedRef<FStackFunctionInputAdapter> FunctionInputAdapter = MakeShared<FStackFunctionInputAdapter>();
+		FunctionInputAdapter->Initialize(InFunctionCallNode, InInputParameter);
+		return FunctionInputAdapter;
+	}
+
+	bool FStackFunctionInputAdapter::IsValidAdapter() const
+	{
+		return FunctionCallNodeConstWeak.Get() != nullptr && InputParameter.IsValid();
+	}
+
+	bool FStackFunctionInputAdapter::IsValidWriteAdapter() const
+	{
+		return FunctionCallNodeWeak.IsSet() && FunctionCallNodeWeak.GetValue().Get() != nullptr && InputParameter.IsValid();
+	}
+
+	FName FStackFunctionInputAdapter::GetName() const
+	{
+		return InputName;
+	}
+
+	const FNiagaraTypeDefinition& FStackFunctionInputAdapter::GetType() const
+	{
+		return InputParameter.GetType();
+	}
+
+	bool FStackFunctionInputAdapter::TryGetLocalValue(bool& OutValue) const
+	{
+		return false;
+	}
+
+	bool FStackFunctionInputAdapter::TryGetLocalValue(int32& OutValue) const
+	{
+		return false;
+	}
+
+	bool FStackFunctionInputAdapter::TryGetLocalValue(float& OutValue) const
+	{
+		return false;
+	}
+
+	bool FStackFunctionInputAdapter::TryGetLocalValue(FVector2f& OutValue) const
+	{
+		return false;
+	}
+
+	bool FStackFunctionInputAdapter::TryGetLocalValue(FVector3f& OutValue) const
+	{
+		return false;
+	}
+
+	bool FStackFunctionInputAdapter::TryGetLocalValue(FVector4f& OutValue) const
+	{
+		return false;
+	}
+
+	bool FStackFunctionInputAdapter::TryGetLocalValueData(TNotNull<UStruct*> ValueType, const uint8*& OutData) const
+	{
+		return false;
+	}
+
+	FNiagaraVariableBase FStackFunctionInputAdapter::GetLinkedValue() const
+	{
+		return FNiagaraVariableBase();
+	}
+
+	UNiagaraDataInterface* FStackFunctionInputAdapter::GetDataValue() const
+	{
+		return nullptr;
+	}
+
+	UObject* FStackFunctionInputAdapter::GetObjectValue() const
+	{
+		return nullptr;
+	}
+
+	TSharedPtr<FStackFunctionAdapter> FStackFunctionInputAdapter::GetDynamicValue() const
+	{
+		return TSharedPtr<FStackFunctionAdapter>();
+	}
+
+	void FStackFunctionInputAdapter::Initialize(UNiagaraNodeFunctionCall* InFunctionCallNode, const FNiagaraVariableBase& InInputParameter)
+	{
+		InputParameter = InInputParameter;
+		InputName = FNiagaraParameterHandle(InputParameter.GetName()).GetName();
+		FunctionCallNodeWeak = InFunctionCallNode;
+		FunctionCallNodeConstWeak = InFunctionCallNode;
+	}
+
+	void FStackFunctionInputAdapter::Initialize(const UNiagaraNodeFunctionCall* InFunctionCallNode, const FNiagaraVariableBase& InInputParameter)
+	{
+		InputParameter = InInputParameter;
+		InputName = FNiagaraParameterHandle(InputParameter.GetName()).GetName();
+		FunctionCallNodeConstWeak = InFunctionCallNode;
+	}
+
+} // namespace UE::Niagara
\ No newline at end of file
diff --git a/Engine/Plugins/FX/Niagara/Source/NiagaraEditor/Private/SystemAdapter/StackScriptAdapter.cpp b/Engine/Plugins/FX/Niagara/Source/NiagaraEditor/Private/SystemAdapter/StackScriptAdapter.cpp
new file mode 100644
--- /dev/null
+++ b/Engine/Plugins/FX/Niagara/Source/NiagaraEditor/Private/SystemAdapter/StackScriptAdapter.cpp
@@ -0,0 +1,162 @@
+// Copyright Epic Games, Inc. All Rights Reserved.
+
+#include "SystemAdapter/StackScriptAdapter.h"
+
+#include "NiagaraGraph.h"
+#include "NiagaraNodeFunctionCall.h"
+#include "NiagaraNodeOutput.h"
+#include "NiagaraScript.h"
+#include "NiagaraScriptSource.h"
+#include "NiagaraSystem.h"
+#include "SystemAdapter/StackFunctionAdapter.h"
+#include "ViewModels/Stack/NiagaraStackGraphUtilities.h"
+
+namespace UE::Niagara
+{
+	class FStackModuleAdapterFactory : public TAdapterRefListTwoSources<UNiagaraSystem, UNiagaraScript, FStackFunctionAdapter>::IFactory
+	{
+	public:
+		FStackModuleAdapterFactory(FGuid InOwningEmitterHandleId)
+			: OwningEmitterHandleId(InOwningEmitterHandleId)
+		{
+		}
+
+		virtual void CreateAdapters(UNiagaraSystem* OwningSystem, UNiagaraScript* OwningScript, TArray<TSharedRef<FStackFunctionAdapter>>& OutAdapters) const override
+		{
+			CreateAdaptersInternal(OwningSystem, OwningScript, OutAdapters);
+		}
+
+		virtual void CreateAdapters(const UNiagaraSystem* OwningSystem, const UNiagaraScript* OwningScript, TArray<TSharedRef<FStackFunctionAdapter>>& OutAdapters) const override
+		{
+			CreateAdaptersInternal(OwningSystem, OwningScript, OutAdapters);
+		}
+
+	private:
+		template<typename TSystem, typename TScript>
+		void CreateAdaptersInternal(TSystem* OwningSystem, TScript OwningScript, TArray<TSharedRef<FStackFunctionAdapter>>& OutAdapters) const
+		{
+			const UNiagaraScriptSource* ScriptSource = CastChecked<UNiagaraScriptSource>(OwningScript->GetLatestSource());
+			UNiagaraNodeOutput* OutputNode = ScriptSource->NodeGraph->FindEquivalentOutputNode(OwningScript->GetUsage(), OwningScript->GetUsageId());
+			if (OutputNode != nullptr)
+			{
+				TArray<FNiagaraStackGraphUtilities::FStackNodeGroup> StackGroups;
+				FNiagaraStackGraphUtilities::GetStackNodeGroups(*OutputNode, StackGroups);
+				for (const FNiagaraStackGraphUtilities::FStackNodeGroup& StackGroup : StackGroups)
+				{
+					UNiagaraNodeFunctionCall* ModuleNode = Cast<UNiagaraNodeFunctionCall>(StackGroup.EndNode);
+					if (ModuleNode != nullptr)
+					{
+						OutAdapters.Add(FStackFunctionAdapter::Create(OwningSystem, OwningEmitterHandleId, OwningScript, ModuleNode));
+					}
+				}
+			}
+		}
+
+	private:
+		FGuid OwningEmitterHandleId;
+	};
+
+	TSharedRef<FStackScriptAdapter> FStackScriptAdapter::Create(UNiagaraSystem* InOwningSystem, const FGuid& InOwningEmitterHandleId, UNiagaraScript* InScript)
+	{
+		TSharedRef<FStackScriptAdapter> ScriptAdapter = MakeShared<FStackScriptAdapter>();
+		ScriptAdapter->Initialize(InOwningSystem, InOwningEmitterHandleId, InScript);
+		return ScriptAdapter;
+	}
+
+	TSharedRef<FStackScriptAdapter> FStackScriptAdapter::Create(const UNiagaraSystem* InOwningSystem, const FGuid& InOwningEmitterHandleId, const UNiagaraScript* InScript)
+	{
+		TSharedRef<FStackScriptAdapter> ScriptAdapter = MakeShared<FStackScriptAdapter>();
+		ScriptAdapter->Initialize(InOwningSystem, InOwningEmitterHandleId, InScript);
+		return ScriptAdapter;
+	}
+
+	TSharedRef<const FStackScriptAdapter> FStackScriptAdapter::CreateConst(const UNiagaraSystem* InOwningSystem, const FGuid& InOwningEmitterHandleId, const UNiagaraScript* InScript)
+	{
+		TSharedRef<FStackScriptAdapter> ScriptAdapter = MakeShared<FStackScriptAdapter>();
+		ScriptAdapter->Initialize(InOwningSystem, InOwningEmitterHandleId, InScript);
+		return ScriptAdapter;
+	}
+
+	bool FStackScriptAdapter::IsValidAdapter() const
+	{
+		return ScriptConstWeak.Get() != nullptr;
+	}
+
+	bool FStackScriptAdapter::IsValidWriteAdapter() const
+	{
+		return ScriptWeak.IsSet() && ScriptWeak.GetValue().Get() != nullptr;
+	}
+
+	ENiagaraScriptUsage FStackScriptAdapter::GetUsage() const
+	{
+		const UNiagaraScript* ConstScript = ScriptConstWeak.Get();
+		return ConstScript != nullptr ? ConstScript->GetUsage() : ENiagaraScriptUsage::Function;
+	}
+
+	FGuid FStackScriptAdapter::GetUsageId() const
+	{
+		const UNiagaraScript* ConstScript = ScriptConstWeak.Get();
+		return ConstScript != nullptr ? ConstScript->GetUsageId() : FGuid();
+	}
+
+	TSharedRefCollection<FStackFunctionAdapter> FStackScriptAdapter::GetModules()
+	{
+		return Modules.Get();
+	}
+
+	TConstSharedRefCollection<FStackFunctionAdapter> FStackScriptAdapter::GetModules() const
+	{
+		return Modules.Get();
+	}
+
+	TSharedPtr<FStackFunctionAdapter> FStackScriptAdapter::FindModuleByName(const FName& InName)
+	{
+		FString NameString = InName.ToString();
+		return Modules.Get().FindByPredicate([&NameString](const TSharedRef<FStackFunctionAdapter>& ModuleAdapter)
+			{ return ModuleAdapter->GetFunctionName() == NameString; });
+	}
+
+	TSharedPtr<const FStackFunctionAdapter> FStackScriptAdapter::FindModuleByName(const FName& InName) const
+	{
+		FString NameString = InName.ToString();
+		return Modules.Get().FindByPredicate([&NameString](TSharedRef<const FStackFunctionAdapter> ModuleAdapter)
+			{ return ModuleAdapter->GetFunctionName() == NameString; });
+	}
+
+	void FStackScriptAdapter::FindSetParametersModulesBySetParameter(const FNiagaraVariableBase& InParameter, TArray<TSharedRef<FStackFunctionAdapter>>& OutSetParametersModules)
+	{
+		Modules.Get().FindAllByPredicate(OutSetParametersModules, [&InParameter](const TSharedRef<FStackFunctionAdapter>& ModuleAdapter)
+			{ return ModuleAdapter->GetType() == FStackFunctionAdapter::EFunctionType::SetParametersModule && ModuleAdapter->SetsParameter(InParameter); });
+	}
+
+	void FStackScriptAdapter::FindSetParametersModulesBySetParameter(const FNiagaraVariableBase& InParameter, TArray<TSharedRef<const FStackFunctionAdapter>>& OutSetParametersModules) const
+	{
+		Modules.Get().FindAllByPredicate(OutSetParametersModules, [&InParameter](const TSharedRef<const FStackFunctionAdapter>& ModuleAdapter)
+			{ return ModuleAdapter->GetType() == FStackFunctionAdapter::EFunctionType::SetParametersModule && ModuleAdapter->SetsParameter(InParameter); });
+	}
+
+	TSharedPtr<FStackFunctionAdapter> FStackScriptAdapter::FindFirstSetParametersModuleBySetParameter(const FNiagaraVariableBase& InParameter)
+	{
+		return Modules.Get().FindByPredicate([&InParameter](const TSharedRef<FStackFunctionAdapter>& ModuleAdapter)
+			{ return ModuleAdapter->GetType() == FStackFunctionAdapter::EFunctionType::SetParametersModule && ModuleAdapter->SetsParameter(InParameter); });
+	}
+
+	TSharedPtr<const FStackFunctionAdapter> FStackScriptAdapter::FindFirstSetParametersModuleBySetParameter(const FNiagaraVariableBase& InParameter) const
+	{
+		return Modules.Get().FindByPredicate([&InParameter](const TSharedRef<const FStackFunctionAdapter>& ModuleAdapter)
+			{ return ModuleAdapter->GetType() == FStackFunctionAdapter::EFunctionType::SetParametersModule && ModuleAdapter->SetsParameter(InParameter); });
+	}
+
+	void FStackScriptAdapter::Initialize(UNiagaraSystem* InOwningSystem, const FGuid& InOwningEmitterHandleId, UNiagaraScript* InScript)
+	{
+		ScriptWeak = InScript;
+		ScriptConstWeak = InScript;
+		Modules.Initialize(InOwningSystem, InScript, MakeShared<FStackModuleAdapterFactory>(InOwningEmitterHandleId));
+	}
+
+	void FStackScriptAdapter::Initialize(const UNiagaraSystem* InOwningSystem, const FGuid& InOwningEmitterHandleId, const UNiagaraScript* InScript)
+	{
+		ScriptConstWeak = InScript;
+		Modules.Initialize(InOwningSystem, InScript, MakeShared<FStackModuleAdapterFactory>(InOwningEmitterHandleId));
+	}
+}
\ No newline at end of file
diff --git a/Engine/Plugins/FX/Niagara/Source/NiagaraEditor/Private/SystemAdapter/StackScriptAdapterCollection.cpp b/Engine/Plugins/FX/Niagara/Source/NiagaraEditor/Private/SystemAdapter/StackScriptAdapterCollection.cpp
new file mode 100644
--- /dev/null
+++ b/Engine/Plugins/FX/Niagara/Source/NiagaraEditor/Private/SystemAdapter/StackScriptAdapterCollection.cpp
@@ -0,0 +1,25 @@
+// Copyright Epic Games, Inc. All Rights Reserved.
+
+#include "SystemAdapter/StackScriptAdapterCollection.h"
+
+namespace UE::Niagara
+{
+	void FStackScriptAdapterCollection::Initialize(TSharedRef<IStackScriptAdapterOwner> InScriptOwner)
+	{
+		ScriptOwnerWeak = InScriptOwner;
+	}
+
+	int32 FStackScriptAdapterCollection::Num() const
+	{
+		TSharedPtr<IStackScriptAdapterOwner> ScriptOwner = ScriptOwnerWeak.Pin();
+		return ScriptOwnerWeak.IsValid() ? ScriptOwner->GetNumScripts() : 0;
+	}
+
+	TSharedRef<FStackScriptAdapter> FStackScriptAdapterCollection::operator[](int32 Index) const
+	{
+		TSharedPtr<IStackScriptAdapterOwner> ScriptOwner = ScriptOwnerWeak.Pin();
+		checkf(ScriptOwner.IsValid(), TEXT("Script owner is invalid"));
+		return ScriptOwner->GetScriptAt(Index);
+	}
+
+} // namespace UE::Niagara
\ No newline at end of file
diff --git a/Engine/Plugins/FX/Niagara/Source/NiagaraEditor/Private/SystemAdapter/SystemAdapter.cpp b/Engine/Plugins/FX/Niagara/Source/NiagaraEditor/Private/SystemAdapter/SystemAdapter.cpp
new file mode 100644
--- /dev/null
+++ b/Engine/Plugins/FX/Niagara/Source/NiagaraEditor/Private/SystemAdapter/SystemAdapter.cpp
@@ -0,0 +1,459 @@
+// Copyright Epic Games, Inc. All Rights Reserved.
+
+#include "SystemAdapter/SystemAdapter.h"
+
+#include "NiagaraConstants.h"
+#include "NiagaraSystem.h"
+#include "SystemAdapter/EmitterAdapter.h"
+#include "SystemAdapter/StackScriptAdapter.h"
+
+namespace UE::Niagara
+{
+	TSharedRef<FUserParameterAdapter> FUserParameterAdapter::Create(UNiagaraSystem* InSystem, const FNiagaraVariableWithOffset& InUserParameterWithOffset)
+	{
+		TSharedRef<FUserParameterAdapter> UserParameterAdapter = MakeShared<FUserParameterAdapter>();
+		UserParameterAdapter->Initialize(InSystem, InUserParameterWithOffset);
+		return UserParameterAdapter;
+	}
+
+	TSharedRef<FUserParameterAdapter> FUserParameterAdapter::Create(const UNiagaraSystem* InSystem, const FNiagaraVariableWithOffset& InUserParameterWithOffset)
+	{
+		TSharedRef<FUserParameterAdapter> UserParameterAdapter = MakeShared<FUserParameterAdapter>();
+		UserParameterAdapter->Initialize(InSystem, InUserParameterWithOffset);
+		return UserParameterAdapter;
+	}
+
+	TSharedRef<const FUserParameterAdapter> FUserParameterAdapter::CreateConst(const UNiagaraSystem* InSystem, const FNiagaraVariableWithOffset& InUserParameterWithOffset)
+	{
+		TSharedRef<FUserParameterAdapter> UserParameterAdapter = MakeShared<FUserParameterAdapter>();
+		UserParameterAdapter->Initialize(InSystem, InUserParameterWithOffset);
+		return UserParameterAdapter;
+	}
+
+	bool FUserParameterAdapter::IsValidAdapter() const
+	{
+		return SystemConstWeak.Get() != nullptr;
+	}
+
+	bool FUserParameterAdapter::IsValidWriteAdapter() const
+	{
+		return SystemWeak.IsSet() && SystemWeak.GetValue().Get() != nullptr;
+	}
+
+	FName FUserParameterAdapter::GetName() const
+	{
+		if (NameCache == NAME_None)
+		{
+			const int32 UserPrefixLength = FNiagaraConstants::UserNamespace.GetStringLength() + 1;
+			FString FullNameString = UserParameterWithOffset.GetName().ToString();
+			NameCache = FName(FullNameString.RightChop(UserPrefixLength));
+		}
+		return NameCache;
+	}
+
+	bool FUserParameterAdapter::TryGetValue(bool& OutValue) const
+	{
+		const UNiagaraSystem* SystemConst = SystemConstWeak.Get();
+		if (SystemConst != nullptr)
+		{
+			if (UserParameterWithOffset.GetType() == FNiagaraTypeDefinition::GetBoolDef())
+			{
+				OutValue = SystemConst->GetExposedParameters().GetParameterValueFromOffset<FNiagaraBool>(UserParameterWithOffset.Offset).GetValue();
+				return true;
+			}
+		}
+		return false;
+	}
+
+	bool FUserParameterAdapter::TryGetValue(int32& OutValue) const
+	{
+		const UNiagaraSystem* SystemConst = SystemConstWeak.Get();
+		if (SystemConst != nullptr)
+		{
+			if (UserParameterWithOffset.GetType() == FNiagaraTypeDefinition::GetIntDef())
+			{
+				OutValue = SystemConst->GetExposedParameters().GetParameterValueFromOffset<FNiagaraInt32>(UserParameterWithOffset.Offset).Value;
+				return true;
+			}
+		}
+		return false;
+	}
+
+	bool FUserParameterAdapter::TryGetValue(float& OutValue) const
+	{
+		const UNiagaraSystem* SystemConst = SystemConstWeak.Get();
+		if (SystemConst != nullptr)
+		{
+			if (UserParameterWithOffset.GetType() == FNiagaraTypeDefinition::GetFloatDef())
+			{
+				OutValue = SystemConst->GetExposedParameters().GetParameterValueFromOffset<FNiagaraFloat>(UserParameterWithOffset.Offset).Value;
+				return true;
+			}
+		}
+		return false;
+	}
+
+	bool FUserParameterAdapter::TryGetValue(FVector2f& OutValue) const
+	{
+		const UNiagaraSystem* SystemConst = SystemConstWeak.Get();
+		if (SystemConst != nullptr)
+		{
+			if (UserParameterWithOffset.GetType() == FNiagaraTypeDefinition::GetVec2Def())
+			{
+				OutValue = SystemConst->GetExposedParameters().GetParameterValueFromOffset<FVector2f>(UserParameterWithOffset.Offset);
+				return true;
+			}
+		}
+		return false;
+	}
+
+	bool FUserParameterAdapter::TryGetValue(FVector3f& OutValue) const
+	{
+		const UNiagaraSystem* SystemConst = SystemConstWeak.Get();
+		if (SystemConst != nullptr)
+		{
+			if (UserParameterWithOffset.GetType() == FNiagaraTypeDefinition::GetVec3Def())
+			{
+				OutValue = SystemConst->GetExposedParameters().GetParameterValueFromOffset<FVector3f>(UserParameterWithOffset.Offset);
+				return true;
+			}
+		}
+		return false;
+	}
+
+	bool FUserParameterAdapter::TryGetValue(FVector4f& OutValue) const
+	{
+		const UNiagaraSystem* SystemConst = SystemConstWeak.Get();
+		if (SystemConst != nullptr)
+		{
+			if (UserParameterWithOffset.GetType() == FNiagaraTypeDefinition::GetVec4Def())
+			{
+				OutValue = SystemConst->GetExposedParameters().GetParameterValueFromOffset<FVector4f>(UserParameterWithOffset.Offset);
+				return true;
+			}
+		}
+		return false;
+	}
+
+	bool FUserParameterAdapter::TryGetValue(FVector& OutValue) const
+	{
+		const UNiagaraSystem* SystemConst = SystemConstWeak.Get();
+		if (SystemConst != nullptr)
+		{
+			if (UserParameterWithOffset.GetType() == FNiagaraTypeDefinition::GetPositionDef())
+			{
+				const FVector* Position = SystemConst->GetExposedParameters().GetPositionParameterValue(UserParameterWithOffset.GetName());
+				if (Position != nullptr)
+				{
+					OutValue = *Position;
+					return true;
+				}
+			}
+		}
+		return false;
+	}
+
+	bool FUserParameterAdapter::TryGetValue(FLinearColor& OutValue) const
+	{
+		const UNiagaraSystem* SystemConst = SystemConstWeak.Get();
+		if (SystemConst != nullptr)
+		{
+			if (UserParameterWithOffset.GetType() == FNiagaraTypeDefinition::GetColorDef())
+			{
+				OutValue = SystemConst->GetExposedParameters().GetParameterValueFromOffset<FLinearColor>(UserParameterWithOffset.Offset);
+				return true;
+			}
+		}
+		return false;
+	}
+
+	bool FUserParameterAdapter::TryGetValue(FQuat4f& OutValue) const
+	{
+		const UNiagaraSystem* SystemConst = SystemConstWeak.Get();
+		if (SystemConst != nullptr)
+		{
+			if (UserParameterWithOffset.GetType() == FNiagaraTypeDefinition::GetQuatDef())
+			{
+				OutValue = SystemConst->GetExposedParameters().GetParameterValueFromOffset<FQuat4f>(UserParameterWithOffset.Offset);
+				return true;
+			}
+		}
+		return false;
+	}
+
+	bool FUserParameterAdapter::TryGetValueData(TNotNull<UStruct*> ValueType, const uint8*& OutData) const
+	{
+		const UNiagaraSystem* SystemConst = SystemConstWeak.Get();
+		if (SystemConst != nullptr)
+		{
+			if (UserParameterWithOffset.GetType().GetStruct() == ValueType)
+			{
+				const uint8* Data = SystemConst->GetExposedParameters().GetParameterData(UserParameterWithOffset.Offset, ValueType->GetStructureSize());
+				if (Data != nullptr)
+				{
+					OutData = Data;
+					return true;
+				}
+			}
+		}
+		return false;
+	}
+
+	void FUserParameterAdapter::Initialize(UNiagaraSystem* InSystem, const FNiagaraVariableWithOffset& InUserParameterWithOffset)
+	{
+		SystemWeak = InSystem;
+		SystemConstWeak = InSystem;
+		UserParameterWithOffset = InUserParameterWithOffset;
+	}
+
+	void FUserParameterAdapter::Initialize(const UNiagaraSystem* InSystem, const FNiagaraVariableWithOffset& InUserParameterWithOffset)
+	{
+		SystemConstWeak = InSystem;
+		UserParameterWithOffset = InUserParameterWithOffset;
+	}
+
+	class FUserParameterAdapterFactory : public TAdapterRefListOneSource<UNiagaraSystem, FUserParameterAdapter>::IFactory
+	{
+	public:
+		virtual void CreateAdapters(UNiagaraSystem* SourceObject, TArray<TSharedRef<FUserParameterAdapter>>& OutAdapters) const override
+		{
+			CreateAdaptersInternal(SourceObject, OutAdapters);
+		}
+
+		virtual void CreateAdapters(const UNiagaraSystem* SourceObject, TArray<TSharedRef<FUserParameterAdapter>>& OutAdapters) const override
+		{
+			CreateAdaptersInternal(SourceObject, OutAdapters);
+		}
+
+	private:
+		template<typename TSystem>
+		static void CreateAdaptersInternal(TSystem* SourceObject, TArray<TSharedRef<FUserParameterAdapter>>& OutAdapters)
+		{
+			if (SourceObject != nullptr && SourceObject->GetExposedParameters().Num() > 0)
+			{
+				for (const FNiagaraVariableWithOffset& UserParameter : SourceObject->GetExposedParameters().ReadParameterVariables())
+				{
+					OutAdapters.Add(FUserParameterAdapter::Create(SourceObject, UserParameter));
+				}
+			}
+		}
+	};
+
+	class FEmitterAdapterFactory : public TAdapterRefListOneSource<UNiagaraSystem, FEmitterAdapter>::IFactory
+	{
+	public:
+		virtual void CreateAdapters(UNiagaraSystem* InOwningSystem, TArray<TSharedRef<FEmitterAdapter>>& OutAdapters) const override
+		{
+			CreateAdaptersInternal(InOwningSystem, OutAdapters);
+		}
+
+		virtual void CreateAdapters(const UNiagaraSystem* InOwningSystem, TArray<TSharedRef<FEmitterAdapter>>& OutAdapters) const override
+		{
+			CreateAdaptersInternal(InOwningSystem, OutAdapters);
+		}
+
+	private:
+		template<typename TSystem>
+		static void CreateAdaptersInternal(TSystem* InOwningSystem, TArray<TSharedRef<FEmitterAdapter>>& OutAdapters)
+		{
+			if (InOwningSystem == nullptr)
+			{
+				return;
+			}
+			for (const FNiagaraEmitterHandle& EmitterHandle : InOwningSystem->GetEmitterHandles())
+			{
+				OutAdapters.Add(FEmitterAdapter::Create(InOwningSystem, EmitterHandle));
+			}
+		}
+	};
+
+	class FSystemScriptAdapterFactory : public TAdapterPtr<UNiagaraSystem, FStackScriptAdapter>::IFactory
+	{
+	public:
+		FSystemScriptAdapterFactory(ENiagaraScriptUsage InUsage)
+			: Usage(InUsage)
+		{
+		}
+
+		virtual TSharedRef<FStackScriptAdapter> CreateAdapter(UNiagaraSystem* InSourceObject) const override
+		{
+			return CreateAdapterInternal(InSourceObject);
+		}
+
+		virtual TSharedRef<FStackScriptAdapter> CreateAdapter(const UNiagaraSystem* InSourceObject) const override
+		{
+			return CreateAdapterInternal(InSourceObject);
+		}
+
+	private:
+		template<typename TSystem>
+		TSharedRef<FStackScriptAdapter> CreateAdapterInternal(TSystem InSystem) const
+		{
+			if (InSystem != nullptr)
+			{
+				if (Usage == ENiagaraScriptUsage::SystemSpawnScript)
+				{
+					return FStackScriptAdapter::Create(InSystem, FGuid(), InSystem->GetSystemSpawnScript());
+				}
+				else if (Usage == ENiagaraScriptUsage::SystemUpdateScript)
+				{
+					return FStackScriptAdapter::Create(InSystem, FGuid(), InSystem->GetSystemUpdateScript());
+				}
+			}
+			return MakeShared<FStackScriptAdapter>();
+		}
+
+	private:
+		ENiagaraScriptUsage Usage;
+	};
+
+	TSharedRef<FSystemAdapter> FSystemAdapter::Create(UNiagaraSystem* InSystem)
+	{
+		TSharedRef<FSystemAdapter> SystemAdapter = MakeShared<FSystemAdapter>();
+		SystemAdapter->Initialize(InSystem);
+		return SystemAdapter;
+	}
+
+	TSharedRef<const FSystemAdapter> FSystemAdapter::CreateConst(const UNiagaraSystem* InSystem)
+	{
+		TSharedRef<FSystemAdapter> SystemAdapter = MakeShared<FSystemAdapter>();
+		SystemAdapter->Initialize(InSystem);
+		return SystemAdapter;
+	}
+
+	bool FSystemAdapter::IsValidAdapter() const
+	{
+		return SystemConstWeak.Get() != nullptr;
+	}
+
+	bool FSystemAdapter::IsValidWriteAdapter() const
+	{
+		return SystemWeak.IsSet() && SystemWeak->Get() != nullptr;
+	}
+
+	TSharedRefCollection<FUserParameterAdapter> FSystemAdapter::GetUserParameters()
+	{
+		return UserParameters.Get();
+	}
+
+	TConstSharedRefCollection<FUserParameterAdapter> FSystemAdapter::GetUserParameters() const
+	{
+		return UserParameters.Get();
+	}
+
+	TSharedPtr<FUserParameterAdapter> FSystemAdapter::GetUserParameterByName(const FName& InName)
+	{
+		return UserParameters.Get().FindByPredicate([&InName](const TSharedRef<FUserParameterAdapter>& UserParameterAdapter)
+			{ return UserParameterAdapter->GetName() == InName; });
+	}
+
+	TSharedPtr<const FUserParameterAdapter> FSystemAdapter::GetUserParameterByName(const FName& InName) const
+	{
+		return UserParameters.Get().FindByPredicate([&InName](const TSharedRef<const FUserParameterAdapter>& UserParameterAdapter)
+			{ return UserParameterAdapter->GetName() == InName; });
+	}
+
+	TSharedPtr<FUserParameterAdapter> FSystemAdapter::GetUserParameterByNamespacedName(const FName& InNamespacedName)
+	{
+		return UserParameters.Get().FindByPredicate([&InNamespacedName](const TSharedRef<FUserParameterAdapter>& UserParameterAdapter)
+			{ return UserParameterAdapter->GetNamespacedName() == InNamespacedName; });
+	}
+
+	TSharedPtr<const FUserParameterAdapter> FSystemAdapter::GetUserParameterByNamespacedName(const FName& InNamespacedName) const
+	{
+		return UserParameters.Get().FindByPredicate([&InNamespacedName](const TSharedRef<const FUserParameterAdapter>& UserParameterAdapter)
+			{ return UserParameterAdapter->GetNamespacedName() == InNamespacedName; });
+	}
+
+	TSharedRefCollection<FEmitterAdapter> FSystemAdapter::GetEmitters()
+	{
+		return Emitters.Get();
+	}
+
+	TConstSharedRefCollection<FEmitterAdapter> FSystemAdapter::GetEmitters() const
+	{
+		return Emitters.Get();
+	}
+
+	TSharedPtr<FEmitterAdapter> FSystemAdapter::GetEmitterByName(FName InEmitterName)
+	{
+		return Emitters.Get().FindByPredicate([&InEmitterName](const TSharedRef<FEmitterAdapter>& EmitterAdapter)
+			{ return EmitterAdapter->GetName() == InEmitterName; });
+	}
+
+	TSharedPtr<const FEmitterAdapter> FSystemAdapter::GetEmitterByName(FName InEmitterName) const
+	{
+		return Emitters.Get().FindByPredicate([&InEmitterName](const TSharedRef<const FEmitterAdapter>& EmitterAdapter)
+			{ return EmitterAdapter->GetName() == InEmitterName; });
+	}
+
+	TSharedRef<FStackScriptAdapter> FSystemAdapter::GetSpawnScript()
+	{
+		return SpawnScript.Get();
+	}
+
+	TSharedRef<const FStackScriptAdapter> FSystemAdapter::GetSpawnScript() const
+	{
+		return SpawnScript.Get();
+	}
+
+	TSharedRef<FStackScriptAdapter> FSystemAdapter::GetUpdateScript()
+	{
+		return UpdateScript.Get();
+	}
+
+	TSharedRef<const FStackScriptAdapter> FSystemAdapter::GetUpdateScript() const
+	{
+		return UpdateScript.Get();
+	}
+
+	int32 FSystemAdapter::GetNumScripts() const
+	{
+		return 2;
+	}
+
+	TSharedRef<FStackScriptAdapter> FSystemAdapter::GetScriptAt(int32 Index)
+	{
+		checkf(Index > INDEX_NONE && Index < GetNumScripts(), TEXT("Index out of range"));
+		if (Index == 0)
+		{
+			return GetSpawnScript();
+		}
+		else // Index == 1
+		{
+			return GetUpdateScript();
+		}
+	}
+
+	FScriptSharedRefCollection FSystemAdapter::GetScripts()
+	{
+		return FScriptSharedRefCollection(&StackScriptAdapterCollection);
+	}
+
+	FScriptConstSharedRefCollection FSystemAdapter::GetScripts() const
+	{
+		return FScriptConstSharedRefCollection(&StackScriptAdapterCollection);
+	}
+
+	void FSystemAdapter::Initialize(UNiagaraSystem* InSystem)
+	{
+		SystemWeak = InSystem;
+		SystemConstWeak = InSystem;
+		UserParameters.Initialize(InSystem, MakeShared<FUserParameterAdapterFactory>());
+		Emitters.Initialize(InSystem, MakeShared<FEmitterAdapterFactory>());
+		SpawnScript.Initialize(InSystem, MakeShared<FSystemScriptAdapterFactory>(ENiagaraScriptUsage::SystemSpawnScript));
+		UpdateScript.Initialize(InSystem, MakeShared<FSystemScriptAdapterFactory>(ENiagaraScriptUsage::SystemUpdateScript));
+		StackScriptAdapterCollection.Initialize(StaticCastSharedRef<IStackScriptAdapterOwner>(this->AsShared()));
+	}
+
+	void FSystemAdapter::Initialize(const UNiagaraSystem* InSystem)
+	{
+		SystemConstWeak = InSystem;
+		UserParameters.Initialize(InSystem, MakeShared<FUserParameterAdapterFactory>());
+		Emitters.Initialize(InSystem, MakeShared<FEmitterAdapterFactory>());
+		SpawnScript.Initialize(InSystem, MakeShared<FSystemScriptAdapterFactory>(ENiagaraScriptUsage::SystemSpawnScript));
+		UpdateScript.Initialize(InSystem, MakeShared<FSystemScriptAdapterFactory>(ENiagaraScriptUsage::SystemUpdateScript));
+		StackScriptAdapterCollection.Initialize(StaticCastSharedRef<IStackScriptAdapterOwner>(this->AsShared()));
+	}
+
+} // namespace UE::Niagara
\ No newline at end of file
diff --git a/Engine/Plugins/FX/Niagara/Source/NiagaraEditor/Public/SystemAdapter/AdapterShared.h b/Engine/Plugins/FX/Niagara/Source/NiagaraEditor/Public/SystemAdapter/AdapterShared.h
new file mode 100644
--- /dev/null
+++ b/Engine/Plugins/FX/Niagara/Source/NiagaraEditor/Public/SystemAdapter/AdapterShared.h
@@ -0,0 +1,408 @@
+// Copyright Epic Games, Inc. All Rights Reserved.
+
+#pragma once
+
+#include "Misc/Optional.h"
+#include "Templates/SharedPointer.h"
+#include "UObject/WeakObjectPtr.h"
+#include "UObject/WeakObjectPtrTemplates.h"
+
+template<typename TSourceItem, typename TTargetItem, typename TIndexable>
+class TSharedRefCollectionBase
+{
+public:
+	class FIterator
+	{
+	public:
+		FIterator(const TIndexable* InSource, int32 InIndex)
+			: Source(InSource)
+			, Index(InIndex)
+		{
+		}
+
+		bool operator!=(const FIterator& Other)
+		{
+			return Source != Other.Source || Index != Other.Index;
+		}
+
+		FIterator& operator++()
+		{
+			Index++;
+			return *this;
+		}
+
+		const TSharedRef<TTargetItem> operator*() const
+		{
+			return (*Source)[Index];
+		}
+
+	private:
+		const TIndexable* Source = nullptr;
+		int32 Index = INDEX_NONE;
+	};
+
+	TSharedRefCollectionBase()
+	{
+	}
+
+	TSharedRefCollectionBase(const TIndexable* InSource)
+		: Source(InSource)
+	{
+	}
+
+	FIterator begin() const
+	{
+		return Source != nullptr ? FIterator(Source, 0) : FIterator(nullptr, INDEX_NONE);
+	}
+
+	FIterator end() const
+	{
+		return Source != nullptr ? FIterator(Source, Source->Num()) : FIterator(nullptr, INDEX_NONE);
+	}
+
+	int32 Num() const 
+	{
+		return Source != nullptr ? Source->Num() : 0;
+	}
+
+	TSharedRef<TTargetItem> operator[](int32 Index) const
+	{
+		checkf(Source != nullptr, TEXT("Source is invalid"));
+		checkf(Index >= 0 && Index < Source->Num(), TEXT("Index out of range"));
+		return (*Source)[Index];
+	}
+
+	template<typename TPredicate>
+	int32 IndexOfByPredicate(TPredicate Predicate)
+	{
+		return Source != nullptr ? Source->IndexOfByPredicate(Predicate) : INDEX_NONE;
+	}
+
+	template<typename TPredicate>
+	TSharedPtr<TTargetItem> FindByPredicate(TPredicate Predicate)
+	{
+		int32 FoundIndex = IndexOfByPredicate(Predicate);
+		return Source != nullptr && FoundIndex != INDEX_NONE ? (*Source)[FoundIndex] : TSharedPtr<TTargetItem>();
+	}
+
+	template<typename TPredicate>
+	void FindAllByPredicate(TArray<TSharedRef<TTargetItem>>& OutFound, TPredicate Predicate)
+	{
+		if (Source == nullptr)
+		{
+			return;
+		}
+		for (const TSharedRef<TTargetItem>& Item : (*Source))
+		{
+			if (Predicate(Item))
+			{
+				OutFound.Add(Item);
+			}
+		}
+	}
+
+private:
+	const TIndexable* Source = nullptr;
+};
+
+template<typename TItem>
+class TSharedRefCollection : public TSharedRefCollectionBase<TItem, TItem, TArray<TSharedRef<TItem>>>
+{
+}
+;
+
+template<typename TItem>
+class TConstSharedRefCollection : public TSharedRefCollectionBase<TItem, const TItem, TArray<TSharedRef<TItem>>>
+{
+};
+
+class IAdapter
+{
+public:
+	virtual ~IAdapter() = default;
+	virtual bool IsValidAdapter() const = 0;
+	virtual bool IsValidWriteAdapter() const = 0;
+};
+
+template<typename TSourceObject, typename TAdapter>
+class TAdapterRefListOneSource
+{
+public:
+	class IFactory
+	{
+	public:
+		virtual ~IFactory() = default;
+		virtual void CreateAdapters(TSourceObject* SourceObject, TArray<TSharedRef<TAdapter>>& OutAdapters) const = 0;
+		virtual void CreateAdapters(const TSourceObject* SourceObject, TArray<TSharedRef<TAdapter>>& OutAdapters) const = 0;
+	};
+
+	void Initialize(TSourceObject* SourceObject, TSharedRef<const IFactory> InAdapterFactory)
+	{
+		SourceWeak = SourceObject;
+		AdapterFactory = InAdapterFactory;
+	}
+
+	void Initialize(const TSourceObject* SourceObject, TSharedRef<const IFactory> InAdapterFactory)
+	{
+		SourceConstWeak = SourceObject;
+		AdapterFactory = InAdapterFactory;
+	}
+
+	TSharedRefCollection<TAdapter> Get()
+	{
+		if (ensureMsgf(SourceWeak.IsSet(), TEXT("Adapter is only valid for const use.")))
+		{
+			if (AdapterList.IsSet() == false)
+			{
+				TArray<TSharedRef<TAdapter>> Adapters;
+				AdapterFactory->CreateAdapters(SourceWeak.GetValue().Get(), Adapters);
+				AdapterList.Emplace(Adapters);
+			}
+		}
+		else
+		{
+			if (AdapterList.IsSet() == false)
+			{
+				AdapterList = TArray<TSharedRef<TAdapter>>();
+			}
+		}
+		return TSharedRefCollection<TAdapter>(&AdapterList.GetValue());
+	}
+
+	TConstSharedRefCollection<TAdapter> Get() const
+	{
+		if (AdapterList.IsSet() == false)
+		{
+			TArray<TSharedRef<TAdapter>> Adapters;
+			if (SourceWeak.IsSet())
+			{
+				AdapterFactory->CreateAdapters(SourceWeak.GetValue().Get(), Adapters);
+			}
+			else
+			{
+				AdapterFactory->CreateAdapters(SourceConstWeak.Get(), Adapters);
+			}
+			AdapterList.Emplace(Adapters);
+		}
+		return TConstSharedRefCollection<TAdapter>(&AdapterList.GetValue());
+	}
+
+private:
+	TOptional<TWeakObjectPtr<TSourceObject>> SourceWeak;
+	TWeakObjectPtr<const TSourceObject> SourceConstWeak;
+	TSharedPtr<const IFactory> AdapterFactory;
+
+	mutable TOptional<TArray<TSharedRef<TAdapter>>> AdapterList;
+};
+
+template<typename TSourceObject1, typename TSourceObject2, typename TAdapter>
+class TAdapterRefListTwoSources
+{
+public:
+	class IFactory
+	{
+	public:
+		virtual ~IFactory() = default;
+		virtual void CreateAdapters(TSourceObject1* SourceObject1, TSourceObject2* SourceObject2, TArray<TSharedRef<TAdapter>>& OutAdapters) const = 0;
+		virtual void CreateAdapters(const TSourceObject1* SourceObject1, const TSourceObject2* SourceObject2, TArray<TSharedRef<TAdapter>>& OutAdapters) const = 0;
+	};
+
+	void Initialize(TSourceObject1* SourceObject1, TSourceObject2* SourceObject2, TSharedRef<const IFactory> InAdapterFactory)
+	{
+		SourceWeak1 = SourceObject1;
+		SourceWeak2 = SourceObject2;
+		AdapterFactory = InAdapterFactory;
+	}
+
+	void Initialize(const TSourceObject1* SourceObject1, const TSourceObject2* SourceObject2, TSharedRef<const IFactory> InAdapterFactory)
+	{
+		SourceConstWeak1 = SourceObject1;
+		SourceConstWeak2 = SourceObject2;
+		AdapterFactory = InAdapterFactory;
+	}
+
+	TSharedRefCollection<TAdapter> Get()
+	{
+		if (ensureMsgf(SourceWeak1.IsSet() && SourceWeak2.IsSet(), TEXT("Adapter is only valid for const use.")))
+		{
+			if (AdapterList.IsSet() == false)
+			{
+				TArray<TSharedRef<TAdapter>> Adapters;
+				AdapterFactory->CreateAdapters(SourceWeak1.GetValue().Get(), SourceWeak2.GetValue().Get(), Adapters);
+				AdapterList.Emplace(Adapters);
+			}
+		}
+		else
+		{
+			if (AdapterList.IsSet() == false)
+			{
+				AdapterList = TArray<TSharedRef<TAdapter>>();
+			}
+		}
+		return TSharedRefCollection<TAdapter>(&AdapterList.GetValue());
+	}
+
+	TConstSharedRefCollection<TAdapter> Get() const
+	{
+		if (AdapterList.IsSet() == false)
+		{
+			TArray<TSharedRef<TAdapter>> Adapters;
+			if (SourceWeak1.IsSet() && SourceWeak2.IsSet())
+			{
+				AdapterFactory->CreateAdapters(SourceWeak1.GetValue().Get(), SourceWeak2.GetValue().Get(), Adapters);
+			}
+			else
+			{
+				AdapterFactory->CreateAdapters(SourceConstWeak1.Get(), SourceConstWeak2.Get(), Adapters);
+			}
+			AdapterList.Emplace(Adapters);
+		}
+		return TConstSharedRefCollection<TAdapter>(&AdapterList.GetValue());
+	}
+
+private:
+	TOptional<TWeakObjectPtr<TSourceObject1>> SourceWeak1;
+	TOptional<TWeakObjectPtr<TSourceObject2>> SourceWeak2;
+	TWeakObjectPtr<const TSourceObject1> SourceConstWeak1;
+	TWeakObjectPtr<const TSourceObject2> SourceConstWeak2;
+	TSharedPtr<const IFactory> AdapterFactory;
+
+	mutable TOptional<TArray<TSharedRef<TAdapter>>> AdapterList;
+};
+
+template<typename TSourceObject1, typename TSourceObject2, typename TSourceObject3, typename TAdapter>
+class TAdapterRefListThreeSources
+{
+public:
+	class IFactory
+	{
+	public:
+		virtual ~IFactory() = default;
+		virtual void CreateAdapters(TSourceObject1* SourceObject1, TSourceObject2* SourceObject2, TSourceObject3* SourceObject3, TArray<TSharedRef<TAdapter>>& OutAdapters) const = 0;
+		virtual void CreateAdapters(const TSourceObject1* SourceObject1, const TSourceObject2* SourceObject2, const TSourceObject3* SourceObject3, TArray<TSharedRef<TAdapter>>& OutAdapters) const = 0;
+	};
+
+	void Initialize(TSourceObject1* SourceObject1, TSourceObject2* SourceObject2, TSourceObject3* SourceObject3, TSharedRef<const IFactory> InAdapterFactory)
+	{
+		SourceWeak1 = SourceObject1;
+		SourceWeak2 = SourceObject2;
+		SourceWeak3 = SourceObject3;
+		AdapterFactory = InAdapterFactory;
+	}
+
+	void Initialize(const TSourceObject1* SourceObject1, const TSourceObject2* SourceObject2, const TSourceObject3* SourceObject3, TSharedRef<const IFactory> InAdapterFactory)
+	{
+		SourceConstWeak1 = SourceObject1;
+		SourceConstWeak2 = SourceObject2;
+		SourceConstWeak3 = SourceObject3;
+		AdapterFactory = InAdapterFactory;
+	}
+
+	TSharedRefCollection<TAdapter> Get()
+	{
+		if (ensureMsgf(SourceWeak1.IsSet() && SourceWeak2.IsSet() && SourceWeak3.IsSet(), TEXT("Adapter is only valid for const use.")))
+		{
+			if (AdapterList.IsSet() == false)
+			{
+				TArray<TSharedRef<TAdapter>> Adapters;
+				AdapterFactory->CreateAdapters(SourceWeak1.GetValue().Get(), SourceWeak2.GetValue().Get(), SourceWeak3.GetValue().Get(), Adapters);
+				AdapterList.Emplace(Adapters);
+			}
+		}
+		else
+		{
+			if (AdapterList.IsSet() == false)
+			{
+				AdapterList = TArray<TSharedRef<TAdapter>>();
+			}
+		}
+		return TSharedRefCollection<TAdapter>(&AdapterList.GetValue());
+	}
+
+	TConstSharedRefCollection<TAdapter> Get() const
+	{
+		if (AdapterList.IsSet() == false)
+		{
+			TArray<TSharedRef<TAdapter>> Adapters;
+			if (SourceWeak1.IsSet() && SourceWeak2.IsSet() && SourceWeak3.IsSet())
+			{
+				AdapterFactory->CreateAdapters(SourceWeak1.GetValue().Get(), SourceWeak2.GetValue().Get(), SourceWeak3.GetValue().Get(), Adapters);
+			}
+			else
+			{
+				AdapterFactory->CreateAdapters(SourceConstWeak1.Get(), SourceConstWeak2.Get(), SourceConstWeak3.Get(), Adapters);
+			}
+			AdapterList.Emplace(Adapters);
+		}
+		return TConstSharedRefCollection<TAdapter>(&AdapterList.GetValue());
+	}
+
+private:
+	TOptional<TWeakObjectPtr<TSourceObject1>> SourceWeak1;
+	TOptional<TWeakObjectPtr<TSourceObject2>> SourceWeak2;
+	TOptional<TWeakObjectPtr<TSourceObject3>> SourceWeak3;
+	TWeakObjectPtr<const TSourceObject1> SourceConstWeak1;
+	TWeakObjectPtr<const TSourceObject2> SourceConstWeak2;
+	TWeakObjectPtr<const TSourceObject3> SourceConstWeak3;
+	TSharedPtr<const IFactory> AdapterFactory;
+
+	mutable TOptional<TArray<TSharedRef<TAdapter>>> AdapterList;
+};
+
+template<typename TSourceObject, typename TAdapter>
+class TAdapterPtr
+{
+public:
+	class IFactory
+	{
+	public:
+		virtual ~IFactory() = default;
+		virtual TSharedRef<TAdapter> CreateAdapter(TSourceObject* InSourceObject) const = 0;
+		virtual TSharedRef<TAdapter> CreateAdapter(const TSourceObject* InSourceObject) const = 0;
+	};
+
+	void Initialize(TSourceObject* InSourceObject, TSharedRef<const IFactory> InAdapterFactory)
+	{
+		SourceWeak = InSourceObject;
+		AdapterFactory = InAdapterFactory;
+	}
+
+	void Initialize(const TSourceObject* InSourceObject, TSharedRef<const IFactory> InAdapterFactory)
+	{
+		SourceConstWeak = InSourceObject;
+		AdapterFactory = InAdapterFactory;
+	}
+
+	TSharedRef<TAdapter> Get()
+	{
+		checkf(SourceWeak.IsSet(), TEXT("Adapter is only valid for const use."));
+		if (Adapter.IsSet() == false)
+		{
+			Adapter.Emplace(AdapterFactory->CreateAdapter(SourceWeak.GetValue().Get()));
+		}
+		return Adapter.GetValue();
+	}
+
+	TSharedRef<const TAdapter> Get() const
+	{
+		if (Adapter.IsSet() == false)
+		{
+			if (SourceWeak.IsSet())
+			{
+				Adapter.Emplace(AdapterFactory->CreateAdapter(SourceWeak.GetValue().Get()));
+			}
+			else
+			{
+				Adapter.Emplace(AdapterFactory->CreateAdapter(SourceConstWeak.Get()));
+			}
+		}
+		return Adapter.GetValue();
+	}
+
+private:
+	TOptional<TWeakObjectPtr<TSourceObject>> SourceWeak;
+	TWeakObjectPtr<const TSourceObject> SourceConstWeak;
+	TSharedPtr<const IFactory> AdapterFactory;
+
+	mutable TOptional<TSharedRef<TAdapter>> Adapter;
+};
\ No newline at end of file
diff --git a/Engine/Plugins/FX/Niagara/Source/NiagaraEditor/Public/SystemAdapter/EmitterAdapter.h b/Engine/Plugins/FX/Niagara/Source/NiagaraEditor/Public/SystemAdapter/EmitterAdapter.h
new file mode 100644
--- /dev/null
+++ b/Engine/Plugins/FX/Niagara/Source/NiagaraEditor/Public/SystemAdapter/EmitterAdapter.h
@@ -0,0 +1,81 @@
+// Copyright Epic Games, Inc. All Rights Reserved.
+
+#pragma once
+
+#include "Misc/Guid.h"
+#include "SystemAdapter/AdapterShared.h"
+#include "SystemAdapter/StackScriptAdapterCollection.h"
+
+struct FNiagaraEmitterHandle;
+class  UNiagaraSystem;
+
+namespace UE::Niagara
+{
+	class FStackScriptAdapter;
+
+	class FEmitterAdapter : public TSharedFromThis<FEmitterAdapter>, public IAdapter, public IStackScriptAdapterOwner
+	{
+	public:
+		FEmitterAdapter()
+		{
+		}
+
+		NIAGARAEDITOR_API static TSharedRef<FEmitterAdapter> Create(UNiagaraSystem* InOwningSystem, const FNiagaraEmitterHandle& InEmitterHandle);
+
+		static TSharedRef<FEmitterAdapter> Create(const UNiagaraSystem* InOwningSystem, const FNiagaraEmitterHandle& InEmitterHandle);
+
+		NIAGARAEDITOR_API static TSharedRef<const FEmitterAdapter> CreateConst(const UNiagaraSystem* InOwningSystem, const FNiagaraEmitterHandle& InEmitterHandle);
+
+		NIAGARAEDITOR_API virtual bool IsValidAdapter() const override;
+
+		NIAGARAEDITOR_API virtual bool IsValidWriteAdapter() const override;
+
+		NIAGARAEDITOR_API FName GetName() const { return EmitterName; }
+
+		NIAGARAEDITOR_API TSharedRef<FStackScriptAdapter> GetEmitterSpawnScript();
+
+		NIAGARAEDITOR_API TSharedRef<const FStackScriptAdapter> GetEmitterSpawnScript() const;
+
+		NIAGARAEDITOR_API TSharedRef<FStackScriptAdapter> GetEmitterUpdateScript();
+
+		NIAGARAEDITOR_API TSharedRef<const FStackScriptAdapter> GetEmitterUpdateScript() const;
+
+		NIAGARAEDITOR_API TSharedRef<FStackScriptAdapter> GetParticleSpawnScript();
+
+		NIAGARAEDITOR_API TSharedRef<const FStackScriptAdapter> GetParticleSpawnScript() const;
+
+		NIAGARAEDITOR_API TSharedRef<FStackScriptAdapter> GetParticleUpdateScript();
+
+		NIAGARAEDITOR_API TSharedRef<const FStackScriptAdapter> GetParticleUpdateScript() const;
+
+		virtual int32 GetNumScripts() const override;
+		virtual TSharedRef<FStackScriptAdapter> GetScriptAt(int32 Index) override;
+
+		NIAGARAEDITOR_API FScriptSharedRefCollection GetScripts();
+
+		NIAGARAEDITOR_API FScriptConstSharedRefCollection GetScripts() const;
+
+	private:
+		void Initialize(UNiagaraSystem* InOwningSystem, const FNiagaraEmitterHandle& InEmitterHandle);
+
+		void Initialize(const UNiagaraSystem* InOwningSystem, const FNiagaraEmitterHandle& InEmitterHandle);
+
+	private:
+		TOptional<TWeakObjectPtr<UNiagaraSystem>> OwningSystemWeak;
+		TWeakObjectPtr<const UNiagaraSystem> OwningSystemConstWeak;
+		FName EmitterName;
+		FGuid EmitterHandleId;
+
+		TAdapterPtr<UNiagaraSystem, FStackScriptAdapter> EmitterSpawnScript;
+		TAdapterPtr<UNiagaraSystem, FStackScriptAdapter> EmitterUpdateScript;
+
+		TAdapterPtr<UNiagaraSystem, FStackScriptAdapter> ParticleSpawnScript;
+		TAdapterPtr<UNiagaraSystem, FStackScriptAdapter> ParticleUpdateScript;
+
+		FStackScriptAdapterCollection StackScriptAdapterCollection;
+	};
+
+	typedef TSharedRef<FEmitterAdapter> FEmitterAdapterRef;
+	typedef TSharedRef<const FEmitterAdapter> FEmitterAdapterConstRef;
+
+} // namespace UE::Niagara
\ No newline at end of file
diff --git a/Engine/Plugins/FX/Niagara/Source/NiagaraEditor/Public/SystemAdapter/RangeQuery.h b/Engine/Plugins/FX/Niagara/Source/NiagaraEditor/Public/SystemAdapter/RangeQuery.h
new file mode 100644
--- /dev/null
+++ b/Engine/Plugins/FX/Niagara/Source/NiagaraEditor/Public/SystemAdapter/RangeQuery.h
@@ -0,0 +1,488 @@
+// Copyright Epic Games, Inc. All Rights Reserved.
+
+#pragma once
+
+#include "Delegates/Delegate.h"
+#include "Delegates/DelegateCombinations.h"
+#include "Templates/SharedPointer.h"
+
+namespace UE::RangeQuery
+{
+
+template<typename TItem>
+class IRangeQueryIterator
+{
+public:
+	virtual ~IRangeQueryIterator() = default;
+
+	virtual bool IsValidIterator() const = 0;
+	virtual void NextItem() = 0;
+	virtual TItem GetItem() const = 0;
+};
+
+template<typename TSourceItem, typename TTargetItem>
+class IRangeQueryWrapperIterator : public IRangeQueryIterator<TTargetItem>
+{
+public:
+	virtual void SetSourceIterator(TSharedRef<IRangeQueryIterator<TSourceItem>> InSourceIterator) = 0;
+};
+
+template<typename TItem>
+class TRangeQuery
+{
+public:
+	class FIterator
+	{
+	public:
+		FIterator(IRangeQueryIterator<TItem>* InSourceIterator)
+			: SourceIterator(InSourceIterator)
+		{
+		}
+
+		bool operator!=(const FIterator& Other) const
+		{
+			return IsValidIterator() != Other.IsValidIterator();
+		}
+
+		FIterator& operator++()
+		{
+			SourceIterator->NextItem();
+			return *this;
+		}
+
+		TItem operator*() const
+		{
+			return SourceIterator->GetItem();
+		}
+
+	private:
+		bool IsValidIterator() const
+		{
+			return SourceIterator != nullptr && SourceIterator->IsValidIterator();
+		}
+
+	private:
+		IRangeQueryIterator<TItem>* SourceIterator = nullptr;
+	};
+
+	TRangeQuery()
+	{
+	}
+
+	TRangeQuery(TSharedRef<IRangeQueryIterator<TItem>> InSourceIterator)
+		: SourceIterator(InSourceIterator)
+	{
+	}
+
+	const TSharedPtr<IRangeQueryIterator<TItem>>& GetSourceIterator() const { return SourceIterator; }
+	void SetSourceIterator(TSharedRef<IRangeQueryIterator<TItem>> InSourceIterator)
+	{
+		SourceIterator = InSourceIterator;
+	}
+
+	FIterator begin() const
+	{
+		return FIterator(SourceIterator.Get());
+	}
+
+	FIterator end() const
+	{
+		return FIterator(nullptr);
+	}
+
+	TRangeQuery<TItem>& Add(TSharedRef<IRangeQueryWrapperIterator<TItem, TItem>> InWrapper)
+	{
+		checkf(SourceIterator.IsValid(), TEXT("Source iterator is not set."));
+		InWrapper->SetSourceIterator(SourceIterator.ToSharedRef());
+		SourceIterator = InWrapper;
+		return *this;
+	}
+
+	template<typename TTargetItem>
+	TRangeQuery<TTargetItem> Add(TSharedRef<IRangeQueryWrapperIterator<TItem, TTargetItem>> InWrapper)
+	{
+		checkf(SourceIterator.IsValid(), TEXT("Source iterator is not set."));
+		InWrapper->SetSourceIterator(SourceIterator.ToSharedRef());
+		return TRangeQuery<TTargetItem>(InWrapper);
+	}
+
+	TRangeQuery<TItem>& operator+(TSharedRef<IRangeQueryWrapperIterator<TItem, TItem>> InWrapper)
+	{
+		checkf(SourceIterator.IsValid(), TEXT("Source iterator is not set."));
+		InWrapper->SetSourceIterator(SourceIterator.ToSharedRef());
+		SourceIterator = InWrapper;
+		return *this;
+	}
+
+	template<typename TTargetItem>
+	TRangeQuery<TTargetItem> operator+(TSharedRef<IRangeQueryWrapperIterator<TItem, TTargetItem>> InWrapper)
+	{
+		checkf(SourceIterator.IsValid(), TEXT("Source iterator is not set."));
+		InWrapper->SetSourceIterator(SourceIterator.ToSharedRef());
+		return TRangeQuery<TTargetItem>(InWrapper);
+	}
+
+	TItem First() const
+	{
+		checkf(SourceIterator.IsValid(), TEXT("Source iterator is not set."));
+		for (const TItem& Item : *this)
+		{
+			return Item;
+		}
+		return nullptr;
+	}
+
+	template<typename TDefault>
+	TDefault FirstOrDefault(TDefault InDefault) const
+	{
+		checkf(SourceIterator.IsValid(), TEXT("Source iterator is not set."));
+		for (const TItem& Item : *this)
+		{
+			return Item;
+		}
+		return InDefault;
+	}
+
+	template<typename TArrayItem>
+	void AppendToArray(TArray<TArrayItem>& OutArray)
+	{
+		checkf(SourceIterator.IsValid(), TEXT("Source iterator is not set."));
+		for (const TItem& Item : *this)
+		{
+			OutArray.Add(Item);
+		}
+	}
+
+private:
+	TSharedPtr<IRangeQueryIterator<TItem>> SourceIterator;
+};
+
+template<typename TItem, typename TIndexable>
+class TRangeQueryIndexableIterator : public IRangeQueryIterator<TItem>
+{
+public:
+	static TSharedRef<TRangeQueryIndexableIterator<TItem, TIndexable>> CreateAsRef(const TIndexable* InIndexableRange)
+	{
+		TSharedRef<TRangeQueryIndexableIterator<TItem, TIndexable>> Iterator = MakeShared<TRangeQueryIndexableIterator<TItem, TIndexable>>();
+		Iterator->SetSource(InIndexableRange, false);
+		return Iterator;
+	}
+
+	static TSharedRef<TRangeQueryIndexableIterator<TItem, TIndexable>> CreateAsCopy(const TIndexable& InIndexableRange)
+	{
+		TIndexable* IndexableRangeCopy = new TIndexable();
+		*IndexableRangeCopy = InIndexableRange;
+		
+		TSharedRef<TRangeQueryIndexableIterator<TItem, TIndexable>> Iterator = MakeShared<TRangeQueryIndexableIterator<TItem, TIndexable>>();
+		Iterator->SetSource(IndexableRangeCopy, true);
+		return Iterator;
+	}
+
+	TRangeQueryIndexableIterator()
+	{
+	}
+
+	~TRangeQueryIndexableIterator()
+	{
+		if (bOwnsRange)
+		{
+			delete IndexableRange;
+		}
+	}
+
+	virtual bool IsValidIterator() const override
+	{
+		return Index >= 0 && IndexableRange != nullptr && Index < IndexableRange->Num();
+	}
+
+	virtual void NextItem() override
+	{
+		Index++;
+	}
+
+	virtual TItem GetItem() const override
+	{
+		return (*IndexableRange)[Index];
+	}
+
+private:
+	void SetSource(const TIndexable* InIndexableRange, bool bInOwnsRange)
+	{
+		IndexableRange = InIndexableRange;
+		Index = 0;
+		bOwnsRange = bInOwnsRange;
+	}
+
+private:
+	const TIndexable* IndexableRange = nullptr;
+	int32 Index = INDEX_NONE;
+	bool bOwnsRange = false;
+};
+
+template<typename TItem>
+class TRangeQueryPredicateIterator : public IRangeQueryWrapperIterator<TItem, TItem>
+{
+public:
+	DECLARE_DELEGATE_RetVal_OneParam(bool, FPredicate, const TItem& /* Item */);
+
+public:
+	TRangeQueryPredicateIterator(TSharedRef<IRangeQueryIterator<TItem>> InSourceIterator, FPredicate InPredicate)
+		: SourceIterator(InSourceIterator)
+		, Predicate(InPredicate)
+	{
+	}
+
+	TRangeQueryPredicateIterator(FPredicate InPredicate)
+		: Predicate(InPredicate)
+	{
+	}
+
+	virtual void SetSourceIterator(TSharedRef<IRangeQueryIterator<TItem>> InSourceIterator) override
+	{
+		SourceIterator = InSourceIterator;
+	}
+
+	virtual bool IsValidIterator() const override
+	{
+		checkf(SourceIterator.IsValid(), TEXT("Source iterator is not set."));
+		if (bMovedToInitialValidItem == false)
+		{
+			if (SourceIterator->IsValidIterator() && Predicate.Execute(SourceIterator->GetItem()) == false)
+			{
+				MoveToNextValidItem();
+			}
+			bMovedToInitialValidItem = true;
+		}
+		return SourceIterator->IsValidIterator();
+	}
+
+	virtual void NextItem() override
+	{
+		checkf(SourceIterator.IsValid(), TEXT("Source iterator is not set."));
+		MoveToNextValidItem();
+	}
+
+	virtual TItem GetItem() const override
+	{
+		checkf(SourceIterator.IsValid(), TEXT("Source iterator is not set."));
+		return SourceIterator->GetItem();
+	}
+
+private:
+	void MoveToNextValidItem() const
+	{
+		SourceIterator->NextItem();
+		while (SourceIterator->IsValidIterator() && Predicate.Execute(SourceIterator->GetItem()) == false)
+		{
+			SourceIterator->NextItem();
+		}
+	}
+
+private:
+	TSharedPtr<IRangeQueryIterator<TItem>> SourceIterator;
+	FPredicate Predicate;
+	mutable bool bMovedToInitialValidItem = false;
+};
+
+template<typename TItem, typename TSelectedItem>
+class TRangeQuerySelectIterator : public IRangeQueryWrapperIterator<TItem, TSelectedItem>
+{
+public:
+	DECLARE_DELEGATE_RetVal_OneParam(TSelectedItem, FSelecter, const TItem&);
+
+public:
+	TRangeQuerySelectIterator(TSharedRef<IRangeQueryIterator<TItem>> InSourceIterator, FSelecter InSelecter)
+		: SourceIterator(InSourceIterator)
+		, Selecter(InSelecter)
+	{
+	}
+
+	TRangeQuerySelectIterator(FSelecter InSelecter)
+		: Selecter(InSelecter)
+	{
+	}
+
+	virtual void SetSourceIterator(TSharedRef<IRangeQueryIterator<TItem>> InSourceIterator) override
+	{
+		SourceIterator = InSourceIterator;
+	}
+
+	virtual bool IsValidIterator() const override
+	{
+		checkf(SourceIterator.IsValid(), TEXT("Source iterator is not set."));
+		return SourceIterator->IsValidIterator();
+	}
+
+	virtual void NextItem() override
+	{
+		checkf(SourceIterator.IsValid(), TEXT("Source iterator is not set."));
+		SourceIterator->NextItem();
+	}
+
+	virtual TSelectedItem GetItem() const override
+	{
+		checkf(SourceIterator.IsValid(), TEXT("Source iterator is not set."));
+		return Selecter.Execute(SourceIterator->GetItem());
+	}
+
+private:
+	TSharedPtr<IRangeQueryIterator<TItem>> SourceIterator;
+	FSelecter Selecter;
+};
+
+template<typename TItem, typename TSelectedItem>
+class TRangeQuerySelectRangeIterator : public IRangeQueryWrapperIterator<TItem, TSelectedItem>
+{
+public:
+	DECLARE_DELEGATE_RetVal_OneParam(TSharedRef<IRangeQueryIterator<TSelectedItem>>, FRangeSelecter, TItem);
+
+public:
+	TRangeQuerySelectRangeIterator(TSharedRef<IRangeQueryIterator<TItem>> InSourceIterator, FRangeSelecter InRangeSelecter)
+		: SourceIterator(InSourceIterator)
+		, RangeSelecter(InRangeSelecter)
+	{
+	}
+
+	TRangeQuerySelectRangeIterator(FRangeSelecter InRangeSelecter)
+		: RangeSelecter(InRangeSelecter)
+	{
+
+	}
+
+	virtual void SetSourceIterator(TSharedRef<IRangeQueryIterator<TItem>> InSourceIterator) override
+	{
+		SourceIterator = InSourceIterator;
+	}
+
+	virtual bool IsValidIterator() const override
+	{
+		checkf(SourceIterator.IsValid(), TEXT("Source iterator is not set."));
+		if (bMovedToInitialValidRangeIterator == false)
+		{
+			if (SourceIterator->IsValidIterator())
+			{
+				SelectedRangeIterator.Emplace(RangeSelecter.Execute(SourceIterator->GetItem()));
+				if (SelectedRangeIterator.GetValue()->IsValidIterator() == false)
+				{
+					MoveToNextValidRangeIterator();
+				}
+			}
+			bMovedToInitialValidRangeIterator = true;
+		}
+		return SelectedRangeIterator.IsSet() && SelectedRangeIterator.GetValue()->IsValidIterator();
+	}
+
+	virtual void NextItem() override
+	{
+		checkf(SourceIterator.IsValid(), TEXT("Source iterator is not set."));
+		checkf(SelectedRangeIterator.IsSet(), TEXT("Iterator is not valid."));
+		SelectedRangeIterator.GetValue()->NextItem();
+		if (SelectedRangeIterator.GetValue()->IsValidIterator() == false && SourceIterator->IsValidIterator())
+		{
+			MoveToNextValidRangeIterator();
+		}
+	}
+
+	virtual TSelectedItem GetItem() const override
+	{
+		checkf(SourceIterator.IsValid(), TEXT("Source iterator is not set."));
+		checkf(SelectedRangeIterator.IsSet(), TEXT("Iterator is not valid."));
+		return SelectedRangeIterator.GetValue()->GetItem();
+	}
+
+private:
+	void MoveToNextValidRangeIterator() const
+	{
+		SourceIterator->NextItem();
+		if (SourceIterator->IsValidIterator())
+		{
+			SelectedRangeIterator.Emplace(RangeSelecter.Execute(SourceIterator->GetItem()));
+			while (SelectedRangeIterator.GetValue()->IsValidIterator() == false && SourceIterator->IsValidIterator())
+			{
+				SourceIterator->NextItem();
+				if (SourceIterator->IsValidIterator())
+				{
+					SelectedRangeIterator.Emplace(RangeSelecter.Execute(SourceIterator->GetItem()));
+				}
+			}
+		}
+	}
+
+private:
+	TSharedPtr<IRangeQueryIterator<TItem>> SourceIterator;
+	mutable bool bMovedToInitialValidRangeIterator = false;
+	mutable TOptional<TSharedRef<IRangeQueryIterator<TSelectedItem>>> SelectedRangeIterator;
+	FRangeSelecter RangeSelecter;
+};
+
+template<typename TItem>
+TSharedRef<IRangeQueryWrapperIterator<TItem, TItem>> Where(typename TRangeQueryPredicateIterator<TItem>::FPredicate Predicate)
+{
+	return MakeShared<TRangeQueryPredicateIterator<TItem>>(Predicate);
+};
+
+template<typename TItem, typename TPredicate>
+TSharedRef<IRangeQueryWrapperIterator<TItem, TItem>> Where(TPredicate Predicate)
+{
+	return MakeShared<TRangeQueryPredicateIterator<TItem>>(
+		TRangeQueryPredicateIterator<TItem>::FPredicate::CreateLambda([Predicate](const TItem& Item) { return Predicate(Item); }));
+};
+
+template<typename TItem>
+TSharedRef<IRangeQueryWrapperIterator<TItem, TItem>> WhereValid()
+{
+	return MakeShared<TRangeQueryPredicateIterator<TItem>>(
+		TRangeQueryPredicateIterator<TItem>::FPredicate::CreateLambda([](const TItem& Item) { return Item.IsValid(); }));
+}
+
+template<typename TItem, typename TSelectedItem>
+TSharedRef<IRangeQueryWrapperIterator<TItem, TSelectedItem>> Select(typename TRangeQuerySelectIterator<TItem, TSelectedItem>::FSelecter Selecter)
+{
+	return MakeShared<TRangeQuerySelectIterator<TItem, TSelectedItem>>(Selecter);
+};
+
+template<typename TItem, typename TSelectedItem, typename TSelecter>
+TSharedRef<IRangeQueryWrapperIterator<TItem, TSelectedItem>> Select(TSelecter Selecter)
+{
+	return MakeShared<TRangeQuerySelectIterator<TItem, TSelectedItem>>(
+		TRangeQuerySelectIterator<TItem, TSelectedItem>::FSelecter::CreateLambda([Selecter](const TItem& Item) { return Selecter(Item); }));
+};
+
+template<typename TItem, typename TSelectedItem>
+TSharedRef<IRangeQueryWrapperIterator<TItem, TSelectedItem>> SelectRange(typename TRangeQuerySelectRangeIterator<TItem, TSelectedItem>::FRangeSelecter RangeSelecter)
+{
+	return MakeShared<TRangeQuerySelectRangeIterator<TItem, TSelectedItem>>(RangeSelecter);
+};
+
+template<typename TItem, typename TSelectedItem, typename TRangeSelecter>
+TSharedRef<IRangeQueryWrapperIterator<TItem, TSelectedItem>> SelectRange(TRangeSelecter RangeSelecter)
+{
+	return MakeShared<TRangeQuerySelectRangeIterator<TItem, TSelectedItem>>(
+		TRangeQuerySelectRangeIterator<TItem, TSelectedItem>::FRangeSelecter::CreateLambda([RangeSelecter](const TItem& Item) { return RangeSelecter(Item); }));
+};
+
+template<typename TItem, typename TIterable>
+TSharedRef<IRangeQueryIterator<TItem>> MakeIterator(const TIterable* Iterable)
+{
+	return TRangeQueryIndexableIterator<TItem, TIterable>::CreateAsRef(Iterable);
+}
+
+template<typename TItem, typename TIterable>
+TSharedRef<IRangeQueryIterator<TItem>> MakeIterator(const TIterable& Iterable)
+{
+	return TRangeQueryIndexableIterator<TItem, TIterable>::CreateAsCopy(Iterable);
+}
+
+template<typename TItem, typename TIterable>
+TRangeQuery<TItem> MakeQuery(const TIterable* Iterable)
+{
+	return TRangeQuery<TItem>(MakeIterator<TItem, TIterable>(Iterable));
+}
+
+template<typename TItem, typename TIterable>
+TRangeQuery<TItem> MakeQuery(const TIterable& Iterable)
+{
+	return TRangeQuery<TItem>(MakeIterator<TItem, TIterable>(Iterable));
+}
+
+} // namespace UE::RangeQuery
diff --git a/Engine/Plugins/FX/Niagara/Source/NiagaraEditor/Public/SystemAdapter/StackFunctionAdapter.h b/Engine/Plugins/FX/Niagara/Source/NiagaraEditor/Public/SystemAdapter/StackFunctionAdapter.h
new file mode 100644
--- /dev/null
+++ b/Engine/Plugins/FX/Niagara/Source/NiagaraEditor/Public/SystemAdapter/StackFunctionAdapter.h
@@ -0,0 +1,71 @@
+// Copyright Epic Games, Inc. All Rights Reserved.
+
+#pragma once
+
+#include "SystemAdapter/AdapterShared.h"
+#include "NiagaraTypes.h"
+
+class UNiagaraNodeFunctionCall;
+class UNiagaraScript;
+class UNiagaraSystem;
+
+namespace UE::Niagara
+{
+	class FStackFunctionInputAdapter;
+
+	class FStackFunctionAdapter : public IAdapter
+	{
+	public:
+		enum class EFunctionType
+		{
+			ScriptModule,
+			SetParametersModule,
+		};
+
+		NIAGARAEDITOR_API static TSharedRef<FStackFunctionAdapter> Create(UNiagaraSystem* InOwningSystem, FGuid InOwningEmitterHandleId, UNiagaraScript* InOwningScript, UNiagaraNodeFunctionCall* InFunctionCallNode);
+
+		static TSharedRef<FStackFunctionAdapter> Create(const UNiagaraSystem* InOwningSystem, FGuid InOwningEmitterHandleId, const UNiagaraScript* InOwningScript, const UNiagaraNodeFunctionCall* InFunctionCallNode);
+
+		NIAGARAEDITOR_API static TSharedRef<const FStackFunctionAdapter> CreateConst(const UNiagaraSystem* InOwningSystem, FGuid InOwningEmitterHandleId, const UNiagaraScript* InOwningScript, const UNiagaraNodeFunctionCall* InFunctionCallNode);
+
+		NIAGARAEDITOR_API virtual bool IsValidAdapter() const override;
+
+		NIAGARAEDITOR_API virtual bool IsValidWriteAdapter() const override;
+
+		EFunctionType GetType() const { return Type; }
+
+		NIAGARAEDITOR_API FString GetFunctionName() const;
+
+		NIAGARAEDITOR_API UNiagaraScript* GetModuleScript() const;
+
+		NIAGARAEDITOR_API void GetSetParameters(TArray<FNiagaraVariableBase>& OutParameters) const;
+
+		NIAGARAEDITOR_API bool SetsParameter(const FNiagaraVariableBase& InParameter) const;
+
+		NIAGARAEDITOR_API TSharedRefCollection<FStackFunctionInputAdapter> GetInputs();
+
+		NIAGARAEDITOR_API TConstSharedRefCollection<FStackFunctionInputAdapter> GetInputs() const;
+
+		NIAGARAEDITOR_API TSharedPtr<FStackFunctionInputAdapter> FindInputByName(FName InputName);
+
+		NIAGARAEDITOR_API TSharedPtr<const FStackFunctionInputAdapter> FindInputByName(FName InputName) const;
+
+	private:
+		void Initialize(UNiagaraSystem* InOwningSystem, FGuid InOwningEmitterHandleId, UNiagaraScript* InOwningScript, UNiagaraNodeFunctionCall* InFunctionCallNode);
+
+		void Initialize(const UNiagaraSystem* InOwningSystem, FGuid InOwningEmitterHandleId, const UNiagaraScript* InOwningScript, const UNiagaraNodeFunctionCall* InFunctionCallNode);
+
+	private:
+		TOptional<TWeakObjectPtr<UNiagaraNodeFunctionCall>> FunctionCallNodeWeak;
+		TWeakObjectPtr<const UNiagaraNodeFunctionCall> FunctionCallNodeConstWeak;
+		EFunctionType Type = EFunctionType::ScriptModule;
+
+		TAdapterRefListThreeSources<UNiagaraSystem, UNiagaraScript, UNiagaraNodeFunctionCall, FStackFunctionInputAdapter> Inputs;
+	};
+
+	typedef TSharedRef<FStackFunctionAdapter> FStackFunctionAdapterRef;
+	typedef TSharedRef<const FStackFunctionAdapter> FStackFunctionAdapterConstRef;
+	typedef TSharedPtr<FStackFunctionAdapter> FStackFunctionAdapterPtr;
+	typedef TSharedPtr<const FStackFunctionAdapter> FStackFunctionAdapterConstPtr;
+
+} // namespace UE::Niagara
\ No newline at end of file
diff --git a/Engine/Plugins/FX/Niagara/Source/NiagaraEditor/Public/SystemAdapter/StackFunctionInputAdapter.h b/Engine/Plugins/FX/Niagara/Source/NiagaraEditor/Public/SystemAdapter/StackFunctionInputAdapter.h
new file mode 100644
--- /dev/null
+++ b/Engine/Plugins/FX/Niagara/Source/NiagaraEditor/Public/SystemAdapter/StackFunctionInputAdapter.h
@@ -0,0 +1,96 @@
+// Copyright Epic Games, Inc. All Rights Reserved.
+
+#pragma once
+
+#include "NiagaraTypes.h"
+#include "SystemAdapter/AdapterShared.h"
+
+class UNiagaraDataInterface;
+class UNiagaraNodeFunctionCall;
+
+namespace UE::Niagara
+{
+	class FStackFunctionAdapter;
+
+	class FStackFunctionInputAdapter : public IAdapter
+	{
+	public:
+		enum class EValueMode
+		{
+			LocalValue,
+			LinkedValue,
+			DataValue,
+			ObjectValue,
+			DynamicValue,
+			Unknown
+		};
+
+		NIAGARAEDITOR_API static TSharedRef<FStackFunctionInputAdapter> Create(UNiagaraNodeFunctionCall* InFunctionCallNode, const FNiagaraVariableBase& InInputParameter);
+
+		static TSharedRef<FStackFunctionInputAdapter> Create(const UNiagaraNodeFunctionCall* InFunctionCallNode, const FNiagaraVariableBase& InInputParameter);
+
+		NIAGARAEDITOR_API static TSharedRef<const FStackFunctionInputAdapter> CreateConst(const UNiagaraNodeFunctionCall* InFunctionCallNode, const FNiagaraVariableBase& InInputParameter);
+
+		NIAGARAEDITOR_API virtual bool IsValidAdapter() const override;
+
+		NIAGARAEDITOR_API virtual bool IsValidWriteAdapter() const override;
+
+		NIAGARAEDITOR_API FName GetName() const;
+
+		NIAGARAEDITOR_API const FNiagaraTypeDefinition& GetType() const;
+
+		NIAGARAEDITOR_API EValueMode GetValueMode() const { return ValueMode; }
+
+		NIAGARAEDITOR_API bool TryGetLocalValue(bool& OutValue) const;
+
+		NIAGARAEDITOR_API bool TryGetLocalValue(int32& OutValue) const;
+
+		NIAGARAEDITOR_API bool TryGetLocalValue(float& OutValue) const;
+
+		NIAGARAEDITOR_API bool TryGetLocalValue(FVector2f& OutValue) const;
+
+		NIAGARAEDITOR_API bool TryGetLocalValue(FVector3f& OutValue) const;
+
+		NIAGARAEDITOR_API bool TryGetLocalValue(FVector4f& OutValue) const;
+
+		NIAGARAEDITOR_API bool TryGetLocalValueData(TNotNull<UStruct*> ValueType, const uint8*& OutData) const;
+
+		template<typename TValue>
+		bool TryGetLocalValue(TValue& OutValue) const
+		{
+			const uint8* ValueData;
+			if (TryGetLocalValueData(TValue::StaticStruct(), ValueData))
+			{
+				TValue::StaticStruct()->CopyScriptStruct(&OutValue, ValueData);
+				return true;
+			}
+			return false;
+		}
+
+		NIAGARAEDITOR_API FNiagaraVariableBase GetLinkedValue() const;
+
+		NIAGARAEDITOR_API UNiagaraDataInterface* GetDataValue() const;
+
+		NIAGARAEDITOR_API UObject* GetObjectValue() const;
+
+		NIAGARAEDITOR_API TSharedPtr<FStackFunctionAdapter> GetDynamicValue() const;
+
+	private:
+		void Initialize(UNiagaraNodeFunctionCall* InFunctionCallNode, const FNiagaraVariableBase& InInputParameter);
+
+		void Initialize(const UNiagaraNodeFunctionCall* InFunctionCallNode, const FNiagaraVariableBase& InInputParameter);
+
+	private:
+		TOptional<TWeakObjectPtr<UNiagaraNodeFunctionCall>> FunctionCallNodeWeak;
+		TWeakObjectPtr<const UNiagaraNodeFunctionCall> FunctionCallNodeConstWeak;
+		FNiagaraVariableBase InputParameter;
+		FName InputName;
+		EValueMode ValueMode = EValueMode::Unknown;
+	};
+
+	typedef TSharedRef<FStackFunctionInputAdapter> FStackFunctionInputAdapterRef;
+	typedef TSharedRef<const FStackFunctionInputAdapter> FStackFunctionInputAdapterConstRef;
+	typedef TSharedPtr<FStackFunctionInputAdapter> FStackFunctionInputAdapterPtr;
+	typedef TSharedPtr<const FStackFunctionInputAdapter> FStackFunctionInputAdapterConstPtr;
+
+} // namespace UE::Niagara
\ No newline at end of file
diff --git a/Engine/Plugins/FX/Niagara/Source/NiagaraEditor/Public/SystemAdapter/StackScriptAdapter.h b/Engine/Plugins/FX/Niagara/Source/NiagaraEditor/Public/SystemAdapter/StackScriptAdapter.h
new file mode 100644
--- /dev/null
+++ b/Engine/Plugins/FX/Niagara/Source/NiagaraEditor/Public/SystemAdapter/StackScriptAdapter.h
@@ -0,0 +1,66 @@
+// Copyright Epic Games, Inc. All Rights Reserved.
+
+#pragma once
+
+#include "AdapterShared.h"
+#include "NiagaraCommon.h"
+#include "NiagaraTypes.h"
+
+class UNiagaraNodeFunctionCall;
+class UNiagaraScript;
+
+namespace UE::Niagara
+{
+	class FStackFunctionAdapter;
+
+	class FStackScriptAdapter : public IAdapter
+	{
+	public:
+		NIAGARAEDITOR_API static TSharedRef<FStackScriptAdapter> Create(UNiagaraSystem* InOwningSystem, const FGuid& InOwningEmitterHandleId, UNiagaraScript* InScript);
+
+		static TSharedRef<FStackScriptAdapter> Create(const UNiagaraSystem* InOwningSystem, const FGuid& InOwningEmitterHandleId, const UNiagaraScript* InScript);
+
+		NIAGARAEDITOR_API static TSharedRef<const FStackScriptAdapter> CreateConst(const UNiagaraSystem* InOwningSystem, const FGuid& InOwningEmitterHandleId, const UNiagaraScript* InScript);
+
+		NIAGARAEDITOR_API virtual bool IsValidAdapter() const override;
+
+		NIAGARAEDITOR_API virtual bool IsValidWriteAdapter() const override;
+
+		NIAGARAEDITOR_API ENiagaraScriptUsage GetUsage() const;
+
+		NIAGARAEDITOR_API FGuid GetUsageId() const;
+
+		NIAGARAEDITOR_API TSharedRefCollection<FStackFunctionAdapter> GetModules();
+
+		NIAGARAEDITOR_API TConstSharedRefCollection<FStackFunctionAdapter> GetModules() const;
+
+		NIAGARAEDITOR_API TSharedPtr<FStackFunctionAdapter> FindModuleByName(const FName& InName);
+
+		NIAGARAEDITOR_API TSharedPtr<const FStackFunctionAdapter> FindModuleByName(const FName& InName) const;
+
+		NIAGARAEDITOR_API void FindSetParametersModulesBySetParameter(const FNiagaraVariableBase& InParameter, TArray<TSharedRef<FStackFunctionAdapter>>& OutSetParametersModules);
+
+		NIAGARAEDITOR_API void FindSetParametersModulesBySetParameter(const FNiagaraVariableBase& InParameter, TArray<TSharedRef<const FStackFunctionAdapter>>& OutSetParametersModules) const;
+
+		NIAGARAEDITOR_API TSharedPtr<FStackFunctionAdapter> FindFirstSetParametersModuleBySetParameter(const FNiagaraVariableBase& InParameter);
+
+		NIAGARAEDITOR_API TSharedPtr<const FStackFunctionAdapter> FindFirstSetParametersModuleBySetParameter(const FNiagaraVariableBase& InParameter) const;
+
+	private:
+		void Initialize(UNiagaraSystem* InOwningSystem, const FGuid& InOwningEmitterHandleId, UNiagaraScript* InScript);
+
+		void Initialize(const UNiagaraSystem* InOwningSystem, const FGuid& InOwningEmitterHandleId, const UNiagaraScript* InScript);
+
+	private:
+		TOptional<TWeakObjectPtr<UNiagaraScript>> ScriptWeak;
+		TWeakObjectPtr<const UNiagaraScript> ScriptConstWeak;
+
+		TAdapterRefListTwoSources<UNiagaraSystem, UNiagaraScript, FStackFunctionAdapter> Modules;
+	};
+
+	typedef TSharedRef<FStackScriptAdapter> FStackScriptAdapterRef;
+	typedef TSharedRef<const FStackScriptAdapter> FStackScriptAdapterConstRef;
+	typedef TSharedPtr<FStackScriptAdapter> FStackScriptAdapterPtr;
+	typedef TSharedPtr<const FStackScriptAdapter> FStackScriptAdapterConstPtr;
+
+} // namespace UE::Niagara
\ No newline at end of file
diff --git a/Engine/Plugins/FX/Niagara/Source/NiagaraEditor/Public/SystemAdapter/StackScriptAdapterCollection.h b/Engine/Plugins/FX/Niagara/Source/NiagaraEditor/Public/SystemAdapter/StackScriptAdapterCollection.h
new file mode 100644
--- /dev/null
+++ b/Engine/Plugins/FX/Niagara/Source/NiagaraEditor/Public/SystemAdapter/StackScriptAdapterCollection.h
@@ -0,0 +1,41 @@
+// Copyright Epic Games, Inc. All Rights Reserved.
+
+#pragma once
+
+#include "SystemAdapter/AdapterShared.h"
+
+namespace UE::Niagara
+{
+	class FStackScriptAdapter;
+
+	class IStackScriptAdapterOwner
+	{
+	public:
+		virtual ~IStackScriptAdapterOwner() = default;
+
+		virtual int32 GetNumScripts() const = 0;
+		virtual TSharedRef<FStackScriptAdapter> GetScriptAt(int32 Index) = 0;
+	};
+
+	class FStackScriptAdapterCollection
+	{
+	public:
+		void Initialize(TSharedRef<IStackScriptAdapterOwner> InScriptOwner);
+
+		NIAGARAEDITOR_API int32 Num() const;
+
+		NIAGARAEDITOR_API TSharedRef<FStackScriptAdapter> operator[](int32 Index) const;
+
+	private:
+		TWeakPtr<IStackScriptAdapterOwner> ScriptOwnerWeak;
+	};
+
+	class FScriptSharedRefCollection : public TSharedRefCollectionBase<FStackScriptAdapter, FStackScriptAdapter, FStackScriptAdapterCollection>
+	{
+	};
+	
+	class FScriptConstSharedRefCollection : public TSharedRefCollectionBase<FStackScriptAdapter, const FStackScriptAdapter, FStackScriptAdapterCollection>
+	{
+	};
+
+} // namespace UE::Niagara
\ No newline at end of file
diff --git a/Engine/Plugins/FX/Niagara/Source/NiagaraEditor/Public/SystemAdapter/SystemAdapter.h b/Engine/Plugins/FX/Niagara/Source/NiagaraEditor/Public/SystemAdapter/SystemAdapter.h
new file mode 100644
--- /dev/null
+++ b/Engine/Plugins/FX/Niagara/Source/NiagaraEditor/Public/SystemAdapter/SystemAdapter.h
@@ -0,0 +1,155 @@
+// Copyright Epic Games, Inc. All Rights Reserved.
+
+#pragma once
+
+#include "NiagaraParameterStore.h"
+#include "SystemAdapter/AdapterShared.h"
+#include "SystemAdapter/StackScriptAdapterCollection.h"
+#include "Templates/SharedPointer.h"
+
+class  UNiagaraSystem;
+
+namespace UE::Niagara
+{
+	class FEmitterAdapter;
+	class FParameterAdapter;
+	class FStackScriptAdapter;
+	class IStackScriptAdapterCollection;
+
+	class FUserParameterAdapter : public IAdapter
+	{
+		friend class FUserParameterAdapterList;
+
+	public:
+		FUserParameterAdapter()
+		{
+		}
+
+		NIAGARAEDITOR_API static TSharedRef<FUserParameterAdapter> Create(UNiagaraSystem* InSystem, const FNiagaraVariableWithOffset& InUserParameterWithOffset);
+
+		static TSharedRef<FUserParameterAdapter> Create(const UNiagaraSystem* InSystem, const FNiagaraVariableWithOffset& InUserParameterWithOffset);
+
+		NIAGARAEDITOR_API static TSharedRef<const FUserParameterAdapter> CreateConst(const UNiagaraSystem* InSystem, const FNiagaraVariableWithOffset& InUserParameterWithOffset);
+
+		NIAGARAEDITOR_API virtual bool IsValidAdapter() const override;
+
+		NIAGARAEDITOR_API virtual bool IsValidWriteAdapter() const override;
+
+		NIAGARAEDITOR_API FName GetName() const;
+
+		NIAGARAEDITOR_API FName GetNamespacedName() const { return UserParameterWithOffset.GetName(); }
+
+		NIAGARAEDITOR_API bool TryGetValue(bool &OutValue) const;
+
+		NIAGARAEDITOR_API bool TryGetValue(int32& OutValue) const;
+
+		NIAGARAEDITOR_API bool TryGetValue(float& OutValue) const;
+
+		NIAGARAEDITOR_API bool TryGetValue(FVector2f& OutValue) const;
+
+		NIAGARAEDITOR_API bool TryGetValue(FVector3f& OutValue) const;
+
+		NIAGARAEDITOR_API bool TryGetValue(FVector4f& OutValue) const;
+
+		NIAGARAEDITOR_API bool TryGetValue(FVector& OutValue) const;
+
+		NIAGARAEDITOR_API bool TryGetValue(FLinearColor& OutValue) const;
+
+		NIAGARAEDITOR_API bool TryGetValue(FQuat4f& OutValue) const;
+
+		NIAGARAEDITOR_API bool TryGetValue(FNiagaraPosition& OutValue) const;
+
+		NIAGARAEDITOR_API bool TryGetValueData(TNotNull<UStruct*> ValueType, const uint8*& OutData) const;
+		
+		template<typename TValue>
+		bool TryGetValue(TValue& OutValue) const
+		{
+			const uint8* ValueData;
+			if (TryGetValueData(TValue::StaticStruct(), ValueData))
+			{
+				TValue::StaticStruct()->CopyScriptStruct(&OutValue, ValueData);
+				return true;
+			}
+			return false;
+		}
+
+	private:
+		void Initialize(UNiagaraSystem* InSystem, const FNiagaraVariableWithOffset& InUserParameterWithOffset);
+
+		void Initialize(const UNiagaraSystem* InSystem, const FNiagaraVariableWithOffset& InUserParameterWithOffset);
+
+	private:
+		TOptional<TWeakObjectPtr<UNiagaraSystem>> SystemWeak;
+		TWeakObjectPtr<const UNiagaraSystem> SystemConstWeak;
+		FNiagaraVariableWithOffset UserParameterWithOffset;
+		mutable FName NameCache = NAME_None;
+	};
+
+	class FSystemAdapter : public TSharedFromThis<FSystemAdapter>, public IAdapter, public IStackScriptAdapterOwner
+	{
+	public:
+		FSystemAdapter()
+		{
+		}
+
+		NIAGARAEDITOR_API static TSharedRef<FSystemAdapter> Create(UNiagaraSystem* InSystem);
+
+		NIAGARAEDITOR_API static TSharedRef<const FSystemAdapter> CreateConst(const UNiagaraSystem* InSystem);
+
+		NIAGARAEDITOR_API virtual bool IsValidAdapter() const override;
+
+		NIAGARAEDITOR_API virtual bool IsValidWriteAdapter() const override;
+
+		NIAGARAEDITOR_API TSharedRefCollection<FUserParameterAdapter> GetUserParameters();
+	
+		NIAGARAEDITOR_API TConstSharedRefCollection<FUserParameterAdapter> GetUserParameters() const;
+
+		NIAGARAEDITOR_API TSharedPtr<FUserParameterAdapter> GetUserParameterByName(const FName& InName);
+
+		NIAGARAEDITOR_API TSharedPtr<const FUserParameterAdapter> GetUserParameterByName(const FName& InName) const;
+
+		NIAGARAEDITOR_API TSharedPtr<FUserParameterAdapter> GetUserParameterByNamespacedName(const FName& InNamespacedName);
+
+		NIAGARAEDITOR_API TSharedPtr<const FUserParameterAdapter> GetUserParameterByNamespacedName(const FName& InNamespacedName) const;
+
+		NIAGARAEDITOR_API TSharedRefCollection<FEmitterAdapter> GetEmitters();
+
+		NIAGARAEDITOR_API TConstSharedRefCollection<FEmitterAdapter> GetEmitters() const;
+
+		NIAGARAEDITOR_API TSharedPtr<FEmitterAdapter> GetEmitterByName(FName InEmitterName);
+
+		NIAGARAEDITOR_API TSharedPtr<const FEmitterAdapter> GetEmitterByName(FName InEmitterName) const;
+
+		NIAGARAEDITOR_API TSharedRef<FStackScriptAdapter> GetSpawnScript();
+
+		NIAGARAEDITOR_API TSharedRef<const FStackScriptAdapter> GetSpawnScript() const;
+
+		NIAGARAEDITOR_API TSharedRef<FStackScriptAdapter> GetUpdateScript();
+
+		NIAGARAEDITOR_API TSharedRef<const FStackScriptAdapter> GetUpdateScript() const;
+
+		virtual int32 GetNumScripts() const override;
+
+		virtual TSharedRef<FStackScriptAdapter> GetScriptAt(int32 Index) override;
+
+		NIAGARAEDITOR_API FScriptSharedRefCollection GetScripts();
+
+		NIAGARAEDITOR_API FScriptConstSharedRefCollection GetScripts() const;
+
+	private:
+		void Initialize(UNiagaraSystem* InSystem);
+
+		void Initialize(const UNiagaraSystem* InSystem);
+
+	private:
+		TOptional<TWeakObjectPtr<UNiagaraSystem>> SystemWeak;
+		TWeakObjectPtr<const UNiagaraSystem> SystemConstWeak;
+
+		TAdapterRefListOneSource<UNiagaraSystem, FUserParameterAdapter> UserParameters;
+		TAdapterRefListOneSource<UNiagaraSystem, FEmitterAdapter> Emitters;
+		TAdapterPtr<UNiagaraSystem, FStackScriptAdapter> SpawnScript;
+		TAdapterPtr<UNiagaraSystem, FStackScriptAdapter> UpdateScript;
+		FStackScriptAdapterCollection StackScriptAdapterCollection;
+	};
+
+} // namespace UE::Niagara
\ No newline at end of file
