- Export functions for plugin renderers
- With this change the sprite renderer (for example) could be made purely inside a plugin

#Synced-CL 51499580

diff --git a/Engine/Plugins/FX/Niagara/Source/Niagara/Classes/NiagaraGPUInstanceCountManager.h b/Engine/Plugins/FX/Niagara/Source/Niagara/Classes/NiagaraGPUInstanceCountManager.h
--- a/Engine/Plugins/FX/Niagara/Source/Niagara/Classes/NiagaraGPUInstanceCountManager.h
+++ b/Engine/Plugins/FX/Niagara/Source/Niagara/Classes/NiagaraGPUInstanceCountManager.h
@@ -98,7 +98,7 @@ public:
 	bool HasPendingGPUReadback() const;
 
 	/** Add a draw indirect task to generate the draw indirect args. Returns the draw indirect arg buffer offset. */
-	FIndirectArgSlot AddDrawIndirect(FRHICommandListBase& RHICmdList, uint32 InstanceCountBufferOffset, uint32 NumIndicesPerInstance, uint32 StartIndexLocation, bool bIsInstancedStereoEnabled, bool bCulled, ENiagaraGpuComputeTickStage::Type ReadyTickStage);
+	NIAGARA_API FIndirectArgSlot AddDrawIndirect(FRHICommandListBase& RHICmdList, uint32 InstanceCountBufferOffset, uint32 NumIndicesPerInstance, uint32 StartIndexLocation, bool bIsInstancedStereoEnabled, bool bCulled, ENiagaraGpuComputeTickStage::Type ReadyTickStage);
 
 	// Resize instance count and draw indirect buffers to ensure it is big enough to hold all draw indirect args.
 	void ResizeBuffers(FRHICommandListImmediate& RHICmdList, int32 ReservedInstanceCounts);
diff --git a/Engine/Plugins/FX/Niagara/Source/Niagara/Public/NiagaraCommon.h b/Engine/Plugins/FX/Niagara/Source/Niagara/Public/NiagaraCommon.h
--- a/Engine/Plugins/FX/Niagara/Source/Niagara/Public/NiagaraCommon.h
+++ b/Engine/Plugins/FX/Niagara/Source/Niagara/Public/NiagaraCommon.h
@@ -1622,16 +1622,16 @@ namespace FNiagaraUtilities
 #endif
 
 	// Whether GPU particles are currently allowed. Could change depending on config and runtime switches.
-	bool AllowGPUParticles();
+	NIAGARA_API bool AllowGPUParticles();
 
 	// Whether compute shaders are allowed. Could change depending on config and runtime switches.
-	bool AllowComputeShaders();
+	NIAGARA_API bool AllowComputeShaders();
 
 	// Are we able to use the GPU for culling?
-	bool AllowGPUCulling();
+	NIAGARA_API bool AllowGPUCulling();
 
 	// Are we able to use the GPU for sorting?
-	bool AllowGPUSorting();
+	NIAGARA_API bool AllowGPUSorting();
 
 	//-TODO: Consider deprecating these methods if we don't need the ShaderPlatform
 	inline bool AllowGPUParticles(EShaderPlatform ShaderPlatform) { return AllowGPUParticles(); }
diff --git a/Engine/Plugins/FX/Niagara/Source/Niagara/Public/NiagaraCullProxyComponent.h b/Engine/Plugins/FX/Niagara/Source/Niagara/Public/NiagaraCullProxyComponent.h
--- a/Engine/Plugins/FX/Niagara/Source/Niagara/Public/NiagaraCullProxyComponent.h
+++ b/Engine/Plugins/FX/Niagara/Source/Niagara/Public/NiagaraCullProxyComponent.h
@@ -19,7 +19,7 @@ struct FNiagaraCulledComponentInfo
 /**
 A specialization of UNiagaraComponent that can act as a proxy for many other NiagaraComponents that have been culled by scalability. 
 */
-UCLASS(NotBlueprintable)
+UCLASS(MinimalAPI, NotBlueprintable)
 class UNiagaraCullProxyComponent : public UNiagaraComponent
 {
 	GENERATED_UCLASS_BODY()
diff --git a/Engine/Plugins/FX/Niagara/Source/Niagara/Public/NiagaraRendererProperties.h b/Engine/Plugins/FX/Niagara/Source/Niagara/Public/NiagaraRendererProperties.h
--- a/Engine/Plugins/FX/Niagara/Source/Niagara/Public/NiagaraRendererProperties.h
+++ b/Engine/Plugins/FX/Niagara/Source/Niagara/Public/NiagaraRendererProperties.h
@@ -405,8 +405,8 @@ public:
 	*/
 	virtual void GetRendererMaterialUsageFeedback(TArray<FNiagaraRendererFeedback>& OutFeedback) const {}
 
-	void GetMaterialUsageFeedback(EMaterialUsage Usage, TConstArrayView<EMaterialDomain> InvalidMaterialDomains, TArray<FNiagaraRendererFeedback>& OutFeedback) const;
-	void GetMaterialUsageFeedback(EMaterialUsage Usage, TArray<FNiagaraRendererFeedback>& OutFeedback) const;
+	NIAGARA_API void GetMaterialUsageFeedback(EMaterialUsage Usage, TConstArrayView<EMaterialDomain> InvalidMaterialDomains, TArray<FNiagaraRendererFeedback>& OutFeedback) const;
+	NIAGARA_API void GetMaterialUsageFeedback(EMaterialUsage Usage, TArray<FNiagaraRendererFeedback>& OutFeedback) const;
 
 	// The icon to display in the niagara stack widget under the renderer section
 	NIAGARA_API virtual const FSlateBrush* GetStackIcon() const;
diff --git a/Engine/Plugins/FX/Niagara/Source/NiagaraVertexFactories/Private/NiagaraSpriteVertexFactory.cpp b/Engine/Plugins/FX/Niagara/Source/NiagaraVertexFactories/Private/NiagaraSpriteVertexFactory.cpp
--- a/Engine/Plugins/FX/Niagara/Source/NiagaraVertexFactories/Private/NiagaraSpriteVertexFactory.cpp
+++ b/Engine/Plugins/FX/Niagara/Source/NiagaraVertexFactories/Private/NiagaraSpriteVertexFactory.cpp
@@ -15,8 +15,6 @@
 IMPLEMENT_GLOBAL_SHADER_PARAMETER_STRUCT(FNiagaraSpriteUniformParameters,"NiagaraSpriteVF");
 IMPLEMENT_GLOBAL_SHADER_PARAMETER_STRUCT(FNiagaraSpriteVFLooseParameters, "NiagaraSpriteVFLooseParameters");
 
-TGlobalResource<FNullDynamicParameterVertexBuffer> GNullNiagaraDynamicParameterVertexBuffer;
-
 class FNiagaraSpriteVertexFactoryShaderParametersVS : public FNiagaraVertexFactoryShaderParametersBase
 {
 	DECLARE_TYPE_LAYOUT(FNiagaraSpriteVertexFactoryShaderParametersVS, NonVirtual);
diff --git a/Engine/Plugins/FX/Niagara/Source/NiagaraVertexFactories/Public/NiagaraVertexFactory.h b/Engine/Plugins/FX/Niagara/Source/NiagaraVertexFactories/Public/NiagaraVertexFactory.h
--- a/Engine/Plugins/FX/Niagara/Source/NiagaraVertexFactories/Public/NiagaraVertexFactory.h
+++ b/Engine/Plugins/FX/Niagara/Source/NiagaraVertexFactories/Public/NiagaraVertexFactory.h
@@ -86,11 +86,11 @@ public:
 */
 class FNiagaraVertexFactoryShaderParametersBase : public FVertexFactoryShaderParameters
 {
-	DECLARE_TYPE_LAYOUT(FNiagaraVertexFactoryShaderParametersBase, NonVirtual);
+	DECLARE_EXPORTED_TYPE_LAYOUT(FNiagaraVertexFactoryShaderParametersBase, NIAGARAVERTEXFACTORIES_API, NonVirtual);
 
 public:
-	void Bind(const FShaderParameterMap& ParameterMap);
-	void GetElementShaderBindings(
+	NIAGARAVERTEXFACTORIES_API void Bind(const FShaderParameterMap& ParameterMap);
+	NIAGARAVERTEXFACTORIES_API void GetElementShaderBindings(
 		const FSceneInterface* Scene,
 		const FSceneView* View,
 		const FMeshMaterialShader* Shader,
