Foundational test automation for Niagara Systems being used in Sequencer. Tests cover Niagara actor being added to a Sequencer as a possessable along with its component and life cycle track and playback.

#tests UE-TC-8898, UE-TC-8898, Local ET + Horde Jobs
#rnx
#virtualized
#Synced-CL 51545932

#Synced-CL 51545932

diff --git a/Engine/Plugins/FX/Niagara/Source/Niagara/Public/MovieScene/MovieSceneNiagaraSystemSpawnSection.h b/Engine/Plugins/FX/Niagara/Source/Niagara/Public/MovieScene/MovieSceneNiagaraSystemSpawnSection.h
--- a/Engine/Plugins/FX/Niagara/Source/Niagara/Public/MovieScene/MovieSceneNiagaraSystemSpawnSection.h
+++ b/Engine/Plugins/FX/Niagara/Source/Niagara/Public/MovieScene/MovieSceneNiagaraSystemSpawnSection.h
@@ -51,12 +51,12 @@ public:
 	NIAGARA_API void SetSectionStartBehavior(ENiagaraSystemSpawnSectionStartBehavior InBehavior);
 
 	UFUNCTION(BlueprintPure, Category = "Niagara|Life Cycle")
-	ENiagaraSystemSpawnSectionEvaluateBehavior GetSectionEvaluateBehavior() const;
+	NIAGARA_API ENiagaraSystemSpawnSectionEvaluateBehavior GetSectionEvaluateBehavior() const;
 	UFUNCTION(BlueprintCallable, Category = "Niagara|Life Cycle")
 	NIAGARA_API void SetSectionEvaluateBehavior(ENiagaraSystemSpawnSectionEvaluateBehavior InBehavior);
 
 	UFUNCTION(BlueprintPure, Category = "Niagara|Life Cycle")
-	ENiagaraSystemSpawnSectionEndBehavior GetSectionEndBehavior() const;
+	NIAGARA_API ENiagaraSystemSpawnSectionEndBehavior GetSectionEndBehavior() const;
 	UFUNCTION(BlueprintCallable, Category = "Niagara|Life Cycle")
 	NIAGARA_API void SetSectionEndBehavior(ENiagaraSystemSpawnSectionEndBehavior InBehavior);
 
