Update NiagaraHlslTranslator to incorporate a hash of the various function details (internal static values and input parameter types) into it's name.  This avoids implementation collisions which could occur.  Also cleans up a few more string operations to use StringBuilders.
#tests Generated bytecode/shaders for all systems and validated that specific scenarios were fixed (like multiple nodes with different pin types)
#Synced-CL 51571507

diff --git a/Engine/Plugins/FX/Niagara/Source/Niagara/Public/NiagaraCommon.h b/Engine/Plugins/FX/Niagara/Source/Niagara/Public/NiagaraCommon.h
--- a/Engine/Plugins/FX/Niagara/Source/Niagara/Public/NiagaraCommon.h
+++ b/Engine/Plugins/FX/Niagara/Source/Niagara/Public/NiagaraCommon.h
@@ -497,8 +497,10 @@ struct FNiagaraFunctionSignature
 	{
 	}
 
-	FNiagaraFunctionSignature(FName InName, TArray<FNiagaraVariable>& InInputs, TArray<FNiagaraVariable>& InOutputs, FName InSource, bool bInRequiresContext, bool bInMemberFunction)
+	FNiagaraFunctionSignature(FName InName, TConstArrayView<FNiagaraVariable> InInputs, TConstArrayView<FNiagaraVariableBase> InOutputs, bool bInRequiresContext, bool bInMemberFunction)
 		: Name(InName)
+		, Inputs(InInputs)
+		, Outputs(InOutputs)
 		, bRequiresContext(bInRequiresContext)
 		, bRequiresExecPin(false)
 		, bMemberFunction(bInMemberFunction)
@@ -516,17 +518,6 @@ struct FNiagaraFunctionSignature
 		, RequiredInputs(INDEX_NONE)
 		, RequiredOutputs(INDEX_NONE)
 	{
-		Inputs.Reserve(InInputs.Num());
-		for (FNiagaraVariable& Var : InInputs)
-		{
-			Inputs.Add(Var);
-		}
-		
-		Outputs.Reserve(InOutputs.Num());
-		for (FNiagaraVariable& Var : InOutputs)
-		{
-			Outputs.Add(Var);
-		}
 	}
 
 	bool operator!=(const FNiagaraFunctionSignature& Other) const
diff --git a/Engine/Plugins/FX/Niagara/Source/Niagara/Public/NiagaraTypes.h b/Engine/Plugins/FX/Niagara/Source/Niagara/Public/NiagaraTypes.h
--- a/Engine/Plugins/FX/Niagara/Source/Niagara/Public/NiagaraTypes.h
+++ b/Engine/Plugins/FX/Niagara/Source/Niagara/Public/NiagaraTypes.h
@@ -1196,6 +1196,12 @@ private:
 	{
 		return HashCombine(HashCombine(Type.UnderlyingTypePathNameHash, GetTypeHash(Type.UnderlyingType)), GetTypeHash(Type.GetFlags()));
 	}
+
+	friend inline uint32 GetPersistentTypeHash(const FNiagaraTypeDefinition& Type)
+	{
+		const FString PathName = GetPathNameSafe(Type.ClassStructOrEnum);
+		return HashCombine(HashCombine(GetTypeHash(PathName), GetTypeHash(Type.UnderlyingType)), GetTypeHash(Type.GetFlags()));
+	}
 };
 
 template<>
@@ -1433,7 +1439,7 @@ protected:
 
 	friend inline uint32 GetPersistentTypeHash(const FNiagaraVariableBase& Var)
 	{
-		return HashCombine(GetTypeHash(Var.GetType()), GetTypeHash(Var.GetName().ToString()));
+		return HashCombine(GetPersistentTypeHash(Var.GetType()), GetTypeHash(Var.GetName().ToString()));
 	}
 };
 
diff --git a/Engine/Plugins/FX/Niagara/Source/NiagaraEditor/Private/Commandlets/NiagaraDumpBytecodeCommandlet.cpp b/Engine/Plugins/FX/Niagara/Source/NiagaraEditor/Private/Commandlets/NiagaraDumpBytecodeCommandlet.cpp
--- a/Engine/Plugins/FX/Niagara/Source/NiagaraEditor/Private/Commandlets/NiagaraDumpBytecodeCommandlet.cpp
+++ b/Engine/Plugins/FX/Niagara/Source/NiagaraEditor/Private/Commandlets/NiagaraDumpBytecodeCommandlet.cpp
@@ -499,31 +499,35 @@ void UNiagaraDumpByteCodeCommandlet::DumpShaderScript(const UNiagaraScript* GPUC
 #if !UE_BUILD_SHIPPING
 	if (const FNiagaraShaderScript* ShaderScript = GPUComputeScript->GetRenderThreadScript())
 	{
-		const FNiagaraShaderMap* ShaderMap = ShaderScript->GetGameThreadShaderMap();
-		const FNiagaraShaderMapContent* ShaderMapContent = ShaderMap->GetContent();
-
-		FShaderListReport ShaderListReport;
-		ShaderListReport.BindLambda([&HashName, FullName = PathName / FilePath, this](
-				const TCHAR* FriendlyName,
-				const TCHAR* OwnerName,
-				const TCHAR* VertexFactory,
-				const TCHAR* ShaderType,
-				uint32 Permutation,
-				uint32 Bytes,
-				const TCHAR* Hash,
-				void* Ptr)
+		if (const FNiagaraShaderMap* ShaderMap = ShaderScript->GetGameThreadShaderMap())
 		{
-			FShaderMetaData& PermutationData = ShaderMetaData.AddDefaulted_GetRef();
-			PermutationData.SystemHash = HashName;
-			PermutationData.FullName = FullName;
-			PermutationData.PermutationIndex = Permutation;
-			PermutationData.VertexFactory = VertexFactory;
-			PermutationData.ShaderType = ShaderType;
-			PermutationData.ByteCount = Bytes;
-			PermutationData.ShaderHash = Hash;
-		});
-
-		ShaderMapContent->DumpShaderList(*ShaderMap, ShaderListReport);
+			if (const FNiagaraShaderMapContent* ShaderMapContent = ShaderMap->GetContent())
+			{
+
+				FShaderListReport ShaderListReport;
+				ShaderListReport.BindLambda([&HashName, FullName = PathName / FilePath, this](
+						const TCHAR* FriendlyName,
+						const TCHAR* OwnerName,
+						const TCHAR* VertexFactory,
+						const TCHAR* ShaderType,
+						uint32 Permutation,
+						uint32 Bytes,
+						const TCHAR* Hash,
+						void* Ptr)
+				{
+					FShaderMetaData& PermutationData = ShaderMetaData.AddDefaulted_GetRef();
+					PermutationData.SystemHash = HashName;
+					PermutationData.FullName = FullName;
+					PermutationData.PermutationIndex = Permutation;
+					PermutationData.VertexFactory = VertexFactory;
+					PermutationData.ShaderType = ShaderType;
+					PermutationData.ByteCount = Bytes;
+					PermutationData.ShaderHash = Hash;
+				});
+
+				ShaderMapContent->DumpShaderList(*ShaderMap, ShaderListReport);
+			}
+		}
 	}
 #endif // !UE_BUILD_SHIPPING
 }
diff --git a/Engine/Plugins/FX/Niagara/Source/NiagaraEditor/Private/NiagaraGraph.cpp b/Engine/Plugins/FX/Niagara/Source/NiagaraEditor/Private/NiagaraGraph.cpp
--- a/Engine/Plugins/FX/Niagara/Source/NiagaraEditor/Private/NiagaraGraph.cpp
+++ b/Engine/Plugins/FX/Niagara/Source/NiagaraEditor/Private/NiagaraGraph.cpp
@@ -209,6 +209,43 @@ bool FNiagaraGraphScriptUsageInfo::IsValid() const
 	return true;
 }
 
+void FNiagaraGraphFunctionAliasContext::AppendPermutationString(FStringBuilderBase& StringBuilder) const
+{
+	TArray<FName, TInlineAllocator<16>> SortedKeys;
+	PermutationIdentifiers.GenerateKeyArray(SortedKeys);
+	SortedKeys.StableSort(FNameLexicalLess());
+
+	for (const FName& Key : SortedKeys)
+	{
+		StringBuilder.AppendChar(TEXT(','));
+		Key.AppendString(StringBuilder);
+		StringBuilder.AppendChar(TEXT('='));
+		StringBuilder.Append(PermutationIdentifiers[Key]);
+	}
+}
+
+void FNiagaraGraphFunctionAliasContext::AddPermutationIdentifier(const FName& Name, FStringView Value)
+{
+	FString* ExistingIdentifier = PermutationIdentifiers.Find(Name);
+	if (ExistingIdentifier)
+	{
+		check(*ExistingIdentifier == Value);
+	}
+	else
+	{
+		PermutationIdentifiers.Emplace(Name, Value);
+	}
+}
+
+void FNiagaraGraphFunctionAliasContext::AddPendingStaticSwitchVariable(const FNiagaraVariable& Variable, uint8 SwitchType)
+{
+	FStaticSwitchVariable SwitchVariable;
+	SwitchVariable.Variable = Variable;
+	SwitchVariable.SwitchType = SwitchType;
+
+	CompilerStaticSwitchVariables.AddUnique(SwitchVariable);
+}
+
 FDelegateHandle UNiagaraGraph::AddOnGraphNeedsRecompileHandler(const FOnGraphChanged::FDelegate& InHandler)
 {
 	return OnGraphNeedsRecompile.Add(InHandler);
@@ -3759,34 +3796,15 @@ void UNiagaraGraph::InvalidateNumericCache()
 	CachedNumericConversions.Empty();
 }
 
-FString UNiagaraGraph::GetFunctionAliasByContext(const FNiagaraGraphFunctionAliasContext& FunctionAliasContext) const
+void UNiagaraGraph::UpdateFunctionAliasByContext(FNiagaraGraphFunctionAliasContext& FunctionAliasContext) const
 {
-	FString FunctionAlias;
-	TSet<UClass*> SkipNodeTypes;
 	for (const UEdGraphNode* Node : Nodes)
 	{
-		const UNiagaraNode* NiagaraNode = Cast<const UNiagaraNode>(Node);
-		if (NiagaraNode != nullptr)
+		if (const UNiagaraNode* NiagaraNode = Cast<const UNiagaraNode>(Node))
 		{
-			if (SkipNodeTypes.Contains(NiagaraNode->GetClass()))
-			{
-				continue;
-			}
-			bool OncePerNodeType = false;
-			NiagaraNode->AppendFunctionAliasForContext(FunctionAliasContext, FunctionAlias, OncePerNodeType);
-			if (OncePerNodeType)
-			{
-				SkipNodeTypes.Add(NiagaraNode->GetClass());
-			}
+			NiagaraNode->UpdateFunctionAliasForContext(FunctionAliasContext);
 		}
 	}
-
-	for (const UEdGraphPin* Pin : FunctionAliasContext.StaticSwitchValues)
-	{
-		FunctionAlias += TEXT("_") + FNiagaraHlslTranslator::GetSanitizedFunctionNameSuffix(Pin->GetName())
-			+ TEXT("_") + FNiagaraHlslTranslator::GetSanitizedFunctionNameSuffix(Pin->DefaultValue);
-	}
-	return FunctionAlias;
 }
 
 void UNiagaraGraph::ResolveNumerics(TMap<UNiagaraNode*, bool>& VisitedNodes, UEdGraphNode* Node)
diff --git a/Engine/Plugins/FX/Niagara/Source/NiagaraEditor/Private/NiagaraGraphDigest.cpp b/Engine/Plugins/FX/Niagara/Source/NiagaraEditor/Private/NiagaraGraphDigest.cpp
--- a/Engine/Plugins/FX/Niagara/Source/NiagaraEditor/Private/NiagaraGraphDigest.cpp
+++ b/Engine/Plugins/FX/Niagara/Source/NiagaraEditor/Private/NiagaraGraphDigest.cpp
@@ -889,33 +889,15 @@ void FNiagaraCompilationGraph::FindInputNodes(TArray<const FNiagaraCompilationNo
 	}
 }
 
-FString FNiagaraCompilationGraph::GetFunctionAliasByContext(const FNiagaraDigestFunctionAliasContext& FunctionAliasContext) const
+void FNiagaraCompilationGraph::UpdateFunctionAliasByContext(FNiagaraGraphFunctionAliasContext& FunctionAliasContext) const
 {
-	FString FunctionAlias;
-	TSet<UClass*> SkipNodeTypes;
 	for (const TUniquePtr<FNiagaraCompilationNode>& Node : Nodes)
 	{
 		if (const FNiagaraCompilationNode* NiagaraNode = Node.Get())
 		{
-			if (SkipNodeTypes.Contains(NiagaraNode->GetSourceNodeClass()))
-			{
-				continue;
-			}
-			bool OncePerNodeType = false;
-			NiagaraNode->AppendFunctionAliasForContext(FunctionAliasContext, FunctionAlias, OncePerNodeType);
-			if (OncePerNodeType)
-			{
-				SkipNodeTypes.Add(NiagaraNode->GetSourceNodeClass());
-			}
+			NiagaraNode->UpdateFunctionAliasForContext(FunctionAliasContext);
 		}
 	}
-
-	for (const FNiagaraCompilationPin* Pin : FunctionAliasContext.StaticSwitchValues)
-	{
-		FunctionAlias += TEXT("_") + FNiagaraHlslTranslator::GetSanitizedFunctionNameSuffix(Pin->PinName.ToString())
-			+ TEXT("_") + FNiagaraHlslTranslator::GetSanitizedFunctionNameSuffix(Pin->DefaultValue);
-	}
-	return FunctionAlias;
 }
 
 bool FNiagaraCompilationGraph::HasParametersOfType(const FNiagaraTypeDefinition& Type) const
@@ -3433,14 +3415,14 @@ void FNiagaraCompilationNodeInput::Compile(FTranslator* Translator, TArray<int32
 	}
 }
 
-void FNiagaraCompilationNodeInput::AppendFunctionAliasForContext(const FNiagaraDigestFunctionAliasContext& InFunctionAliasContext, FString& InOutFunctionAlias, bool& OutOnlyOncePerNodeType) const
+void FNiagaraCompilationNodeInput::UpdateFunctionAliasForContext(FNiagaraGraphFunctionAliasContext& InOutFunctionAliasContext) const
 {
 	if (Usage == ENiagaraInputNodeUsage::TranslatorConstant && InputVariable == TRANSLATOR_PARAM_CALL_ID)
 	{
-		OutOnlyOncePerNodeType = true;
 		// The call ID should be unique for each translated node as it is used by the seeded random functions.
 		// We don't want it to be shared across the spawn and update script, so functions including it will have the usage added to their name.
-		InOutFunctionAlias += "_ScriptUsage" + FString::FormatAsNumber((uint8)InFunctionAliasContext.ScriptUsage);
+		static const FName NAME_ScriptUsage = TEXT("ScriptUsage");
+		InOutFunctionAliasContext.AddPermutationIdentifier(NAME_ScriptUsage, FString::FormatAsNumber(static_cast<uint8>(InOutFunctionAliasContext.ScriptUsage)));
 	}
 }
 
@@ -4291,17 +4273,19 @@ void FNiagaraCompilationNodeStaticSwitch::Compile(FTranslator* Translator, TArra
 	{
 		FNiagaraFixedConstantResolver ConstantResolver(Translator);
 
-		const FNiagaraVariable* Found = FNiagaraConstants::FindStaticSwitchConstant(SwitchConstant);
-		FNiagaraVariable Constant = Found ? *Found : FNiagaraVariable();
-		if (Found && ConstantResolver.ResolveConstant(Constant))
+		if (const FNiagaraVariable* Found = FNiagaraConstants::FindStaticSwitchConstant(SwitchConstant))
 		{
-			if (SwitchType == ENiagaraStaticSwitchType::Bool)
+			FNiagaraVariable Constant = *Found;
+			if (ConstantResolver.ResolveConstant(Constant))
 			{
-				SelectorValue = Constant.GetValue<bool>();
-			}
-			else if (SwitchType == ENiagaraStaticSwitchType::Integer || SwitchType == ENiagaraStaticSwitchType::Enum)
-			{
-				SelectorValue = Constant.GetValue<int32>();
+				if (SwitchType == ENiagaraStaticSwitchType::Bool)
+				{
+					SelectorValue = Constant.GetValue<bool>();
+				}
+				else if (SwitchType == ENiagaraStaticSwitchType::Integer || SwitchType == ENiagaraStaticSwitchType::Enum)
+				{
+					SelectorValue = Constant.GetValue<int32>();
+				}
 			}
 		}
 	}
@@ -4338,6 +4322,17 @@ void FNiagaraCompilationNodeStaticSwitch::ResolveNumerics()
 	}
 }
 
+void FNiagaraCompilationNodeStaticSwitch::UpdateFunctionAliasForContext(FNiagaraGraphFunctionAliasContext& InOutFunctionAliasContext) const
+{
+	if (bSetByCompiler)
+	{
+		if (const FNiagaraVariable* Found = FNiagaraConstants::FindStaticSwitchConstant(SwitchConstant))
+		{
+			InOutFunctionAliasContext.AddPendingStaticSwitchVariable(*Found, static_cast<uint8>(SwitchType));
+		}
+	}
+}
+
 bool FNiagaraCompilationNodeStaticSwitch::ResolveConstantValue(const FNiagaraCompilationInputPin& Pin, int32& Value)
 {
 	if (Pin.LinkedTo)
@@ -4635,51 +4630,44 @@ void FNiagaraCompilationNodeUsageSelector::Compile(FTranslator* Translator, TArr
 	}
 }
 
-void FNiagaraCompilationNodeUsageSelector::AppendFunctionAliasForContext(const FNiagaraDigestFunctionAliasContext& InFunctionAliasContext, FString& InOutFunctionAlias, bool& OutOnlyOncePerNodeType) const
+void FNiagaraCompilationNodeUsageSelector::UpdateFunctionAliasForContext(FNiagaraGraphFunctionAliasContext& InOutFunctionAliasContext) const
 {
-	OutOnlyOncePerNodeType = true;
-
-	FString UsageString;
-	switch (InFunctionAliasContext.CompileUsage)
-	{
-	case ENiagaraScriptUsage::SystemSpawnScript:
-	case ENiagaraScriptUsage::SystemUpdateScript:
-		UsageString = "System";
-		break;
-	case ENiagaraScriptUsage::EmitterSpawnScript:
-	case ENiagaraScriptUsage::EmitterUpdateScript:
-		UsageString = "Emitter";
-		break;
-	case ENiagaraScriptUsage::ParticleSpawnScript:
-	case ENiagaraScriptUsage::ParticleUpdateScript:
-	case ENiagaraScriptUsage::ParticleEventScript:
-	case ENiagaraScriptUsage::ParticleSimulationStageScript:
-	case ENiagaraScriptUsage::ParticleGPUComputeScript:
-		UsageString = "Particle";
-		break;
-	}
-
-	if (UsageString.IsEmpty() == false)
+	static const FName NAME_Usage = TEXT("Usage");
+	switch (InOutFunctionAliasContext.CompileUsage)
 	{
-		InOutFunctionAlias += "_" + UsageString;
+		case ENiagaraScriptUsage::SystemSpawnScript:
+		case ENiagaraScriptUsage::SystemUpdateScript:
+			InOutFunctionAliasContext.AddPermutationIdentifier(NAME_Usage, TEXT("System"));
+			break;
+		case ENiagaraScriptUsage::EmitterSpawnScript:
+		case ENiagaraScriptUsage::EmitterUpdateScript:
+			InOutFunctionAliasContext.AddPermutationIdentifier(NAME_Usage, TEXT("Emitter"));
+			break;
+		case ENiagaraScriptUsage::ParticleSpawnScript:
+		case ENiagaraScriptUsage::ParticleUpdateScript:
+		case ENiagaraScriptUsage::ParticleEventScript:
+		case ENiagaraScriptUsage::ParticleSimulationStageScript:
+		case ENiagaraScriptUsage::ParticleGPUComputeScript:
+			InOutFunctionAliasContext.AddPermutationIdentifier(NAME_Usage, TEXT("Particle"));
+			break;
 	}
 }
 
 const FNiagaraCompilationOutputPin* FNiagaraCompilationNodeUsageSelector::FindOutputPin(const FNiagaraVariable& Variable) const
 {
 	const int32 OutputIndex = OutputVars.IndexOfByPredicate([&Variable](const FNiagaraVariable& OutputVariable) -> bool
-		{
-			return OutputVariable == Variable;
-		});
+	{
+		return OutputVariable == Variable;
+	});
 
 	if (OutputVarGuids.IsValidIndex(OutputIndex))
 	{
 		const FGuid& OutputVarGuid = OutputVarGuids[OutputIndex];
 
 		return OutputPins.FindByPredicate([&OutputVarGuid](const FNiagaraCompilationOutputPin& OutputPin) -> bool
-			{
-				return OutputPin.PersistentGuid == OutputVarGuid;
-			});
+		{
+			return OutputPin.PersistentGuid == OutputVarGuid;
+		});
 	}
 
 	return nullptr;
diff --git a/Engine/Plugins/FX/Niagara/Source/NiagaraEditor/Private/NiagaraGraphDigest.h b/Engine/Plugins/FX/Niagara/Source/NiagaraEditor/Private/NiagaraGraphDigest.h
--- a/Engine/Plugins/FX/Niagara/Source/NiagaraEditor/Private/NiagaraGraphDigest.h
+++ b/Engine/Plugins/FX/Niagara/Source/NiagaraEditor/Private/NiagaraGraphDigest.h
@@ -120,16 +120,6 @@ public:
 	TArray<const FNiagaraCompilationInputPin*> LinkedTo;
 };
 
-struct FNiagaraDigestFunctionAliasContext
-{
-	// the usage as defined in the compilation request (same for all translation stages)
-	ENiagaraScriptUsage CompileUsage;
-
-	// the usage as defined in the current translation stage
-	ENiagaraScriptUsage ScriptUsage;
-	TArray<const FNiagaraCompilationInputPin*> StaticSwitchValues;
-};
-
 class INiagaraCompilationFeedbackInterface
 {
 public:
@@ -216,7 +206,7 @@ public:
 	virtual void BuildParameterMapHistory(FParameterMapHistoryBuilder& Builder, bool bRecursive, bool bFilterForCompilation) const;
 	virtual void Compile(FTranslator* Translator, TArray<int32>& Outputs) const;
 	virtual void ResolveNumerics();
-	virtual void AppendFunctionAliasForContext(const FNiagaraDigestFunctionAliasContext& InFunctionAliasContext, FString& InOutFunctionAlias, bool& OutOnlyOncePerNodeType) const {};
+	virtual void UpdateFunctionAliasForContext(FNiagaraGraphFunctionAliasContext& InOutFunctionAliasContext) const {};
 
 	bool ConditionalRouteParameterMapAroundMe(FParameterMapHistoryBuilder& Builder) const;
 	void RouteParameterMapAroundMe(FParameterMapHistoryBuilder& Builder) const;
@@ -288,7 +278,7 @@ public:
 	  * function calls to the same graph. For example, if the graph contains static switches and two functions call it with
 	  * different switch parameters, the final function names in the hlsl must be different.
 	  */
-	FString GetFunctionAliasByContext(const FNiagaraDigestFunctionAliasContext& FunctionAliasContext) const;
+	void UpdateFunctionAliasByContext(FNiagaraGraphFunctionAliasContext& FunctionAliasContext) const;
 
 	bool HasParametersOfType(const FNiagaraTypeDefinition& Type) const;
 
@@ -503,7 +493,7 @@ public:
 
 	virtual void BuildParameterMapHistory(FParameterMapHistoryBuilder& Builder, bool bRecursive, bool bFilterForCompilation) const override;
 	virtual void Compile(FTranslator* Translator, TArray<int32>& Outputs) const;
-	virtual void AppendFunctionAliasForContext(const FNiagaraDigestFunctionAliasContext& InFunctionAliasContext, FString& InOutFunctionAlias, bool& OutOnlyOncePerNodeType) const override;
+	virtual void UpdateFunctionAliasForContext(FNiagaraGraphFunctionAliasContext& InOutFunctionAliasContext) const override;
 
 	FNiagaraVariable InputVariable;
 	ENiagaraInputNodeUsage Usage;
@@ -686,7 +676,7 @@ public:
 
 	virtual void BuildParameterMapHistory(FParameterMapHistoryBuilder& Builder, bool bRecursive, bool bFilterForCompilation) const override;
 	virtual void Compile(FTranslator* Translator, TArray<int32>& Outputs) const;
-	virtual void AppendFunctionAliasForContext(const FNiagaraDigestFunctionAliasContext& InFunctionAliasContext, FString& InOutFunctionAlias, bool& OutOnlyOncePerNodeType) const override;
+	virtual void UpdateFunctionAliasForContext(FNiagaraGraphFunctionAliasContext& InOutFunctionAliasContext) const override;
 
 	const FNiagaraCompilationOutputPin* FindOutputPin(const FNiagaraVariable& Variable) const;
 
@@ -722,6 +712,7 @@ public:
 	virtual void BuildParameterMapHistory(FParameterMapHistoryBuilder& Builder, bool bRecursive, bool bFilterForCompilation) const override;
 	virtual void Compile(FTranslator* Translator, TArray<int32>& Outputs) const;
 	virtual void ResolveNumerics() override;
+	virtual void UpdateFunctionAliasForContext(FNiagaraGraphFunctionAliasContext& InOutFunctionAliasContext) const override;
 	virtual TArray<const FNiagaraCompilationInputPin*> EvaluateBranches(FNiagaraCompilationGraphInstanceContext& Context, FNiagaraCompilationBranchMap& Branches) const override;
 
 	static bool ResolveConstantValue(const FNiagaraCompilationInputPin& Pin, int32& Value);
diff --git a/Engine/Plugins/FX/Niagara/Source/NiagaraEditor/Private/NiagaraHlslTranslator.cpp b/Engine/Plugins/FX/Niagara/Source/NiagaraEditor/Private/NiagaraHlslTranslator.cpp
--- a/Engine/Plugins/FX/Niagara/Source/NiagaraEditor/Private/NiagaraHlslTranslator.cpp
+++ b/Engine/Plugins/FX/Niagara/Source/NiagaraEditor/Private/NiagaraHlslTranslator.cpp
@@ -2,6 +2,7 @@
 
 #include "NiagaraHlslTranslator.h"
 
+#include "Algo/Contains.h"
 #include "EdGraphSchema_Niagara.h"
 #include "INiagaraEditorTypeUtilities.h"
 #include "Modules/ModuleManager.h"
@@ -294,7 +295,7 @@ void AppendFunctionSignatureSymbol(const FNiagaraFunctionSignature& Sig, FString
 	}
 	else
 	{
-		SigStr << TEXT("_Func_");
+		SigStr << TEXT("_Func");
 	}
 
 	if (Sig.bRequiresExecPin)
@@ -798,18 +799,14 @@ bool TNiagaraHlslTranslator<FNiagaraCompilationGraphBridge>::ValidateTypePins(co
 }
 
 template<typename GraphBridge>
-void TNiagaraHlslTranslator<GraphBridge>::GenerateFunctionSignature(ENiagaraScriptUsage ScriptUsage, FString InName, const FString& InFullName, const FString& InFunctionNameSuffix, const FGraph* FuncGraph, TArray<int32>& Inputs,
-	bool bHasNumericInputs, bool bHasParameterMapParameters, const TArray<const FInputPin*>& StaticSwitchValues, FNiagaraFunctionSignature& OutSig) const
+void TNiagaraHlslTranslator<GraphBridge>::GenerateFunctionSignature(ENiagaraScriptUsage ScriptUsage, FStringView InName, FName InFullName, FStringView InFunctionNameSuffix, const FGraph* FuncGraph, TConstArrayView<const FInputPin*> CallInputs, TArray<int32>& Inputs,
+	FString& FunctionComment, FNiagaraFunctionSignature& OutSig) const
 {
 	NIAGARA_SCOPE_CYCLE_COUNTER(STAT_NiagaraEditor_Module_NiagaraHLSLTranslator_GenerateFunctionSignature);
 
-	TArray<FNiagaraVariable> InputVars;
-	TArray<const FInputNode*> InputsNodes;
-	bool bHasDIParameters = false;
-
+	TArray<const FInputNode*, TInlineAllocator<16>> InputNodes;
 	{
 		NIAGARA_SCOPE_CYCLE_COUNTER(STAT_NiagaraEditor_Module_NiagaraHLSLTranslator_GenerateFunctionSignature_FindInputNodes);
-		InputsNodes.Reserve(100);
 		UNiagaraGraph::FFindInputNodeOptions Options;
 		Options.bSort = true;
 		Options.bFilterDuplicates = true;
@@ -819,115 +816,176 @@ void TNiagaraHlslTranslator<GraphBridge>::GenerateFunctionSignature(ENiagaraScri
 		// to compile a graph that would otherwise work correctly.
 		Options.bFilterByScriptUsage = ScriptUsage == ENiagaraScriptUsage::EmitterSpawnScript || ScriptUsage == ENiagaraScriptUsage::EmitterUpdateScript;
 		Options.TargetScriptUsage = ScriptUsage;
-		InputsNodes = GraphBridge::GetGraphInputNodes(FuncGraph, Options);
+		InputNodes = GraphBridge::GetGraphInputNodes(FuncGraph, Options);
 
-		if (Inputs.Num() != InputsNodes.Num())
+		if (Inputs.Num() != InputNodes.Num())
 		{
 			const_cast<TNiagaraHlslTranslator<GraphBridge>*>(this)->Error(FText::Format(LOCTEXT("GenerateFunctionSignatureFail", "Generating function signature for {0} failed.  The function call is providing a different number of inputs than the function graph supplies."),
-				FText::FromString(InFullName)));
+				FText::FromString(InFullName.ToString())));
 			return;
 		}
 	}
 
+	TArray<FNiagaraVariable, TInlineAllocator<16>> ConstantLiteralInputVars;
+	TArray<FNiagaraVariable, TInlineAllocator<16>> InputVars;
+	TArray<FNiagaraVariableBase, TInlineAllocator<16>> OutputVars;
+
+	bool bHasDIParameters = false;
 	{
-		NIAGARA_SCOPE_CYCLE_COUNTER(STAT_NiagaraEditor_Module_NiagaraHLSLTranslator_GenerateFunctionSignature_Inputs);
+		const int32 InputCount = Inputs.Num();
+
+		ConstantLiteralInputVars.Reserve(InputCount);
+		InputVars.Reserve(InputCount);
 
-		InName.Reserve(100 * InputsNodes.Num());
-		InputVars.Reserve(InputsNodes.Num());
-		TArray<uint32> ConstantInputIndicesToRemove;
-		for (int32 i = 0; i < InputsNodes.Num(); ++i)
+		for (int32 InputIt = 0; InputIt < InputCount; ++InputIt)
 		{
-			//Only add to the signature if the caller has provided it, otherwise we use a local default.
-			if (Inputs[i] != INDEX_NONE)
+			if (Inputs[InputIt] == INDEX_NONE)
+			{
+				continue;
+			}
+
+			FNiagaraVariable InputVar = GraphBridge::GetInputVariable(InputNodes[InputIt]);
+			if (GetLiteralConstantVariable(InputVar))
 			{
-				FNiagaraVariable InputVar = GraphBridge::GetInputVariable(InputsNodes[i]);
-				if (GetLiteralConstantVariable(InputVar))
+				ConstantLiteralInputVars.Emplace(InputVar);
+				Inputs[InputIt] = INDEX_NONE;
+			}
+			else
+			{
+				if (!bHasDIParameters && InputVar.GetType().IsDataInterface())
 				{
-					checkf(InputVar.GetType().IsSameBaseDefinition(FNiagaraTypeDefinition::GetBoolDef()), TEXT("Only boolean types are currently supported for literal constants."));
-					FString LiteralConstantAlias = InputVar.GetName().ToString() + TEXT("_") + (InputVar.GetValue<bool>() ? TEXT("true") : TEXT("false"));
-					InName += TEXT("_") + BuildSanitizedSymbolString(LiteralConstantAlias.Replace(TEXT("."), TEXT("_")));
-					ConstantInputIndicesToRemove.Add(i);
-				}
-				else
-				{
-					InputVars.Add(InputVar);
-					if (InputVar.GetType().IsDataInterface())
-					{
-						bHasDIParameters = true;
-					}
-					else if (bHasNumericInputs)
-					{
-						InName += TEXT("_In");
-						InName += InputVar.GetType().GetName();
-					}
+					bHasDIParameters = true;
 				}
+				InputVars.Emplace(InputVar);
 			}
 		}
-		
-		// Remove the inputs which will be handled by inline constants
-		for (int32 i = ConstantInputIndicesToRemove.Num() - 1; i >= 0; i--)
-		{
-			Inputs.RemoveAt(ConstantInputIndicesToRemove[i]);
-		}
-
-		//Now actually remove the missing inputs so they match the signature.
-		Inputs.Remove(INDEX_NONE);
 	}
 
-	TArray<FNiagaraVariable> OutputVars;
-	{
-		NIAGARA_SCOPE_CYCLE_COUNTER(STAT_NiagaraEditor_Module_NiagaraHLSLTranslator_GenerateFunctionSignature_Outputs);
-
-		OutputVars = GraphBridge::GetGraphOutputNodeVariables(FuncGraph, ScriptUsage);
+	Inputs.Remove(INDEX_NONE);
 
-		for (int32 i = 0; i < OutputVars.Num(); ++i)
-		{
-			//Only add to the signature if the caller has provided it, otherwise we use a local default.
-			if (bHasNumericInputs)
-			{
-				InName += TEXT("_Out");
-				InName += OutputVars[i].GetType().GetName();
-			}
-		}
+	{
+		const TArray<FNiagaraVariable> OutputNodeVariables = GraphBridge::GetGraphOutputNodeVariables(FuncGraph, ScriptUsage);
+		OutputVars.Reserve(OutputNodeVariables.Num());
+		Algo::Transform(OutputNodeVariables, OutputVars, [](const FNiagaraVariable& Var) -> FNiagaraVariableBase { return Var; });
 	}
 
 	const FString* ModuleAliasStr = ActiveHistoryForFunctionCalls.GetModuleAlias();
 	const FString* EmitterAliasStr = ActiveHistoryForFunctionCalls.GetEmitterAlias();
+
 	// For now, we want each module call to be unique due to parameter maps and aliasing causing different variables
 	// to be written within each call.
-	if ((ScriptUsage == ENiagaraScriptUsage::Module || ScriptUsage == ENiagaraScriptUsage::DynamicInput ||
-		ScriptUsage == ENiagaraScriptUsage::EmitterSpawnScript || ScriptUsage == ENiagaraScriptUsage::EmitterUpdateScript
-		|| bHasParameterMapParameters || bHasDIParameters)
-		&& (ModuleAliasStr != nullptr || EmitterAliasStr != nullptr))
-	{
-		NIAGARA_SCOPE_CYCLE_COUNTER(STAT_NiagaraEditor_Module_NiagaraHLSLTranslator_GenerateFunctionSignature_UniqueDueToMaps);
-		FString SignatureName;
-		SignatureName.Reserve(1024);
+	const bool bHasParameterMapParameters = GraphBridge::GraphHasParametersOfType(FuncGraph, FNiagaraTypeDefinition::GetParameterMapDef());
+	const bool bScriptUsageRequiresUnique = ScriptUsage == ENiagaraScriptUsage::Module
+		|| ScriptUsage == ENiagaraScriptUsage::DynamicInput
+		|| ScriptUsage == ENiagaraScriptUsage::EmitterSpawnScript
+		|| ScriptUsage == ENiagaraScriptUsage::EmitterUpdateScript;
+	const bool bParametersRequiresUnique = bHasParameterMapParameters || bHasDIParameters;
+	const bool bHasUniqueNaming = (ModuleAliasStr != nullptr || EmitterAliasStr != nullptr);
+
+	if (bHasUniqueNaming && (bScriptUsageRequiresUnique || bParametersRequiresUnique))
+	{
+		FNameBuilder SignatureName;
 		if (ModuleAliasStr != nullptr)
 		{
-			SignatureName = BuildSanitizedSymbolString(*ModuleAliasStr);
+			SignatureName << BuildSanitizedSymbolString(*ModuleAliasStr);
 		}
 		if (EmitterAliasStr != nullptr)
 		{
-			FString Prefix = ModuleAliasStr != nullptr ? TEXT("_") : TEXT("");
-			SignatureName += Prefix;
-			SignatureName += BuildSanitizedSymbolString(*EmitterAliasStr);
+			if (ModuleAliasStr)
+			{
+				SignatureName << TEXT('_');
+			}
+			SignatureName << BuildSanitizedSymbolString(*EmitterAliasStr);
 		}
 		if (InFunctionNameSuffix.IsEmpty() == false)
 		{
-			SignatureName += TEXT("_") + InFunctionNameSuffix;
+			SignatureName << TEXT('_');
+			SignatureName << InFunctionNameSuffix;
 		}
-		SignatureName.ReplaceInline(TEXT("."), TEXT("_"));
-		OutSig = FNiagaraFunctionSignature(*SignatureName, InputVars, OutputVars, *InFullName, true, false);
+		NiagaraHlslTranslatorImpl::ReplaceInline(SignatureName, 0, TEXT("."), TEXT("_"));
+		OutSig = FNiagaraFunctionSignature(*SignatureName, InputVars, OutputVars, true, false);
 	}
 	else
 	{
-		FGraphFunctionAliasContext FunctionAliasContext;
-		FunctionAliasContext.CompileUsage = GetCurrentUsage();
-		FunctionAliasContext.ScriptUsage = TranslationStages[ActiveStageIdx].ScriptUsage;
-		FunctionAliasContext.StaticSwitchValues = StaticSwitchValues;
-		FString SignatureName = InName + FuncGraph->GetFunctionAliasByContext(FunctionAliasContext);
-		OutSig = FNiagaraFunctionSignature(*SignatureName, InputVars, OutputVars, *InFullName, true, false);
+		// compose a string containing all of the decorations to the function and boil it down into a Guid that we can append to the function name
+		FGuid PermutationGuid;
+		{
+			FNiagaraGraphFunctionAliasContext FunctionAliasContext;
+			FunctionAliasContext.CompileUsage = GetCurrentUsage();
+			FunctionAliasContext.ScriptUsage = TranslationStages[ActiveStageIdx].ScriptUsage;
+
+			for (const FNiagaraVariableBase& StaticSwitchInput : GraphBridge::GraphGetStaticSwitchInputs(FuncGraph))
+			{
+				const FInputPin* const* FoundInputPinPtr = CallInputs.FindByPredicate([PinName = StaticSwitchInput.GetName()](const FInputPin* Pin) -> bool
+				{
+					return Pin && PinName.IsEqual(Pin->PinName);
+				});
+
+				if (FoundInputPinPtr)
+				{
+					const FInputPin& InputPin = **FoundInputPinPtr;
+					FunctionAliasContext.AddPermutationIdentifier(InputPin.PinName, InputPin.DefaultValue);
+				}
+			}
+
+			FuncGraph->UpdateFunctionAliasByContext(FunctionAliasContext);
+
+			FNameBuilder FunctionAliasSuffix;
+			ResolveGraphFunctionAlias(FunctionAliasContext, FunctionAliasSuffix);
+
+			for (const FNiagaraVariable& LiteralConstant : ConstantLiteralInputVars)
+			{
+				checkf(LiteralConstant.GetType().IsSameBaseDefinition(FNiagaraTypeDefinition::GetBoolDef()), TEXT("Only boolean types are currently supported for literal constants."));
+				FunctionAliasSuffix << TEXT(',');
+				LiteralConstant.GetName().AppendString(FunctionAliasSuffix);
+				FunctionAliasSuffix << TEXT('_');
+				FunctionAliasSuffix << (LiteralConstant.GetValue<bool>() ? TEXT("true") : TEXT("false"));
+			}
+
+			if (!InputVars.IsEmpty())
+			{
+				FunctionAliasSuffix << TEXT(",Input - [");
+				for (const FNiagaraVariableBase& InputVar : InputVars)
+				{
+					InputVar.GetType().GetFName().AppendString(FunctionAliasSuffix);
+					FunctionAliasSuffix << TEXT(" ");
+					InputVar.GetName().AppendString(FunctionAliasSuffix);
+					FunctionAliasSuffix << TEXT(";");
+				}
+				FunctionAliasSuffix << TEXT("]");
+			}
+
+			if (!OutputVars.IsEmpty())
+			{
+				FunctionAliasSuffix << TEXT(",Output - [");
+				for (const FNiagaraVariableBase& OutputVar : OutputVars)
+				{
+					OutputVar.GetType().GetFName().AppendString(FunctionAliasSuffix);
+					FunctionAliasSuffix << TEXT(" ");
+					OutputVar.GetName().AppendString(FunctionAliasSuffix);
+					FunctionAliasSuffix << TEXT(";");
+				}
+				FunctionAliasSuffix << TEXT("]");
+			}
+
+			if (FunctionAliasSuffix.Len() > 0)
+			{
+				PermutationGuid = FGuid::NewDeterministicGuid(FunctionAliasSuffix.ToView());
+				FunctionComment = FString(InName, FunctionAliasSuffix.Len());
+				FunctionComment.Append(FunctionAliasSuffix.ToView());
+			}
+		}
+
+		FNameBuilder NameBuilder;
+		NameBuilder.Append(InName);
+
+		if (PermutationGuid.IsValid())
+		{
+			NameBuilder << TEXT('_');
+			PermutationGuid.AppendString(NameBuilder, EGuidFormats::Base36Encoded);
+		}
+
+		OutSig = FNiagaraFunctionSignature(*NameBuilder, InputVars, OutputVars, true, false);
 	}
 
 	// if we are splitting up our functions then we need to mark which stage this function signature is associated 
@@ -938,6 +996,44 @@ void TNiagaraHlslTranslator<GraphBridge>::GenerateFunctionSignature(ENiagaraScri
 	}
 }
 
+template<typename GraphBridge>
+void TNiagaraHlslTranslator<GraphBridge>::ResolveGraphFunctionAlias(const FNiagaraGraphFunctionAliasContext& FunctionAliasContext, FStringBuilderBase& FunctionAliasSuffix) const
+{
+	FunctionAliasContext.AppendPermutationString(FunctionAliasSuffix);
+
+	TConstArrayView<FNiagaraGraphFunctionAliasContext::FStaticSwitchVariable> PendingStaticSwitchVariables = FunctionAliasContext.ReadPendingStaticSwitchVariables();
+	if (!PendingStaticSwitchVariables.IsEmpty())
+	{
+		FNiagaraFixedConstantResolver ConstantResolver(this);
+
+		for (FNiagaraGraphFunctionAliasContext::FStaticSwitchVariable CompilerVariable : PendingStaticSwitchVariables)
+		{
+			int32 SwitchValue = INDEX_NONE;
+
+			if (ConstantResolver.ResolveConstant(CompilerVariable.Variable))
+			{
+				ENiagaraStaticSwitchType SwitchType = static_cast<ENiagaraStaticSwitchType>(CompilerVariable.SwitchType);
+				if (SwitchType == ENiagaraStaticSwitchType::Bool)
+				{
+					SwitchValue = CompilerVariable.Variable.GetValue<bool>() ? 1 : 0;
+				}
+				else if (SwitchType == ENiagaraStaticSwitchType::Integer || SwitchType == ENiagaraStaticSwitchType::Enum)
+				{
+					SwitchValue = CompilerVariable.Variable.GetValue<int32>();
+				}
+			}
+
+			if (SwitchValue != INDEX_NONE)
+			{
+				FunctionAliasSuffix.AppendChar(TEXT(','));
+				CompilerVariable.Variable.GetName().AppendString(FunctionAliasSuffix);
+				FunctionAliasSuffix.AppendChar(TEXT('='));
+				FunctionAliasSuffix.Append(FString::FormatAsNumber(SwitchValue));
+			}
+		}
+	}
+}
+
 //////////////////////////////////////////////////////////////////////////
 
 FNiagaraHlslTranslator::FNiagaraHlslTranslator()
@@ -974,6 +1070,18 @@ FString FNiagaraHlslTranslator::GetFunctionDefinitions()
 		if (!FuncPair.Value.Body.IsEmpty())
 		{
 			FDeclarationPermutationContext PermutationContext(*this, TranslationStages, FuncPair.Value.StageIndices, DefinitionsString);
+			if (!FuncPair.Value.Comments.IsEmpty())
+			{
+				TArray<FStringView, TInlineAllocator<16>> CommentLines;
+				UE::String::ParseTokens(FuncPair.Value.Comments, TEXT(","), [&CommentLines](FStringView Token) { CommentLines.Add(Token); });
+
+				for (FStringView CommentLine : CommentLines)
+				{
+					DefinitionsString += TEXT("// ");
+					DefinitionsString += CommentLine;
+					DefinitionsString += TEXT("\n");
+				}
+			}
 			DefinitionsString += Sig + TEXT("\n{\n") + FuncPair.Value.Body + TEXT("}\n\n");
 		}
 		// Don't do anything if the value is empty on the function pair, as this is indicative of 
@@ -1826,14 +1934,15 @@ FNiagaraTranslateResults TNiagaraHlslTranslator<GraphBridge>::Translate(const FN
 				{
 					const auto& CompileSimStageData = CompileData->CompileSimStageData[SourceSimStageIndex];
 
-					FString StageName = CompileSimStageData.StageName.ToString();
-					StageName = TEXT("_") + GetSanitizedFunctionNameSuffix(StageName);
-
+					FNameBuilder StageNameBuilder;
+					StageNameBuilder.AppendChar(TEXT('_'));
+					AppendSanitizedFunctionNameSuffix(FNameBuilder(CompileSimStageData.StageName).ToView(), StageNameBuilder);
+					
 					const int32 TranslationStageIndex = TranslationStages.Num();
 					FNiagaraHlslTranslationStage& SimStage = TranslationStages.Emplace_GetRef(GraphBridge::GetOutputNodeScriptType(HistoryOutputNode), GraphBridge::GetOutputNodeScriptTypeId(HistoryOutputNode));
 
 					const int32 DestSimStageIndex = CompilationOutput.ScriptData.SimulationStageMetaData.Num();
-					SimStage.PassNamespace = FString::Printf(TEXT("MapSimStage%d%s"), DestSimStageIndex, *StageName);
+					SimStage.PassNamespace = FString::Printf(TEXT("MapSimStage%d%s"), DestSimStageIndex, StageNameBuilder.ToString());
 					SimStage.ChunkModeIndex = (ENiagaraCodeChunkMode)(((int32)ENiagaraCodeChunkMode::SimulationStageBody) + (TranslationStageIndex - 2));
 					if (SimStage.ChunkModeIndex >= ENiagaraCodeChunkMode::SimulationStageBodyMax)
 					{
@@ -2307,7 +2416,9 @@ FNiagaraTranslateResults TNiagaraHlslTranslator<GraphBridge>::Translate(const FN
 
 	//If we're compiling a function then we have all we need already, we don't want to actually generate shader/vm code.
 	if (FunctionCtx())
+	{
 		return TranslateResults;
+	}
 
 	//Now evaluate all the code chunks to generate the shader code.
 	//FString HlslOutput;
@@ -2961,8 +3072,8 @@ void TNiagaraHlslTranslator<GraphBridge>::HandleSimStageSetupAndTeardown(int32 I
 		// This next part might be a big confusing, but because DataInterfaces are in non-editor code, it makes it impossible for them to return graphs or other
 		// structures. We want them to feel free to invoke their own fucntions and not have to do a lot of extra wranging, so we treat them like a custom hlsl node.
 		// The follow section will set up the necessary infrastructure to "Act" like a custom hlsl node for the translator.
-		TranslationStage.CustomReadFunction = FString::Printf(TEXT("SetupFromIterationSource_%s"), *GetSanitizedFunctionNameSuffix(TranslationStage.PassNamespace));
-		TranslationStage.CustomWriteFunction = FString::Printf(TEXT("TeardownFromIterationSource_%s"), *GetSanitizedFunctionNameSuffix(TranslationStage.PassNamespace));
+		TranslationStage.CustomReadFunction = FString::Printf(TEXT("SetupFromIterationSource_%s"), *TranslationStage.PassNamespace);
+		TranslationStage.CustomWriteFunction = FString::Printf(TEXT("TeardownFromIterationSource_%s"), *TranslationStage.PassNamespace);
 		FString SetupFunctionHLSL;
 		FString TeardownFunctionHLSL;
 		FNiagaraFunctionSignature Sig;
@@ -2987,7 +3098,7 @@ void TNiagaraHlslTranslator<GraphBridge>::HandleSimStageSetupAndTeardown(int32 I
 			FString SetupGeneratedHLSL;
 			if (CDO->GenerateSetupHLSL(DIInstanceInfo, MakeArrayView(Sig.Inputs), bSpawnOnly, bPartialWrites, GeneratedErrors, SetupGeneratedHLSL) && SetupGeneratedHLSL.Len() > 0)
 			{
-				Sig.Name = *FString::Printf(TEXT("SetupFromIterationSource_%s_GeneratedSetup"), *GetSanitizedFunctionNameSuffix(TranslationStage.PassNamespace));
+				Sig.Name = *FString::Printf(TEXT("SetupFromIterationSource_%s_GeneratedSetup"), *TranslationStage.PassNamespace);
 				FNiagaraFunctionSignature SignatureOut = Sig;
 				FString SetupOutHLSL;
 				ProcessCustomHlsl(SetupGeneratedHLSL, TranslationStage.ScriptUsage, Sig, Inputs, nullptr, SetupOutHLSL, SignatureOut);
@@ -2998,7 +3109,7 @@ void TNiagaraHlslTranslator<GraphBridge>::HandleSimStageSetupAndTeardown(int32 I
 			FString TeardownGeneratedHLSL;
 			if (CDO->GenerateTeardownHLSL(DIInstanceInfo, MakeArrayView(Sig.Inputs), bSpawnOnly, bPartialWrites, GeneratedErrors, TeardownGeneratedHLSL) && TeardownGeneratedHLSL.Len() > 0)
 			{
-				Sig.Name = *FString::Printf(TEXT("TeardownFromIterationSource_%s_GeneratedTeardown"), *GetSanitizedFunctionNameSuffix(TranslationStage.PassNamespace));
+				Sig.Name = *FString::Printf(TEXT("TeardownFromIterationSource_%s_GeneratedTeardown"), *TranslationStage.PassNamespace);
 				FNiagaraFunctionSignature SignatureOut = Sig;
 				FString TeardownOutHLSL;
 				ProcessCustomHlsl(TeardownGeneratedHLSL, TranslationStage.ScriptUsage, Sig, Inputs, nullptr, TeardownOutHLSL, SignatureOut);
@@ -3020,7 +3131,7 @@ void TNiagaraHlslTranslator<GraphBridge>::HandleSimStageSetupAndTeardown(int32 I
 
 			if (CDO->GenerateIterationSourceNamespaceReadAttributesHLSL(DIInstanceInfo,  IterationSourceVar, MakeArrayView(Sig.Inputs), MakeArrayView(ReadVars), MakeArrayView(AttributeHLSLNames), bSpawnOnly, bSpawnOnly, bPartialWrites, GeneratedErrors, AttributeReadGeneratedHLSL) && AttributeReadGeneratedHLSL.Len() > 0)
 			{
-				Sig.Name = *FString::Printf(TEXT("SetupFromIterationSource_%s_GeneratedReadAttributes"), *GetSanitizedFunctionNameSuffix(TranslationStage.PassNamespace));
+				Sig.Name = *FString::Printf(TEXT("SetupFromIterationSource_%s_GeneratedReadAttributes"), *TranslationStage.PassNamespace);
 				FNiagaraFunctionSignature SignatureOut = Sig;
 				FString AttributeReadOutHLSL;
 				ProcessCustomHlsl(AttributeReadGeneratedHLSL, TranslationStage.ScriptUsage, Sig, Inputs, nullptr, AttributeReadOutHLSL, SignatureOut);
@@ -3041,7 +3152,7 @@ void TNiagaraHlslTranslator<GraphBridge>::HandleSimStageSetupAndTeardown(int32 I
 
 			if (CDO->GenerateIterationSourceNamespaceWriteAttributesHLSL(DIInstanceInfo, IterationSourceVar, MakeArrayView(Sig.Inputs), MakeArrayView(WriteVars), MakeArrayView(AttributeHLSLNames), MakeArrayView(AllVars), bSpawnOnly, bPartialWrites, GeneratedErrors, AttributeWriteGeneratedHLSL) && AttributeWriteGeneratedHLSL.Len() > 0)
 			{
-				Sig.Name = *FString::Printf(TEXT("TeardownFromIterationSource_%s_GeneratedWriteAttributes"), *GetSanitizedFunctionNameSuffix(TranslationStage.PassNamespace));
+				Sig.Name = *FString::Printf(TEXT("TeardownFromIterationSource_%s_GeneratedWriteAttributes"), *TranslationStage.PassNamespace);
 				FNiagaraFunctionSignature SignatureOut = Sig;
 				FString AttributeWriteOutHLSL;
 				ProcessCustomHlsl(AttributeWriteGeneratedHLSL, TranslationStage.ScriptUsage, Sig, Inputs, nullptr, AttributeWriteOutHLSL, SignatureOut);
@@ -3493,11 +3604,12 @@ void FNiagaraHlslTranslator::ConvertCompileInfoToParamInfo(const FNiagaraScriptD
 	GeneratedFunctionSignatures.Reserve(Info.RegisteredFunctions.Num());
 	for (const FNiagaraFunctionSignature& OriginalSig : Info.RegisteredFunctions)
 	{
-		if (SeenFunctions.Contains(OriginalSig))
+		bool bAlreadySeen = false;
+		SeenFunctions.FindOrAdd(OriginalSig, &bAlreadySeen);
+		if (bAlreadySeen)
 		{
 			continue;
 		}
-		SeenFunctions.Add(OriginalSig);
 
 		if (!OriginalSig.bSupportsGPU)
 		{
@@ -4742,7 +4854,7 @@ FString FNiagaraHlslTranslator::GetSanitizedSymbolName(FStringView SymbolName, b
 
 FString FNiagaraHlslTranslator::BuildSanitizedSymbolString(FStringView SymbolName, bool bCollapseNamespaces) const
 {
-	TStringBuilder<1024> StringBuilder;
+	FNameBuilder StringBuilder;
 	AppendSanitizedSymbolString(StringBuilder, SymbolName, bCollapseNamespaces);
 
 	return StringBuilder.ToString();
@@ -4785,27 +4897,69 @@ FString FNiagaraHlslTranslator::GetSanitizedDIFunctionName(const FString& Functi
 	return Sanitized;
 }
 
-FString FNiagaraHlslTranslator::GetSanitizedFunctionNameSuffix(FString Name)
+void FNiagaraHlslTranslator::AppendSanitizedFunctionNameSuffix(FStringView Name, FStringBuilderBase& StringBuilder)
 {
-	if (Name.Len() == 0)
-	{
-		return Name;
+	FNameBuilder SanitizedName;
+
+	{ // single character replacements
+		const FStringView::ElementType UnderscoreKeys[] = { TEXT('.'), TEXT('\\'), TEXT('/'), TEXT(','), TEXT('-'), TEXT(':') };
+		const FStringView::ElementType SkipKeys[] = { TEXT('\t'), TEXT(' ') };
+
+		for (FStringView::ElementType InChar : Name)
+		{
+			if (!Algo::Contains(SkipKeys, InChar))
+			{
+				if (Algo::Contains(UnderscoreKeys, InChar))
+				{
+					SanitizedName.AppendChar(TEXT('_'));
+				}
+				else
+				{
+					SanitizedName.AppendChar(InChar);
+				}
+			}
+		}
 	}
-	FString Ret = Name;
 
-	// remove special characters
-	Ret.ReplaceInline(TEXT("."), TEXT("_"));
-	Ret.ReplaceInline(TEXT("\\"), TEXT("_"));
-	Ret.ReplaceInline(TEXT("/"), TEXT("_"));
-	Ret.ReplaceInline(TEXT(","), TEXT("_"));
-	Ret.ReplaceInline(TEXT("-"), TEXT("_"));
-	Ret.ReplaceInline(TEXT(":"), TEXT("_"));
-	Ret.ReplaceInline(TEXT("\t"), TEXT(""));
-	Ret.ReplaceInline(TEXT(" "), TEXT(""));	
-	Ret.ReplaceInline(TEXT("__"), TEXT("ASC95ASC95")); // Opengl reserves "__" within a name
+	{ // multi character replacements
+		struct FReplacementPair
+		{
+			const TCHAR* Key = nullptr;
+			const TCHAR* Value = nullptr;
+
+			const int32 KeyLen = 0;
+			const int32 ValueLen = 0;
+
+			FReplacementPair(const TCHAR* InKey, const TCHAR* InValue)
+			: Key(InKey)
+			, Value(InValue)
+			, KeyLen(FCString::Strlen(InKey))
+			, ValueLen(FCString::Strlen(InValue))
+			{}
+		};
+
+		const FReplacementPair MultiCharReplacements[] =
+		{
+			{ TEXT("__"), TEXT("ASC95ASC95") }
+		};
+
+		for (const FReplacementPair& Replacement : MultiCharReplacements)
+		{
+			int32 SearchPosition = 0;
+
+			while (SearchPosition != INDEX_NONE)
+			{
+				SearchPosition = SanitizedName.ToView().Find(Replacement.Key, SearchPosition);
+				if (SearchPosition != INDEX_NONE)
+				{
+					SanitizedName.ReplaceAt(SearchPosition, Replacement.KeyLen, Replacement.Value);
+					SearchPosition += Replacement.ValueLen;
+				}
+			}
+		}
+	}
 
-	// Handle internationalization of characters..
-	return ConvertToAsciiString(Ret);
+	NiagaraHlslTranslatorImpl::AppendConvertedToAscii(SanitizedName.ToView(), StringBuilder);
 }
 
 
@@ -7130,7 +7284,7 @@ void TNiagaraHlslTranslator<GraphBridge>::Emitter(const FEmitterNode* EmitterNod
 		}
 
 		// We act like a function call here as the semantics are identical.
-		RegisterFunctionCall(ScriptUsage, Name, FullName, EmitterNode->NodeGuid, GraphBridge::GetEmitterHandleIdString(EmitterNode), EmitterGraph, Signature, false, FString(), {}, Inputs, CallInputs, CallOutputs, Signature);
+		RegisterFunctionCall(ScriptUsage, Name, FName(FullName), EmitterNode->NodeGuid, GraphBridge::GetEmitterHandleIdString(EmitterNode), EmitterGraph, Signature, false, FString(), {}, Inputs, CallInputs, CallOutputs, Signature);
 		GenerateFunctionCall(ScriptUsage, Signature, Inputs, Outputs);
 
 		// Clear out the parameter map writes to emitter module parameters as they should not be shared across emitters.
@@ -8418,7 +8572,7 @@ void TNiagaraHlslTranslator<GraphBridge>::FunctionCall(const FFunctionCallNode*
 		HandleCustomHlslNode(CustomFunctionHlsl, ScriptUsage, Name, FullName, bCustomHlsl, CustomHlsl, CustomHlslIncludeFilePaths, Signature, Inputs);
 	}
 
-	RegisterFunctionCall(ScriptUsage, Name, FullName, FunctionNode->NodeGuid, FString(), FunctionGraph, Signature, bCustomHlsl, CustomHlsl, CustomHlslIncludeFilePaths, Inputs, CallInputs, CallOutputs, OutputSignature);
+	RegisterFunctionCall(ScriptUsage, Name, FName(FullName), FunctionNode->NodeGuid, FString(), FunctionGraph, Signature, bCustomHlsl, CustomHlsl, CustomHlslIncludeFilePaths, Inputs, CallInputs, CallOutputs, OutputSignature);
 
 	if (OutputSignature.IsValid() == false)
 	{
@@ -9173,8 +9327,8 @@ bool IsVariableWriteBeforeRead(const TArray<typename GraphBridge::FParamMapHisto
 }
 
 template<typename GraphBridge>
-void TNiagaraHlslTranslator<GraphBridge>::RegisterFunctionCall(ENiagaraScriptUsage ScriptUsage, const FString& InName, const FString& InFullName, const FGuid& CallNodeId, const FString& InFunctionNameSuffix, const FGraph* SourceGraph,
-                                                  FNiagaraFunctionSignature& InSignature, bool bIsCustomHlsl, const FString& InCustomHlsl, const TArray<FNiagaraCustomHlslInclude>& InCustomHlslIncludeFilePaths, TArray<int32>& Inputs, TConstArrayView<const FInputPin*> CallInputs, TConstArrayView<const FOutputPin*> CallOutputs,
+void TNiagaraHlslTranslator<GraphBridge>::RegisterFunctionCall(ENiagaraScriptUsage ScriptUsage, const FString& InName, FName InFullName, const FGuid& CallNodeId, const FString& InFunctionNameSuffix, const FGraph* SourceGraph,
+                                                  const FNiagaraFunctionSignature& InSignature, bool bIsCustomHlsl, const FString& InCustomHlsl, const TArray<FNiagaraCustomHlslInclude>& InCustomHlslIncludeFilePaths, TArray<int32>& Inputs, TConstArrayView<const FInputPin*> CallInputs, TConstArrayView<const FOutputPin*> CallOutputs,
                                                   FNiagaraFunctionSignature& OutSignature)
 {
 	NIAGARA_SCOPE_CYCLE_COUNTER(STAT_NiagaraEditor_Module_NiagaraHLSLTranslator_RegisterFunctionCall);
@@ -9184,23 +9338,8 @@ void TNiagaraHlslTranslator<GraphBridge>::RegisterFunctionCall(ENiagaraScriptUsa
 	{
 		NIAGARA_SCOPE_CYCLE_COUNTER(STAT_NiagaraEditor_Module_NiagaraHLSLTranslator_RegisterFunctionCall_Source);
 
-		const bool bHasNumericInputs = GraphBridge::GraphHasParametersOfType(SourceGraph, FNiagaraTypeDefinition::GetGenericNumericDef());
-		const bool bHasParameterMapParameters = GraphBridge::GraphHasParametersOfType(SourceGraph, FNiagaraTypeDefinition::GetParameterMapDef());
-
-		TArray<const FInputPin*> StaticSwitchValues;
-		for (const FNiagaraVariableBase& StaticSwitchInput : GraphBridge::GraphGetStaticSwitchInputs(SourceGraph))
-		{
-			for (const FInputPin* Pin : CallInputs)
-			{
-				if (StaticSwitchInput.GetName().IsEqual(Pin->PinName))
-				{
-					StaticSwitchValues.Add(Pin);
-					break;
-				}
-			}
-		}
-
-		GenerateFunctionSignature(ScriptUsage, InName, InFullName, InFunctionNameSuffix, SourceGraph, Inputs, bHasNumericInputs, bHasParameterMapParameters, StaticSwitchValues, OutSignature);
+		FString FunctionComment;
+		GenerateFunctionSignature(ScriptUsage, InName, InFullName, InFunctionNameSuffix, SourceGraph, CallInputs, Inputs, FunctionComment, OutSignature);
 
 		// 		//Sort the input and outputs to match the sorted parameters. They may be different.
 		// 		TArray<FNiagaraVariable> OrderedInputs;
@@ -9345,6 +9484,7 @@ void TNiagaraHlslTranslator<GraphBridge>::RegisterFunctionCall(ENiagaraScriptUsa
 				//Grab all the body chunks for this function.
 				FunctionBody.StageIndices.AddUnique(ActiveStageIdx);
 				FunctionBody.Body.Reserve(256 * ChunksByMode[(int32)ENiagaraCodeChunkMode::Body].Num());
+				FunctionBody.Comments = MoveTemp(FunctionComment);
 
 				for (int32 i = ChunkStartsByMode[(int32)ENiagaraCodeChunkMode::Body]; i < ChunksByMode[(int32)ENiagaraCodeChunkMode::Body].Num(); ++i)
 				{
@@ -9410,8 +9550,8 @@ void TNiagaraHlslTranslator<GraphBridge>::RegisterFunctionCall(ENiagaraScriptUsa
 					FunctionBody.Body += TEXT("\n");
 				}
 
-				FunctionStageWriteTargets.Add(OutSignature, ActiveStageWriteTargets.Top());
-				FunctionStageReadTargets.Add(OutSignature, ActiveStageReadTargets.Top());
+				FunctionBody.StageWriteTargets = ActiveStageWriteTargets.Top();
+				FunctionBody.StageReadTargets = ActiveStageReadTargets.Top();
 			}
 
 			ExitFunction();
@@ -9421,26 +9561,18 @@ void TNiagaraHlslTranslator<GraphBridge>::RegisterFunctionCall(ENiagaraScriptUsa
 			FuncBody->StageIndices.AddUnique(ActiveStageIdx);
 
 			// Just because we had a cached call, doesn't mean that we should ignore adding read or writetargets
-			TArray<FName>* Entries = FunctionStageWriteTargets.Find(OutSignature);
-			if (Entries)
+			for (const FName& Entry : FuncBody->StageWriteTargets)
 			{
-				for (const FName& Entry : *Entries)
-				{
-					const int32 SourceSimStage = TranslationStages[ActiveStageIdx].SimulationStageIndex;
-					CompilationOutput.ScriptData.SimulationStageMetaData[SourceSimStage].OutputDestinations.AddUnique(Entry);
-					ActiveStageWriteTargets.Top().AddUnique(Entry);
-				}
+				const int32 SourceSimStage = TranslationStages[ActiveStageIdx].SimulationStageIndex;
+				CompilationOutput.ScriptData.SimulationStageMetaData[SourceSimStage].OutputDestinations.AddUnique(Entry);
+				ActiveStageWriteTargets.Top().AddUnique(Entry);
 			}
 
-			Entries = FunctionStageReadTargets.Find(OutSignature);
-			if (Entries)
+			for (const FName& Entry : FuncBody->StageReadTargets)
 			{
-				for (const FName& Entry : *Entries)
-				{
-					const int32 SourceSimStage = TranslationStages[ActiveStageIdx].SimulationStageIndex;
-					CompilationOutput.ScriptData.SimulationStageMetaData[SourceSimStage].InputDataInterfaces.AddUnique(Entry);
-					ActiveStageReadTargets.Top().AddUnique(Entry);
-				}
+				const int32 SourceSimStage = TranslationStages[ActiveStageIdx].SimulationStageIndex;
+				CompilationOutput.ScriptData.SimulationStageMetaData[SourceSimStage].InputDataInterfaces.AddUnique(Entry);
+				ActiveStageReadTargets.Top().AddUnique(Entry);
 			}
 		}
 	}
@@ -9471,8 +9603,8 @@ void TNiagaraHlslTranslator<GraphBridge>::RegisterFunctionCall(ENiagaraScriptUsa
 					FunctionBody.Body += TEXT("\n");
 				}
 
-				FunctionStageWriteTargets.Add(OutSignature, ActiveStageWriteTargets.Top());
-				FunctionStageReadTargets.Add(OutSignature, ActiveStageReadTargets.Top());
+				FunctionBody.StageWriteTargets = ActiveStageWriteTargets.Top();
+				FunctionBody.StageReadTargets = ActiveStageReadTargets.Top();
 
 				for (const FNiagaraCustomHlslInclude& FileInclude : InCustomHlslIncludeFilePaths)
 				{
diff --git a/Engine/Plugins/FX/Niagara/Source/NiagaraEditor/Private/NiagaraHlslTranslator.h b/Engine/Plugins/FX/Niagara/Source/NiagaraEditor/Private/NiagaraHlslTranslator.h
--- a/Engine/Plugins/FX/Niagara/Source/NiagaraEditor/Private/NiagaraHlslTranslator.h
+++ b/Engine/Plugins/FX/Niagara/Source/NiagaraEditor/Private/NiagaraHlslTranslator.h
@@ -278,15 +278,15 @@ protected:
 	struct FNiagaraFunctionBody
 	{
 		FString Body;
+		FString Comments;
 		TArray<int32> StageIndices;
+		TArray<FName> StageWriteTargets;
+		TArray<FName> StageReadTargets;
 	};
 
 	TMap<FNiagaraFunctionSignature, FNiagaraFunctionBody> Functions;
 	TArray<FNiagaraCustomHlslInclude> FunctionIncludeFilePaths;
-	TMap<FNiagaraFunctionSignature, TArray<FName>> FunctionStageWriteTargets;
 	TArray<TArray<FName>> ActiveStageWriteTargets;
-
-	TMap<FNiagaraFunctionSignature, TArray<FName>> FunctionStageReadTargets;
 	TArray<TArray<FName>> ActiveStageReadTargets;
 
 	void GenerateFunctionCall(ENiagaraScriptUsage ScriptUsage, FNiagaraFunctionSignature& FunctionSignature, TArrayView<const int32> Inputs, TArray<int32>& Outputs);
@@ -397,7 +397,7 @@ public:
 	static NIAGARAEDITOR_API TArray<FName> ConditionPropertyPath(const FNiagaraTypeDefinition& Type, const TArray<FName>& InPath);
 	static NIAGARAEDITOR_API FString GetSanitizedSymbolName(FStringView SymbolName, bool bCollapseNamespaces = false);
 	static NIAGARAEDITOR_API FString GetSanitizedDIFunctionName(const FString& FunctionName);
-	static NIAGARAEDITOR_API FString GetSanitizedFunctionNameSuffix(FString Name);
+	static NIAGARAEDITOR_API void AppendSanitizedFunctionNameSuffix(FStringView Name, FStringBuilderBase& StringBuilder);
 
 	/** Replaces all non-ascii characters with a "ASCXXX" string, where XXX is their int value */
 	static NIAGARAEDITOR_API FString ConvertToAsciiString(FString Name);
@@ -635,7 +635,7 @@ public:
 
 	virtual FNiagaraTranslateResults Translate(const FNiagaraCompileOptions& InCompileOptions, const FHlslNiagaraTranslatorOptions& Options) override;
 
-	void RegisterFunctionCall(ENiagaraScriptUsage ScriptUsage, const FString& InName, const FString& InFullName, const FGuid& CallNodeId, const FString& InFunctionNameSuffix, const FGraph* SourceGraph, FNiagaraFunctionSignature& InSignature, bool bIsCustomHlsl, const FString& InCustomHlsl, const TArray<FNiagaraCustomHlslInclude>& InCustomHlslIncludeFilePaths, TArray<int32>& Inputs, TConstArrayView<const FInputPin*> CallInputs, TConstArrayView<const FOutputPin*> CallOutputs,
+	void RegisterFunctionCall(ENiagaraScriptUsage ScriptUsage, const FString& InName, FName InFullName, const FGuid& CallNodeId, const FString& InFunctionNameSuffix, const FGraph* SourceGraph, const FNiagaraFunctionSignature& InSignature, bool bIsCustomHlsl, const FString& InCustomHlsl, const TArray<FNiagaraCustomHlslInclude>& InCustomHlslIncludeFilePaths, TArray<int32>& Inputs, TConstArrayView<const FInputPin*> CallInputs, TConstArrayView<const FOutputPin*> CallOutputs,
 		FNiagaraFunctionSignature& OutSignature);
 	/** Compiles an output Pin on a graph node. Caches the result for any future inputs connected to it. */
 	int32 CompileOutputPin(const FPin* Pin);
@@ -744,8 +744,9 @@ protected:
 	// Checks that the Partices.ID parameter is only used if persistent IDs are active
 	void ValidateParticleIDUsage();
 	bool ValidateTypePins(const FNode* NodeToValidate);
-	void GenerateFunctionSignature(ENiagaraScriptUsage ScriptUsage, FString InName, const FString& InFullName, const FString& InFunctionNameSuffix, const FGraph* FuncGraph, TArray<int32>& Inputs,
-		bool bHadNumericInputs, bool bHasParameterMapParameters, const TArray<const FInputPin*>& StaticSwitchValues, FNiagaraFunctionSignature& OutSig) const;
+	void GenerateFunctionSignature(ENiagaraScriptUsage ScriptUsage, FStringView InName, FName InFullName, FStringView InFunctionNameSuffix, const FGraph* FuncGraph, TConstArrayView<const FInputPin*> CallInputs, TArray<int32>& Inputs,
+		FString& FunctionComment, FNiagaraFunctionSignature& OutSig) const;
+	void ResolveGraphFunctionAlias(const FNiagaraGraphFunctionAliasContext& FunctionAliasContext, FStringBuilderBase& FunctionAliasSuffix) const;
 
 	void ValidateFailIfPreviouslyNotSet(const FNiagaraVariable& InVar, bool& bFailIfNotSet);
 
diff --git a/Engine/Plugins/FX/Niagara/Source/NiagaraEditor/Private/NiagaraNodeInput.cpp b/Engine/Plugins/FX/Niagara/Source/NiagaraEditor/Private/NiagaraNodeInput.cpp
--- a/Engine/Plugins/FX/Niagara/Source/NiagaraEditor/Private/NiagaraNodeInput.cpp
+++ b/Engine/Plugins/FX/Niagara/Source/NiagaraEditor/Private/NiagaraNodeInput.cpp
@@ -126,14 +126,15 @@ void UNiagaraNodeInput::BuildParameterMapHistory(FNiagaraParameterMapHistoryBuil
 	}
 }
 
-void UNiagaraNodeInput::AppendFunctionAliasForContext(const FNiagaraGraphFunctionAliasContext& InFunctionAliasContext, FString& InOutFunctionAlias, bool& OutOnlyOncePerNodeType) const
+void UNiagaraNodeInput::UpdateFunctionAliasForContext(FNiagaraGraphFunctionAliasContext& InOutFunctionAliasContext) const
 {
 	if (Usage == ENiagaraInputNodeUsage::TranslatorConstant && Input == TRANSLATOR_PARAM_CALL_ID)
 	{
-		OutOnlyOncePerNodeType = true;
+		static const FName NAME_ScriptUsage = TEXT("ScriptUsage");
+
 		// The call ID should be unique for each translated node as it is used by the seeded random functions.
 		// We don't want it to be shared across the spawn and update script, so functions including it will have the usage added to their name.
-		InOutFunctionAlias += "_ScriptUsage" + FString::FormatAsNumber((uint8)InFunctionAliasContext.ScriptUsage);
+		InOutFunctionAliasContext.AddPermutationIdentifier(NAME_ScriptUsage, FString::FormatAsNumber(static_cast<uint8>(InOutFunctionAliasContext.ScriptUsage)));
 	}
 }
 
@@ -153,7 +154,7 @@ void UNiagaraNodeInput::AllocateDefaultPins()
 
 	const UEdGraphSchema_Niagara* Schema = GetDefault<UEdGraphSchema_Niagara>();
 
-	//If we're a parameter node for a funciton or a module the we allow a "default" input pin.
+	//If we're a parameter node for a function or a module the we allow a "default" input pin.
 	UNiagaraScript* OwnerScript = GetTypedOuter<UNiagaraScript>();
 	if (OwnerScript)
 	{
diff --git a/Engine/Plugins/FX/Niagara/Source/NiagaraEditor/Private/NiagaraNodeStaticSwitch.cpp b/Engine/Plugins/FX/Niagara/Source/NiagaraEditor/Private/NiagaraNodeStaticSwitch.cpp
--- a/Engine/Plugins/FX/Niagara/Source/NiagaraEditor/Private/NiagaraNodeStaticSwitch.cpp
+++ b/Engine/Plugins/FX/Niagara/Source/NiagaraEditor/Private/NiagaraNodeStaticSwitch.cpp
@@ -379,6 +379,17 @@ FName UNiagaraNodeStaticSwitch::GetOptionPinName(const FNiagaraVariable& Variabl
 	return FName(Variable.GetName().ToString() + FString::Printf(TEXT(" if %s"), *Suffix));
 }
 
+void UNiagaraNodeStaticSwitch::UpdateFunctionAliasForContext(FNiagaraGraphFunctionAliasContext& InOutFunctionAliasContext) const
+{
+	if (IsSetByCompiler())
+	{
+		if (const FNiagaraVariable* Found = FNiagaraConstants::FindStaticSwitchConstant(SwitchTypeData.SwitchConstant))
+		{
+			InOutFunctionAliasContext.AddPendingStaticSwitchVariable(*Found, static_cast<uint8>(SwitchTypeData.SwitchType));
+		}
+	}
+}
+
 void UNiagaraNodeStaticSwitch::ChangeSwitchParameterName(const FName& NewName)
 {
 	GetNiagaraGraph()->RenameStaticSwitch(this, NewName);
diff --git a/Engine/Plugins/FX/Niagara/Source/NiagaraEditor/Private/NiagaraNodeStaticSwitch.h b/Engine/Plugins/FX/Niagara/Source/NiagaraEditor/Private/NiagaraNodeStaticSwitch.h
--- a/Engine/Plugins/FX/Niagara/Source/NiagaraEditor/Private/NiagaraNodeStaticSwitch.h
+++ b/Engine/Plugins/FX/Niagara/Source/NiagaraEditor/Private/NiagaraNodeStaticSwitch.h
@@ -122,6 +122,8 @@ protected:
 	
 	virtual FString GetInputCaseName(int32 Case) const override;
 	virtual FName GetOptionPinName(const FNiagaraVariable& Variable, int32 Value) const override;
+
+	virtual void UpdateFunctionAliasForContext(FNiagaraGraphFunctionAliasContext& InOutFunctionAliasContext) const override;
 	//~ End UNiagaraNodeUsageSelector Interface
 
 	virtual bool CanModifyPin(const UEdGraphPin* Pin) const override;
diff --git a/Engine/Plugins/FX/Niagara/Source/NiagaraEditor/Private/NiagaraNodeUsageSelector.cpp b/Engine/Plugins/FX/Niagara/Source/NiagaraEditor/Private/NiagaraNodeUsageSelector.cpp
--- a/Engine/Plugins/FX/Niagara/Source/NiagaraEditor/Private/NiagaraNodeUsageSelector.cpp
+++ b/Engine/Plugins/FX/Niagara/Source/NiagaraEditor/Private/NiagaraNodeUsageSelector.cpp
@@ -401,34 +401,28 @@ UEdGraphPin* UNiagaraNodeUsageSelector::GetPassThroughPin(const UEdGraphPin* Loc
 	return nullptr;
 }
 
-void UNiagaraNodeUsageSelector::AppendFunctionAliasForContext(const FNiagaraGraphFunctionAliasContext& InFunctionAliasContext, FString& InOutFunctionAlias, bool& OutOnlyOncePerNodeType) const
+void UNiagaraNodeUsageSelector::UpdateFunctionAliasForContext(FNiagaraGraphFunctionAliasContext& InOutFunctionAliasContext) const
 {
-	OutOnlyOncePerNodeType = true;
-	
-	FString UsageString;
-	switch (InFunctionAliasContext.CompileUsage)
+	static const FName NAME_Usage = TEXT("Usage");
+
+	switch (InOutFunctionAliasContext.CompileUsage)
 	{
 	case ENiagaraScriptUsage::SystemSpawnScript:
 	case ENiagaraScriptUsage::SystemUpdateScript:
-		UsageString = "System";
+		InOutFunctionAliasContext.AddPermutationIdentifier(NAME_Usage, TEXT("System"));
 		break;
 	case ENiagaraScriptUsage::EmitterSpawnScript:
 	case ENiagaraScriptUsage::EmitterUpdateScript:
-		UsageString = "Emitter";
+		InOutFunctionAliasContext.AddPermutationIdentifier(NAME_Usage, TEXT("Emitter"));
 		break;
 	case ENiagaraScriptUsage::ParticleSpawnScript:
 	case ENiagaraScriptUsage::ParticleUpdateScript:
 	case ENiagaraScriptUsage::ParticleEventScript:
 	case ENiagaraScriptUsage::ParticleSimulationStageScript:
 	case ENiagaraScriptUsage::ParticleGPUComputeScript:
-		UsageString = "Particle";
+		InOutFunctionAliasContext.AddPermutationIdentifier(NAME_Usage, TEXT("Particle"));
 		break;
 	}
-
-	if (UsageString.IsEmpty() == false)
-	{
-		InOutFunctionAlias += "_" + UsageString;
-	}
 }
 
 void UNiagaraNodeUsageSelector::BuildParameterMapHistory(FNiagaraParameterMapHistoryBuilder& OutHistory, bool bRecursive /*= true*/, bool bFilterForCompilation /*= true*/) const
diff --git a/Engine/Plugins/FX/Niagara/Source/NiagaraEditor/Private/NiagaraNodeUsageSelector.h b/Engine/Plugins/FX/Niagara/Source/NiagaraEditor/Private/NiagaraNodeUsageSelector.h
--- a/Engine/Plugins/FX/Niagara/Source/NiagaraEditor/Private/NiagaraNodeUsageSelector.h
+++ b/Engine/Plugins/FX/Niagara/Source/NiagaraEditor/Private/NiagaraNodeUsageSelector.h
@@ -41,7 +41,7 @@ public:
 	virtual bool AllowNiagaraTypeForPinTypeChange(const FNiagaraTypeDefinition& InType, UEdGraphPin* Pin) const override;
 	virtual bool OnNewPinTypeRequested(UEdGraphPin* PinToChange, FNiagaraTypeDefinition NewType) override;
 	virtual UEdGraphPin* GetPassThroughPin(const UEdGraphPin* LocallyOwnedOutputPin, ENiagaraScriptUsage InUsage) const override;
-	virtual void AppendFunctionAliasForContext(const FNiagaraGraphFunctionAliasContext& InFunctionAliasContext, FString& InOutFunctionAlias, bool& OutOnlyOncePerNodeType) const override;
+	virtual void UpdateFunctionAliasForContext(FNiagaraGraphFunctionAliasContext& InOutFunctionAliasContext) const override;
 	//~ End UNiagaraNode Interface
 
 	/** Helper function to create a variable to add to the OutputVars and FGuid to add to OutputVarGuids. */
diff --git a/Engine/Plugins/FX/Niagara/Source/NiagaraEditor/Public/NiagaraGraph.h b/Engine/Plugins/FX/Niagara/Source/NiagaraEditor/Public/NiagaraGraph.h
--- a/Engine/Plugins/FX/Niagara/Source/NiagaraEditor/Public/NiagaraGraph.h
+++ b/Engine/Plugins/FX/Niagara/Source/NiagaraEditor/Public/NiagaraGraph.h
@@ -142,7 +142,32 @@ struct FNiagaraGraphFunctionAliasContext
 
 	// the usage as defined in the current translation stage
 	ENiagaraScriptUsage ScriptUsage;
-	TArray<const UEdGraphPin*> StaticSwitchValues;
+	
+	void AppendPermutationString(FStringBuilderBase& StringBuilder) const;
+	void AddPermutationIdentifier(const FName& Name, FStringView Value);
+	void AddPendingStaticSwitchVariable(const FNiagaraVariable& Variable, uint8 SwitchType);
+
+	struct FStaticSwitchVariable
+	{
+		FNiagaraVariable Variable;
+		uint8 SwitchType;
+
+		friend bool operator==(const FStaticSwitchVariable& Lhs, const FStaticSwitchVariable& Rhs)
+		{
+			return Lhs.Variable == Rhs.Variable && Lhs.SwitchType == Rhs.SwitchType;
+		}
+	};
+
+	TConstArrayView<FStaticSwitchVariable> ReadPendingStaticSwitchVariables() const
+	{
+		return CompilerStaticSwitchVariables;
+	}
+
+private:
+	// the set of identifiers from internal graph nodes that permute the function
+	TMap<FName, FString> PermutationIdentifiers;
+
+	TArray<FStaticSwitchVariable> CompilerStaticSwitchVariables;
 };
 
 /** Options for the FindInputNodes function */
@@ -414,7 +439,7 @@ public:
 	  * function calls to the same graph. For example, if the graph contains static switches and two functions call it with
 	  * different switch parameters, the final function names in the hlsl must be different.
 	  */
-	FString GetFunctionAliasByContext(const FNiagaraGraphFunctionAliasContext& FunctionAliasContext) const;
+	void UpdateFunctionAliasByContext(FNiagaraGraphFunctionAliasContext& FunctionAliasContext) const;
 
 	/** In order to support reducing the number of times we need to fully generate the CompileId this function is introduced
 	  * to work with LastBuiltScriptVersionId & bHasValidLastBuiltScriptVersionId.
diff --git a/Engine/Plugins/FX/Niagara/Source/NiagaraEditor/Public/NiagaraNode.h b/Engine/Plugins/FX/Niagara/Source/NiagaraEditor/Public/NiagaraNode.h
--- a/Engine/Plugins/FX/Niagara/Source/NiagaraEditor/Public/NiagaraNode.h
+++ b/Engine/Plugins/FX/Niagara/Source/NiagaraEditor/Public/NiagaraNode.h
@@ -188,7 +188,7 @@ public:
 
 	NIAGARAEDITOR_API FOnNodeVisualsChanged& OnVisualsChanged();
 
-	virtual void AppendFunctionAliasForContext(const FNiagaraGraphFunctionAliasContext& InFunctionAliasContext, FString& InOutFunctionAlias, bool& OutOnlyOncePerNodeType) const { };
+	virtual void UpdateFunctionAliasForContext(FNiagaraGraphFunctionAliasContext& InOutFunctionAliasContext) const { };
 
 	/** Old style compile hash code. To be removed in the future.*/
 	NIAGARAEDITOR_API virtual void UpdateCompileHashForNode(FSHA1& HashState) const;
diff --git a/Engine/Plugins/FX/Niagara/Source/NiagaraEditor/Public/NiagaraNodeInput.h b/Engine/Plugins/FX/Niagara/Source/NiagaraEditor/Public/NiagaraNodeInput.h
--- a/Engine/Plugins/FX/Niagara/Source/NiagaraEditor/Public/NiagaraNodeInput.h
+++ b/Engine/Plugins/FX/Niagara/Source/NiagaraEditor/Public/NiagaraNodeInput.h
@@ -93,7 +93,7 @@ public:
 	bool ReferencesSameInput(UNiagaraNodeInput* Other) const;
 
 	virtual void BuildParameterMapHistory(FNiagaraParameterMapHistoryBuilder& OutHistory, bool bRecursive = true, bool bFilterForCompilation = true) const override;
-	virtual void AppendFunctionAliasForContext(const FNiagaraGraphFunctionAliasContext& InFunctionAliasContext, FString& InOutFunctionAlias, bool& OutOnlyOncePerNodeType) const override;
+	virtual void UpdateFunctionAliasForContext(FNiagaraGraphFunctionAliasContext& InOutFunctionAliasContext) const override;
 
 	static const FLinearColor TitleColor_Attribute;
 	static const FLinearColor TitleColor_Constant;
diff --git a/Engine/Source/Runtime/Core/Private/UObject/FortniteMainBranchObjectVersion.cpp b/Engine/Source/Runtime/Core/Private/UObject/FortniteMainBranchObjectVersion.cpp
--- a/Engine/Source/Runtime/Core/Private/UObject/FortniteMainBranchObjectVersion.cpp
+++ b/Engine/Source/Runtime/Core/Private/UObject/FortniteMainBranchObjectVersion.cpp
@@ -12,7 +12,7 @@ TMap<FGuid, FGuid> FFortniteMainBranchObjectVersion::GetSystemGuids()
 	SystemGuids.Add(DevGuids.LANDSCAPE_MOBILE_COOK_VERSION, FGuid("32D02EF867C74B71A0D4E0FA41392732"));
 	SystemGuids.Add(DevGuids.MATERIALSHADERMAP_DERIVEDDATA_VER, FGuid("FB0EB87E970C4FB9B217E0E00942BFCB"));
 	SystemGuids.Add(DevGuids.NIAGARASHADERMAP_DERIVEDDATA_VER, FGuid("87F6AE6C269A4592B281752A582E7F13"));
-	SystemGuids.Add(DevGuids.Niagara_LatestScriptCompileVersion, FGuid("FB15C2870A4E4AE8B5F58566F0E3825E"));
+	SystemGuids.Add(DevGuids.Niagara_LatestScriptCompileVersion, FGuid("132E723405764AFAB45E9B353147C87D"));
 	SystemGuids.Add(DevGuids.MaterialTranslationDDCVersion, FGuid("B79683EE9E684C13BAC2FDE8B180477F"));
 
 	// These GUIDs are in AppendRestrictedFortniteMainStreamObjectVersionGuids as they should rarely be changed.
