- Allow FNiagaraUserParameterBinding to specify a list of other types that can be included when looking for UObject parameters

#Synced-CL 51628270

diff --git a/Engine/Plugins/FX/Niagara/Source/Niagara/Internal/DataInterface/NiagaraDataInterfaceStaticMesh.h b/Engine/Plugins/FX/Niagara/Source/Niagara/Internal/DataInterface/NiagaraDataInterfaceStaticMesh.h
--- a/Engine/Plugins/FX/Niagara/Source/Niagara/Internal/DataInterface/NiagaraDataInterfaceStaticMesh.h
+++ b/Engine/Plugins/FX/Niagara/Source/Niagara/Internal/DataInterface/NiagaraDataInterfaceStaticMesh.h
@@ -113,7 +113,7 @@ public:
 	FNiagaraUserParameterBinding LODIndexUserParameter;
 
 	/** Mesh parameter binding can be either an Actor (in which case we find the component), static mesh component or a static mesh. */
-	UPROPERTY(EditAnywhere, Category = "Mesh")
+	UPROPERTY(EditAnywhere, Category = "Mesh", meta = (AdditionalAllowedClasses="/Script/Engine.StaticMesh"))
 	FNiagaraUserParameterBinding MeshParameterBinding;
 
 	/** When attached to an Instanced Static Mesh, which instance should be read from. */
diff --git a/Engine/Plugins/FX/Niagara/Source/Niagara/Public/NiagaraParameterStore.h b/Engine/Plugins/FX/Niagara/Source/Niagara/Public/NiagaraParameterStore.h
--- a/Engine/Plugins/FX/Niagara/Source/Niagara/Public/NiagaraParameterStore.h
+++ b/Engine/Plugins/FX/Niagara/Source/Niagara/Public/NiagaraParameterStore.h
@@ -1124,7 +1124,8 @@ struct FNiagaraParameterDirectBinding<UObject*>
 			LayoutVersion = BoundStore->GetLayoutVersion();
 
 			check(BoundVariable.GetType().IsUObject());
-			UObjectOffset = BoundStore->IndexOf(DestVariable);
+			const FNiagaraVariableWithOffset* VariableWithOffset = BoundStore->FindParameterVariable(DestVariable, true);
+			UObjectOffset = VariableWithOffset && VariableWithOffset->IsUObject() && !VariableWithOffset->IsDataInterface() ? VariableWithOffset->Offset : INDEX_NONE;
 			TObjectPtr<UObject> Ret = BoundStore->GetUObject(UObjectOffset);
 			return Ret;
 		}
diff --git a/Engine/Plugins/FX/Niagara/Source/NiagaraEditor/Private/Customizations/NiagaraTypeCustomizations.cpp b/Engine/Plugins/FX/Niagara/Source/NiagaraEditor/Private/Customizations/NiagaraTypeCustomizations.cpp
--- a/Engine/Plugins/FX/Niagara/Source/NiagaraEditor/Private/Customizations/NiagaraTypeCustomizations.cpp
+++ b/Engine/Plugins/FX/Niagara/Source/NiagaraEditor/Private/Customizations/NiagaraTypeCustomizations.cpp
@@ -1048,10 +1048,29 @@ TArray<FName> FNiagaraUserParameterBindingCustomization::GetNames() const
 
 	if (IsValid() && TargetUserParameterBinding)
 	{
+		// Build list of allowed types
+		TArray<FNiagaraTypeDefinition> AllowedTypes;
+		AllowedTypes.Add(TargetUserParameterBinding->Parameter.GetType());
+		{
+			TArray<FString> AdditionalAllowedClasses;
+			const FString& AdditionalAllowedClassesMetadata = PropertyHandle->GetProperty()->GetMetaData(TEXT("AdditionalAllowedClasses"));
+			if (!AdditionalAllowedClassesMetadata.IsEmpty())
+			{
+				AdditionalAllowedClassesMetadata.ParseIntoArray(AdditionalAllowedClasses, TEXT(","), true);
+				for (const FString& AllowedClass : AdditionalAllowedClasses)
+				{
+					if ( UClass* Class = FindObject<UClass>(nullptr, *AllowedClass) )
+					{
+						AllowedTypes.Add(FNiagaraTypeDefinition(Class));
+					}
+				}
+			}
+		}
+
 		UNiagaraSystem* BaseSystem = BaseSystemWeakPtr.Get();
 		for (const FNiagaraVariable Var : BaseSystem->GetExposedParameters().ReadParameterVariables())
 		{
-			if (FNiagaraParameterUtilities::IsUserParameter(Var) && Var.GetType() == TargetUserParameterBinding->Parameter.GetType())
+			if (FNiagaraParameterUtilities::IsUserParameter(Var) && AllowedTypes.Contains(Var.GetType()))
 			{
 				Names.AddUnique(Var.GetName());
 			}
