- Remove insert above / below in summary view

#Synced-CL 51672242

diff --git a/Engine/Plugins/FX/Niagara/Source/NiagaraEditorWidgets/Private/NiagaraEditorWidgetsUtilities.cpp b/Engine/Plugins/FX/Niagara/Source/NiagaraEditorWidgets/Private/NiagaraEditorWidgetsUtilities.cpp
--- a/Engine/Plugins/FX/Niagara/Source/NiagaraEditorWidgets/Private/NiagaraEditorWidgetsUtilities.cpp
+++ b/Engine/Plugins/FX/Niagara/Source/NiagaraEditorWidgets/Private/NiagaraEditorWidgetsUtilities.cpp
@@ -458,18 +458,23 @@ void ShowInsertModuleMenu(TWeakObjectPtr<UNiagaraStackModuleItem> StackModuleIte
 bool FNiagaraStackEditorWidgetsUtilities::AddStackModuleItemContextMenuActions(FMenuBuilder& MenuBuilder, UNiagaraStackModuleItem& StackModuleItem, TSharedRef<SWidget> TargetWidget)
 {
 	MenuBuilder.BeginSection("ModuleActions", LOCTEXT("ModuleActions", "Module Actions"));
-	{		
-		MenuBuilder.AddMenuEntry(
-			LOCTEXT("InsertModuleAbove", "Insert Above"),
-			LOCTEXT("InsertModuleAboveToolTip", "Insert a new module above this module in the stack."),
-			FSlateIcon(),
-			FUIAction(FExecuteAction::CreateStatic(&ShowInsertModuleMenu, TWeakObjectPtr<UNiagaraStackModuleItem>(&StackModuleItem), 0, TWeakPtr<SWidget>(TargetWidget))));
+	{
+		// Insert Above/Below require valid group add utilities. These are not available when the module
+		// item is displayed in the Emitter Summary context (utilities are null), so skip those entries.
+		if (StackModuleItem.GetGroupAddUtilities() != nullptr)
+		{
+			MenuBuilder.AddMenuEntry(
+				LOCTEXT("InsertModuleAbove", "Insert Above"),
+				LOCTEXT("InsertModuleAboveToolTip", "Insert a new module above this module in the stack."),
+				FSlateIcon(),
+				FUIAction(FExecuteAction::CreateStatic(&ShowInsertModuleMenu, TWeakObjectPtr<UNiagaraStackModuleItem>(&StackModuleItem), 0, TWeakPtr<SWidget>(TargetWidget))));
 
-		MenuBuilder.AddMenuEntry(
-			LOCTEXT("InsertModuleBelow", "Insert Below"),
-			LOCTEXT("InsertModuleBelowToolTip", "Insert a new module below this module in the stack."),
-			FSlateIcon(),
-			FUIAction(FExecuteAction::CreateStatic(&ShowInsertModuleMenu, TWeakObjectPtr<UNiagaraStackModuleItem>(&StackModuleItem), 1, TWeakPtr<SWidget>(TargetWidget))));
+			MenuBuilder.AddMenuEntry(
+				LOCTEXT("InsertModuleBelow", "Insert Below"),
+				LOCTEXT("InsertModuleBelowToolTip", "Insert a new module below this module in the stack."),
+				FSlateIcon(),
+				FUIAction(FExecuteAction::CreateStatic(&ShowInsertModuleMenu, TWeakObjectPtr<UNiagaraStackModuleItem>(&StackModuleItem), 1, TWeakPtr<SWidget>(TargetWidget))));
+		}
 	}
 	MenuBuilder.EndSection();
 	return true;
