- Change PostLoad to ConditionalPostLoad

#Synced-CL 51672242

diff --git a/Engine/Plugins/FX/Niagara/Source/NiagaraEditor/Private/Customizations/SimCache/NiagaraRenderTargetVolumeSimCacheVisualizer.cpp b/Engine/Plugins/FX/Niagara/Source/NiagaraEditor/Private/Customizations/SimCache/NiagaraRenderTargetVolumeSimCacheVisualizer.cpp
--- a/Engine/Plugins/FX/Niagara/Source/NiagaraEditor/Private/Customizations/SimCache/NiagaraRenderTargetVolumeSimCacheVisualizer.cpp
+++ b/Engine/Plugins/FX/Niagara/Source/NiagaraEditor/Private/Customizations/SimCache/NiagaraRenderTargetVolumeSimCacheVisualizer.cpp
@@ -305,7 +305,7 @@ void SNiagaraVolumeTextureViewport::Construct(const FArguments&, TSharedPtr<FNia
 
 	PreviewComponent->OverrideMaterials.Add(DuplicateMat);
 	PreviewComponent->bIssueBlockingRequests = false;
-	PreviewComponent->PostLoad();
+	PreviewComponent->ConditionalPostLoad();
 
 	PostProcessComponent->bUnbound = true;
 	PostProcessComponent->bEnabled = true;
@@ -376,7 +376,7 @@ void SNiagaraVolumeTextureViewport::SetTemperatureMask(ENiagaraRenderTargetVolum
 				NewMask == ENiagaraRenderTargetVolumeVisualizerMask::A,
 				ExprGuid);
 			PreviewComponent->OverrideMaterials[0] = DuplicateMat;
-			PreviewComponent->PostLoad();
+			PreviewComponent->ConditionalPostLoad();
 			SetDensityScale(DensityScale);
 			PreviewComponent->MarkRenderStateDirty();
 			TemparatureMask = NewMask;
diff --git a/Engine/Plugins/FX/Niagara/Source/NiagaraEditor/Private/NiagaraBakerRenderer.cpp b/Engine/Plugins/FX/Niagara/Source/NiagaraEditor/Private/NiagaraBakerRenderer.cpp
--- a/Engine/Plugins/FX/Niagara/Source/NiagaraEditor/Private/NiagaraBakerRenderer.cpp
+++ b/Engine/Plugins/FX/Niagara/Source/NiagaraEditor/Private/NiagaraBakerRenderer.cpp
@@ -743,7 +743,7 @@ void FNiagaraBakerRenderer::RenderSparseVolumeTexture(UTextureRenderTarget2D* Re
 		SVTPreviewComponent->OverrideMaterials.Add(DuplicateMat);
 
 		SVTPreviewComponent->bIssueBlockingRequests = true;
-		SVTPreviewComponent->PostLoad();
+		SVTPreviewComponent->ConditionalPostLoad();
 
 		SVTPreviewScene = MakeShareable(new FAdvancedPreviewScene(FPreviewScene::ConstructionValues()));
 		SVTPreviewScene->SetFloorVisibility(false);
diff --git a/Engine/Plugins/FX/Niagara/Source/NiagaraEditor/Private/NiagaraBakerRendererOutputSparseVolumeTexture.cpp b/Engine/Plugins/FX/Niagara/Source/NiagaraEditor/Private/NiagaraBakerRendererOutputSparseVolumeTexture.cpp
--- a/Engine/Plugins/FX/Niagara/Source/NiagaraEditor/Private/NiagaraBakerRendererOutputSparseVolumeTexture.cpp
+++ b/Engine/Plugins/FX/Niagara/Source/NiagaraEditor/Private/NiagaraBakerRendererOutputSparseVolumeTexture.cpp
@@ -363,7 +363,7 @@ void FNiagaraBakerRendererOutputSparseVolumeTexture::EndBake(FNiagaraBakerFeedba
 		UE_LOG(LogNiagaraBaker, Error, TEXT("Cannot EndInitialize on creating SVT"));
 	}
 
-	SVTAsset->PostLoad();
+	SVTAsset->ConditionalPostLoad();
 
 	// make a second pass over the data and loop it
 	if (OutputVolumeTexture->bEnableLoopedOutput)
@@ -679,7 +679,7 @@ void FNiagaraBakerRendererOutputSparseVolumeTexture::EndBake(FNiagaraBakerFeedba
 			UE_LOG(LogNiagaraBaker, Error, TEXT("Cannot EndInitialize on creating looped SVT"));
 		}
 
-		LoopedSVTAsset->PostLoad();
+		LoopedSVTAsset->ConditionalPostLoad();
 	}
 }
 
