- Deprecate direct access to SourceDIName
- This is set in the instance after parameter store usage is resolved, users should only read, never write

#Synced-CL 51672242

#Synced-CL 51672242

diff --git a/Engine/Plugins/FX/Niagara/Source/Niagara/Classes/NiagaraDataInterface.h b/Engine/Plugins/FX/Niagara/Source/Niagara/Classes/NiagaraDataInterface.h
--- a/Engine/Plugins/FX/Niagara/Source/Niagara/Classes/NiagaraDataInterface.h
+++ b/Engine/Plugins/FX/Niagara/Source/Niagara/Classes/NiagaraDataInterface.h
@@ -405,9 +405,6 @@ struct FNiagaraDataInterfaceProxy
 	virtual int32 PerInstanceDataPassedToRenderThreadSize() const = 0;
 	virtual void ConsumePerInstanceDataFromGameThread(void* PerInstanceData, const FNiagaraSystemInstanceID& Instance) { check(false); }
 
-	// #todo(dmp): move all of this stuff to the RW interface to keep it out of here?
-	FName SourceDIName;
-
 	// New data interface path
 	virtual void ResetData(const FNDIGpuComputeResetContext& Context) { }
 	virtual void PreStage(const FNDIGpuComputePreStageContext& Context) {}
@@ -422,9 +419,25 @@ struct FNiagaraDataInterfaceProxy
 
 	virtual FNiagaraDataInterfaceProxyRW* AsIterationProxy() { return nullptr; }
 
+	inline FName GetSourceDIName() const
+	{
+	PRAGMA_DISABLE_DEPRECATION_WARNINGS
+		return IsInGameThread() || IsInParallelGameThread() ? SourceDIName : SourceDIName_RT;
+	PRAGMA_ENABLE_DEPRECATION_WARNINGS
+	}
+
+	void SetSourceDIName(FName InName);
+
 #if NIAGARA_VALIDATE_NDIPROXY_REFS
 	std::atomic<int32> ProxyTickRefs;
 #endif
+
+	// #todo(dmp): move all of this stuff to the RW interface to keep it out of here?
+	UE_DEPRECATED(5.8, "Please update your API to use the get API to read SourceDIName")
+	FName SourceDIName;
+
+private:
+	FName SourceDIName_RT;
 };
 
 //////////////////////////////////////////////////////////////////////////
diff --git a/Engine/Plugins/FX/Niagara/Source/Niagara/Private/NiagaraComputeExecutionContext.cpp b/Engine/Plugins/FX/Niagara/Source/Niagara/Private/NiagaraComputeExecutionContext.cpp
--- a/Engine/Plugins/FX/Niagara/Source/Niagara/Private/NiagaraComputeExecutionContext.cpp
+++ b/Engine/Plugins/FX/Niagara/Source/Niagara/Private/NiagaraComputeExecutionContext.cpp
@@ -78,27 +78,27 @@ void FNiagaraComputeExecutionContext::InitParams(UNiagaraScript* InGPUComputeScr
 
 bool FNiagaraComputeExecutionContext::IsOutputStage(FNiagaraDataInterfaceProxy* DIProxy, uint32 SimulationStageIndex) const
 {
-	if (DIProxy && !DIProxy->SourceDIName.IsNone())
+	if (DIProxy && !DIProxy->GetSourceDIName().IsNone())
 	{		
-		return SimStageExecData->SimStageMetaData[SimulationStageIndex].OutputDestinations.Contains(DIProxy->SourceDIName);
+		return SimStageExecData->SimStageMetaData[SimulationStageIndex].OutputDestinations.Contains(DIProxy->GetSourceDIName());
 	}
 	return false;
 }
 
 bool FNiagaraComputeExecutionContext::IsInputStage(FNiagaraDataInterfaceProxy* DIProxy, uint32 SimulationStageIndex) const
 {
-	if (DIProxy && !DIProxy->SourceDIName.IsNone())
+	if (DIProxy && !DIProxy->GetSourceDIName().IsNone())
 	{
-		return SimStageExecData->SimStageMetaData[SimulationStageIndex].InputDataInterfaces.Contains(DIProxy->SourceDIName);
+		return SimStageExecData->SimStageMetaData[SimulationStageIndex].InputDataInterfaces.Contains(DIProxy->GetSourceDIName());
 	}
 	return false;
 }
 
 bool FNiagaraComputeExecutionContext::IsIterationStage(FNiagaraDataInterfaceProxy* DIProxy, uint32 SimulationStageIndex) const
 {
-	if (DIProxy && !DIProxy->SourceDIName.IsNone())
+	if (DIProxy && !DIProxy->GetSourceDIName().IsNone())
 	{
-		return SimStageExecData->SimStageMetaData[SimulationStageIndex].IterationSourceType == ENiagaraIterationSource::DataInterface && (SimStageExecData->SimStageMetaData[SimulationStageIndex].IterationDataInterface == DIProxy->SourceDIName);
+		return SimStageExecData->SimStageMetaData[SimulationStageIndex].IterationSourceType == ENiagaraIterationSource::DataInterface && (SimStageExecData->SimStageMetaData[SimulationStageIndex].IterationDataInterface == DIProxy->GetSourceDIName());
 	}
 	return false;
 }
@@ -113,7 +113,7 @@ FNiagaraDataInterfaceProxyRW* FNiagaraComputeExecutionContext::FindIterationInte
 
 	for (FNiagaraDataInterfaceProxyRW* Proxy : InProxies)
 	{
-		if (Proxy->SourceDIName == SimStageExecData->SimStageMetaData[SimulationStageIndex].IterationDataInterface)
+		if (Proxy->GetSourceDIName() == SimStageExecData->SimStageMetaData[SimulationStageIndex].IterationDataInterface)
 		{
 			return Proxy;
 		}
diff --git a/Engine/Plugins/FX/Niagara/Source/Niagara/Private/NiagaraDataInterface.cpp b/Engine/Plugins/FX/Niagara/Source/Niagara/Private/NiagaraDataInterface.cpp
--- a/Engine/Plugins/FX/Niagara/Source/Niagara/Private/NiagaraDataInterface.cpp
+++ b/Engine/Plugins/FX/Niagara/Source/Niagara/Private/NiagaraDataInterface.cpp
@@ -20,6 +20,25 @@
 
 #define LOCTEXT_NAMESPACE "NiagaraDataInterface"
 
+void FNiagaraDataInterfaceProxy::SetSourceDIName(FName InName)
+{
+	PRAGMA_DISABLE_DEPRECATION_WARNINGS
+	if (SourceDIName == InName)
+	{
+		return;
+	}
+	SourceDIName = InName;
+
+	ENQUEUE_RENDER_COMMAND(SetProxySourceDIName)
+	(
+		[this, InName](FRHICommandListImmediate& RHICmdList) mutable
+		{
+			SourceDIName_RT = InName;
+		}
+	);
+	PRAGMA_ENABLE_DEPRECATION_WARNINGS
+}
+
 void FNDIStageTickHandler::Init(ENiagaraScriptUsage ScriptUsage, const FNiagaraScriptRuntimeData* ScriptRuntimeData, FNiagaraSystemInstance* Instance)
 {
 	check(ScriptRuntimeData);
diff --git a/Engine/Plugins/FX/Niagara/Source/Niagara/Private/NiagaraDataInterfaceGrid2DCollection.cpp b/Engine/Plugins/FX/Niagara/Source/Niagara/Private/NiagaraDataInterfaceGrid2DCollection.cpp
--- a/Engine/Plugins/FX/Niagara/Source/Niagara/Private/NiagaraDataInterfaceGrid2DCollection.cpp
+++ b/Engine/Plugins/FX/Niagara/Source/Niagara/Private/NiagaraDataInterfaceGrid2DCollection.cpp
@@ -1959,7 +1959,7 @@ bool UNiagaraDataInterfaceGrid2DCollection::InitPerInstanceData(void* PerInstanc
 	// give an error if there are any unnamed attributes
 	if (NumAttributes > 0 && NumNamedAttribChannelsFound > 0) //-V1051
 	{
-		UE_LOG(LogNiagara, Log, TEXT("%s: Unnamed attributes should not be used with named.  This is a deprecated workflow"), *Proxy->SourceDIName.ToString());
+		UE_LOG(LogNiagara, Log, TEXT("%s: Unnamed attributes should not be used with named.  This is a deprecated workflow"), *Proxy->GetSourceDIName().ToString());
 	}
 
 	InstanceData->WorldBBoxSize = WorldBBoxSize;
@@ -2053,7 +2053,7 @@ bool UNiagaraDataInterfaceGrid2DCollection::InitPerInstanceData(void* PerInstanc
 		check(!RT_Proxy->SystemInstancesToProxyData_RT.Contains(InstanceID));
 		FGrid2DCollectionRWInstanceData_RenderThread* TargetData = &RT_Proxy->SystemInstancesToProxyData_RT.Add(InstanceID);
 
-		TargetData->SourceDIName = RT_Proxy->SourceDIName;
+		TargetData->SourceDIName = RT_Proxy->GetSourceDIName();
 		TargetData->ClearBeforeNonIterationStage = RT_InstanceData.ClearBeforeNonIterationStage;
 
 		TargetData->NumCells = RT_InstanceData.NumCells;
@@ -2699,21 +2699,6 @@ void FGrid2DCollectionRWInstanceData_RenderThread::BeginSimulate(FRDGBuilder& Gr
 
 		const FRDGTextureDesc TextureDesc = FRDGTextureDesc::Create2DArray(NumCells, PixelFormat.GetValue(), FClearValueBinding::Black, ETextureCreateFlags::ShaderResource | ETextureCreateFlags::UAV, uint16(NumAttributes));
 
-	#if 0
-		// Note: Invalid as the render target pool holds onto a pointer and can exist beyond the lifetime of the grid
-		TStringBuilder<128> StringBuilder;
-		SourceDIName.AppendString(StringBuilder);
-		StringBuilder.Append("_");
-		for (FName CurrName : Vars)
-		{
-			CurrName.AppendString(StringBuilder);
-			StringBuilder.Append("_");
-		}
-		StringBuilder.Append("Grid2DCollection");		
-		const TCHAR *GridTextureName = StringBuilder.ToString();
-		UE_LOG(LogNiagara, Warning, TEXT("%s"), *StringBuilder);		
-	#endif
-
 		DestinationData->Initialize(GraphBuilder, TEXT("Grid2D::GridTexture"), TextureDesc);
 
 		// This destination buffer will sometimes have old data in it.  Force it to clear.
@@ -2798,11 +2783,11 @@ void FNiagaraDataInterfaceProxyGrid2DCollectionProxy::PostSimulate(const FNDIGpu
 		FRDGBuilder& GraphBuilder = Context.GetGraphBuilder();
 		if (ProxyData->PreviewAttribute[0] != INDEX_NONE)
 		{
-			GpuComputeDebugInterface.AddAttributeTexture(GraphBuilder, Context.GetSystemInstanceID(), SourceDIName, ProxyData->CurrentData->GetOrCreateTexture(GraphBuilder), FIntPoint::ZeroValue, ProxyData->PreviewAttribute);
+			GpuComputeDebugInterface.AddAttributeTexture(GraphBuilder, Context.GetSystemInstanceID(), GetSourceDIName(), ProxyData->CurrentData->GetOrCreateTexture(GraphBuilder), FIntPoint::ZeroValue, ProxyData->PreviewAttribute);
 		}
 		else
 		{
-			GpuComputeDebugInterface.AddTexture(GraphBuilder, Context.GetSystemInstanceID(), SourceDIName, ProxyData->CurrentData->GetOrCreateTexture(GraphBuilder));
+			GpuComputeDebugInterface.AddTexture(GraphBuilder, Context.GetSystemInstanceID(), GetSourceDIName(), ProxyData->CurrentData->GetOrCreateTexture(GraphBuilder));
 		}
 	}
 #endif
diff --git a/Engine/Plugins/FX/Niagara/Source/Niagara/Private/NiagaraDataInterfaceGrid3DCollection.cpp b/Engine/Plugins/FX/Niagara/Source/Niagara/Private/NiagaraDataInterfaceGrid3DCollection.cpp
--- a/Engine/Plugins/FX/Niagara/Source/Niagara/Private/NiagaraDataInterfaceGrid3DCollection.cpp
+++ b/Engine/Plugins/FX/Niagara/Source/Niagara/Private/NiagaraDataInterfaceGrid3DCollection.cpp
@@ -3147,7 +3147,7 @@ bool UNiagaraDataInterfaceGrid3DCollection::InitPerInstanceData(void* PerInstanc
 	// give an error if there are any unnamed attributes
 	if (NumAttributes > 0 && NumNamedAttribChannelsFound > 0)
 	{
-		UE_LOG(LogNiagara, Warning, TEXT("%s: Unnamed attributes should not be used with named.  This is a deprecated workflow"), *Proxy->SourceDIName.ToString());
+		UE_LOG(LogNiagara, Warning, TEXT("%s: Unnamed attributes should not be used with named.  This is a deprecated workflow"), *Proxy->GetSourceDIName().ToString());
 	}
 
 	if (SetResolutionMethod == ESetResolutionMethod::Independent)
@@ -3220,7 +3220,7 @@ bool UNiagaraDataInterfaceGrid3DCollection::InitPerInstanceData(void* PerInstanc
 		InstanceData->NumCells.Y <= 0 ||
 		InstanceData->NumCells.Z <= 0)
 	{
-		UE_LOG(LogNiagara, Error, TEXT("Zero grid resolution defined on %s"), *Proxy->SourceDIName.ToString());
+		UE_LOG(LogNiagara, Error, TEXT("Zero grid resolution defined on %s"), *Proxy->GetSourceDIName().ToString());
 		return false;
 	}
 
@@ -3233,9 +3233,9 @@ bool UNiagaraDataInterfaceGrid3DCollection::InitPerInstanceData(void* PerInstanc
 
 	if (NumAttribChannelsFound == 0)
 	{
-		if (!Proxy->SourceDIName.IsNone())
+		if (!Proxy->GetSourceDIName().IsNone())
 		{
-			UE_LOG(LogNiagara, Log, TEXT("Zero attributes defined on %s"), *Proxy->SourceDIName.ToString());
+			UE_LOG(LogNiagara, Log, TEXT("Zero attributes defined on %s"), *Proxy->GetSourceDIName().ToString());
 		}
 
 		// Push Updates to Proxy.
@@ -3266,7 +3266,7 @@ bool UNiagaraDataInterfaceGrid3DCollection::InitPerInstanceData(void* PerInstanc
 
 		if (NumAttribChannelsFound - NumNamedAttribChannelsFound > 0)
 		{
-			UE_LOG(LogNiagara, Error, TEXT("%s: Unnamed attributes do not work with rgba textures"), *Proxy->SourceDIName.ToString());
+			UE_LOG(LogNiagara, Error, TEXT("%s: Unnamed attributes do not work with rgba textures"), *Proxy->GetSourceDIName().ToString());
 			return false;
 		}
 
@@ -3291,13 +3291,13 @@ bool UNiagaraDataInterfaceGrid3DCollection::InitPerInstanceData(void* PerInstanc
 
 		if (InstanceData->NumCells.X > MaxDim || InstanceData->NumCells.Y > MaxDim || InstanceData->NumCells.Z > MaxDim)
 		{
-			UE_LOG(LogNiagara, Error, TEXT("Resolution is too high on %s... max is 2048, defined as %d, %d %d"), *Proxy->SourceDIName.ToString(), InstanceData->NumCells.X, InstanceData->NumCells.Y, InstanceData->NumCells.Z);
+			UE_LOG(LogNiagara, Error, TEXT("Resolution is too high on %s... max is 2048, defined as %d, %d %d"), *Proxy->GetSourceDIName().ToString(), InstanceData->NumCells.X, InstanceData->NumCells.Y, InstanceData->NumCells.Z);
 			return false;
 		}
 
 		if (NumAttribChannelsFound > MaxAttributes && MaxAttributes > 0)
 		{
-			UE_LOG(LogNiagara, Error, TEXT("Invalid number of attributes defined on %s... max is %d, num defined is %d"), *Proxy->SourceDIName.ToString(), MaxAttributes, NumAttribChannelsFound);
+			UE_LOG(LogNiagara, Error, TEXT("Invalid number of attributes defined on %s... max is %d, num defined is %d"), *Proxy->GetSourceDIName().ToString(), MaxAttributes, NumAttribChannelsFound);
 			return false;
 		}
 
@@ -3367,7 +3367,7 @@ bool UNiagaraDataInterfaceGrid3DCollection::InitPerInstanceData(void* PerInstanc
 			check(!RT_Proxy->SystemInstancesToProxyData_RT.Contains(InstanceID));
 			FGrid3DCollectionRWInstanceData_RenderThread* TargetData = &RT_Proxy->SystemInstancesToProxyData_RT.Add(InstanceID);
 
-			TargetData->SourceDIName = RT_Proxy->SourceDIName;
+			TargetData->SourceDIName = RT_Proxy->GetSourceDIName();
 
 			TargetData->ClearBeforeNonIterationStage = RT_InstanceData.ClearBeforeNonIterationStage;
 
@@ -3778,19 +3778,19 @@ bool UNiagaraDataInterfaceGrid3DCollection::PerInstanceTickPostSimulate(void* Pe
 
 			if (InstanceData->NumCells.X > MaxDim || InstanceData->NumCells.Y > MaxDim || InstanceData->NumCells.Z > MaxDim)
 			{
-				UE_LOG(LogNiagara, Error, TEXT("Resolution is too high on %s... max is 2048, defined as %d, %d %d"), *Proxy->SourceDIName.ToString(), InstanceData->NumCells.X, InstanceData->NumCells.Y, InstanceData->NumCells.Z);
+				UE_LOG(LogNiagara, Error, TEXT("Resolution is too high on %s... max is 2048, defined as %d, %d %d"), *Proxy->GetSourceDIName().ToString(), InstanceData->NumCells.X, InstanceData->NumCells.Y, InstanceData->NumCells.Z);
 				return false;
 			}
 
 			if ((InstanceData->TotalNumAttributes > MaxAttributes && MaxAttributes > 0))
 			{
-				UE_LOG(LogNiagara, Error, TEXT("Invalid number of attributes defined on %s... max is %d, num defined is %d"), *Proxy->SourceDIName.ToString(), MaxAttributes, InstanceData->TotalNumAttributes);
+				UE_LOG(LogNiagara, Error, TEXT("Invalid number of attributes defined on %s... max is %d, num defined is %d"), *Proxy->GetSourceDIName().ToString(), MaxAttributes, InstanceData->TotalNumAttributes);
 				return false;
 			}
 
 			if (InstanceData->TotalNumAttributes <= 0)
 			{
-				UE_LOG(LogNiagara, Warning, TEXT("Invalid number of attributes defined on %s... max is %d, num defined is %d"), *Proxy->SourceDIName.ToString(), MaxAttributes, InstanceData->TotalNumAttributes);
+				UE_LOG(LogNiagara, Warning, TEXT("Invalid number of attributes defined on %s... max is %d, num defined is %d"), *Proxy->GetSourceDIName().ToString(), MaxAttributes, InstanceData->TotalNumAttributes);
 				return false;
 			}
 
@@ -3941,19 +3941,6 @@ void FGrid3DCollectionRWInstanceData_RenderThread::BeginSimulate(FRDGBuilder& Gr
 		const FIntVector TextureSize(NumCells.X * NumTiles.X, NumCells.Y * NumTiles.Y, NumCells.Z * NumTiles.Z);
 		const FRDGTextureDesc TextureDesc = FRDGTextureDesc::Create3D(TextureSize, PixelFormat.GetValue(), FClearValueBinding::Black, ETextureCreateFlags::ShaderResource | ETextureCreateFlags::UAV);
 
-	#if 0
-		// Note: Invalid as the render target pool holds onto a pointer and can exist beyond the lifetime of the grid
-		TStringBuilder<128> StringBuilder;
-		SourceDIName.AppendString(StringBuilder);
-		StringBuilder.Append("_");
-		for (FName CurrName : Vars)
-		{
-			CurrName.AppendString(StringBuilder);
-			StringBuilder.Append("_");
-		}
-		StringBuilder.Append("Grid3DCollection");
-		GridTextureName = *StringBuilder;
-	#endif
 		DestinationData->Initialize(GraphBuilder, TEXT("Grid3D::GridTexture"), TextureDesc);
 
 		// This destination buffer will sometimes have old data in it.  Force it to clear.
@@ -4031,11 +4018,11 @@ void FNiagaraDataInterfaceProxyGrid3DCollectionProxy::PostSimulate(const FNDIGpu
 		if (ProxyData->PreviewAttribute[0] != INDEX_NONE)
 		{
 			const FIntVector4 TotalNumAttributeVector = FIntVector4(ProxyData->NumTiles.X, ProxyData->NumTiles.Y, ProxyData->NumTiles.Z, 0);
-			GpuComputeDebugInterface.AddAttributeTexture(GraphBuilder, Context.GetSystemInstanceID(), SourceDIName, ProxyData->CurrentData->GetOrCreateTexture(GraphBuilder), TotalNumAttributeVector, ProxyData->PreviewAttribute);
+			GpuComputeDebugInterface.AddAttributeTexture(GraphBuilder, Context.GetSystemInstanceID(), GetSourceDIName(), ProxyData->CurrentData->GetOrCreateTexture(GraphBuilder), TotalNumAttributeVector, ProxyData->PreviewAttribute);
 		}
 		else
 		{
-			GpuComputeDebugInterface.AddTexture(GraphBuilder, Context.GetSystemInstanceID(), SourceDIName, ProxyData->CurrentData->GetOrCreateTexture(GraphBuilder));
+			GpuComputeDebugInterface.AddTexture(GraphBuilder, Context.GetSystemInstanceID(), GetSourceDIName(), ProxyData->CurrentData->GetOrCreateTexture(GraphBuilder));
 		}
 	}
 #endif
diff --git a/Engine/Plugins/FX/Niagara/Source/Niagara/Private/NiagaraDataInterfaceGrid3DCollectionReader.cpp b/Engine/Plugins/FX/Niagara/Source/Niagara/Private/NiagaraDataInterfaceGrid3DCollectionReader.cpp
--- a/Engine/Plugins/FX/Niagara/Source/Niagara/Private/NiagaraDataInterfaceGrid3DCollectionReader.cpp
+++ b/Engine/Plugins/FX/Niagara/Source/Niagara/Private/NiagaraDataInterfaceGrid3DCollectionReader.cpp
@@ -136,12 +136,12 @@ bool UNiagaraDataInterfaceGrid3DCollectionReader::InitPerInstanceData(void* PerI
 		}
 		else
 		{
-			UE_LOG(LogNiagara, Error, TEXT("Emitter %s doesn't have Grid3D Collection DI named %s for DI: %s"), *this->EmitterName, *this->DIName, *Proxy->SourceDIName.ToString());
+			UE_LOG(LogNiagara, Error, TEXT("Emitter %s doesn't have Grid3D Collection DI named %s for DI: %s"), *this->EmitterName, *this->DIName, *Proxy->GetSourceDIName().ToString());
 		}
 	}
 	else
 	{
-		UE_LOG(LogNiagara, Error, TEXT("Emitter %s doesn't exist for DI: %s"), *this->EmitterName, *Proxy->SourceDIName.ToString());
+		UE_LOG(LogNiagara, Error, TEXT("Emitter %s doesn't exist for DI: %s"), *this->EmitterName, *Proxy->GetSourceDIName().ToString());
 	}
 
 	return false;
@@ -243,7 +243,7 @@ bool UNiagaraDataInterfaceGrid3DCollectionReader::PerInstanceTickPostSimulate(vo
 
 		if (InstanceData->OtherInstanceData == nullptr)
 		{
-			UE_LOG(LogNiagara, Error, TEXT("%s: The DI for Emitter %s and DataInterface %s is not initialized."), *Proxy->SourceDIName.ToString(), *this->EmitterName, *this->DIName);
+			UE_LOG(LogNiagara, Error, TEXT("%s: The DI for Emitter %s and DataInterface %s is not initialized."), *Proxy->GetSourceDIName().ToString(), *this->EmitterName, *this->DIName);
 			return false;
 		}
 
diff --git a/Engine/Plugins/FX/Niagara/Source/Niagara/Private/NiagaraDataInterfaceIntRenderTarget2D.cpp b/Engine/Plugins/FX/Niagara/Source/Niagara/Private/NiagaraDataInterfaceIntRenderTarget2D.cpp
--- a/Engine/Plugins/FX/Niagara/Source/Niagara/Private/NiagaraDataInterfaceIntRenderTarget2D.cpp
+++ b/Engine/Plugins/FX/Niagara/Source/Niagara/Private/NiagaraDataInterfaceIntRenderTarget2D.cpp
@@ -139,7 +139,7 @@ struct FNDIIntRenderTarget2DProxy : public FNiagaraDataInterfaceProxyRW
 		if (InstanceData.bPreviewRenderTarget && InstanceData.TransientRDGTexture)
 		{
 			FNiagaraGpuComputeDebugInterface GpuComputeDebugInterface = Context.GetComputeDispatchInterface().GetGpuComputeDebugInterface();
-			GpuComputeDebugInterface.AddTexture(Context.GetGraphBuilder(), Context.GetSystemInstanceID(), SourceDIName, InstanceData.TransientRDGTexture, InstanceData.PreviewDisplayRange);
+			GpuComputeDebugInterface.AddTexture(Context.GetGraphBuilder(), Context.GetSystemInstanceID(), GetSourceDIName(), InstanceData.TransientRDGTexture, InstanceData.PreviewDisplayRange);
 		}
 #endif
 		if (Context.IsFinalPostSimulate())
diff --git a/Engine/Plugins/FX/Niagara/Source/Niagara/Private/NiagaraDataInterfaceRasterizationGrid3D.cpp b/Engine/Plugins/FX/Niagara/Source/Niagara/Private/NiagaraDataInterfaceRasterizationGrid3D.cpp
--- a/Engine/Plugins/FX/Niagara/Source/Niagara/Private/NiagaraDataInterfaceRasterizationGrid3D.cpp
+++ b/Engine/Plugins/FX/Niagara/Source/Niagara/Private/NiagaraDataInterfaceRasterizationGrid3D.cpp
@@ -674,7 +674,7 @@ bool UNiagaraDataInterfaceRasterizationGrid3D::GetFunctionHLSL(const FNiagaraDat
 						val = 0;
 
 						int NewValue = IntValue > 0 ? 2147483647 : -2147483648;
-						{OutputIntGrid}[Index] = In_Value;
+						{OutputIntGrid}[Index] = NewValue;
 					}
 				}
 			}
@@ -891,9 +891,9 @@ bool UNiagaraDataInterfaceRasterizationGrid3D::InitPerInstanceData(void* PerInst
 	InstanceData->NumTiles.Y = NumTilesY;
 	InstanceData->NumTiles.Z = NumTilesZ;
 
-	if (NumAttribChannelsFound == 0 && !Proxy->SourceDIName.IsNone())
+	if (NumAttribChannelsFound == 0 && !Proxy->GetSourceDIName().IsNone())
 	{
-		UE_LOG(LogNiagara, Log, TEXT("Zero attributes defined on %s"), *Proxy->SourceDIName.ToString());
+		UE_LOG(LogNiagara, Log, TEXT("Zero attributes defined on %s"), *Proxy->GetSourceDIName().ToString());
 	}
 
 //	check(InstanceData->NumTiles.X > 0);
@@ -1021,7 +1021,7 @@ bool UNiagaraDataInterfaceRasterizationGrid3D::PerInstanceTickPostSimulate(void*
 
 		if (InstanceData->TotalNumAttributes <= 0)
 		{
-			UE_LOG(LogNiagara, Warning, TEXT("Invalid number of attributes defined on %s... max is %d, num defined is %d"), *Proxy->SourceDIName.ToString(), MaxAttributes, InstanceData->TotalNumAttributes);
+			UE_LOG(LogNiagara, Warning, TEXT("Invalid number of attributes defined on %s... max is %d, num defined is %d"), *Proxy->GetSourceDIName().ToString(), MaxAttributes, InstanceData->TotalNumAttributes);
 			return false;
 		}
 
diff --git a/Engine/Plugins/FX/Niagara/Source/Niagara/Private/NiagaraDataInterfaceRenderTarget2D.cpp b/Engine/Plugins/FX/Niagara/Source/Niagara/Private/NiagaraDataInterfaceRenderTarget2D.cpp
--- a/Engine/Plugins/FX/Niagara/Source/Niagara/Private/NiagaraDataInterfaceRenderTarget2D.cpp
+++ b/Engine/Plugins/FX/Niagara/Source/Niagara/Private/NiagaraDataInterfaceRenderTarget2D.cpp
@@ -1049,7 +1049,7 @@ void FNiagaraDataInterfaceProxyRenderTarget2DProxy::PostSimulate(const FNDIGpuCo
 	if (ProxyData->bPreviewTexture && ProxyData->TransientRDGTexture)
 	{
 		FNiagaraGpuComputeDebugInterface GpuComputeDebugInterface = Context.GetComputeDispatchInterface().GetGpuComputeDebugInterface();
-		GpuComputeDebugInterface.AddTexture(Context.GetGraphBuilder(), Context.GetSystemInstanceID(), SourceDIName, ProxyData->TransientRDGTexture);
+		GpuComputeDebugInterface.AddTexture(Context.GetGraphBuilder(), Context.GetSystemInstanceID(), GetSourceDIName(), ProxyData->TransientRDGTexture);
 	}
 #endif
 
diff --git a/Engine/Plugins/FX/Niagara/Source/Niagara/Private/NiagaraDataInterfaceRenderTarget2DArray.cpp b/Engine/Plugins/FX/Niagara/Source/Niagara/Private/NiagaraDataInterfaceRenderTarget2DArray.cpp
--- a/Engine/Plugins/FX/Niagara/Source/Niagara/Private/NiagaraDataInterfaceRenderTarget2DArray.cpp
+++ b/Engine/Plugins/FX/Niagara/Source/Niagara/Private/NiagaraDataInterfaceRenderTarget2DArray.cpp
@@ -762,7 +762,7 @@ void FNiagaraDataInterfaceProxyRenderTarget2DArrayProxy::PostSimulate(const FNDI
 	if (ProxyData->bPreviewTexture && ProxyData->TransientRDGTexture)
 	{
 		FNiagaraGpuComputeDebugInterface GpuComputeDebugInterface = Context.GetComputeDispatchInterface().GetGpuComputeDebugInterface();
-		GpuComputeDebugInterface.AddTexture(Context.GetGraphBuilder(), Context.GetSystemInstanceID(), SourceDIName, ProxyData->TransientRDGTexture);
+		GpuComputeDebugInterface.AddTexture(Context.GetGraphBuilder(), Context.GetSystemInstanceID(), GetSourceDIName(), ProxyData->TransientRDGTexture);
 	}
 #endif
 
diff --git a/Engine/Plugins/FX/Niagara/Source/Niagara/Private/NiagaraDataInterfaceRenderTargetCube.cpp b/Engine/Plugins/FX/Niagara/Source/Niagara/Private/NiagaraDataInterfaceRenderTargetCube.cpp
--- a/Engine/Plugins/FX/Niagara/Source/Niagara/Private/NiagaraDataInterfaceRenderTargetCube.cpp
+++ b/Engine/Plugins/FX/Niagara/Source/Niagara/Private/NiagaraDataInterfaceRenderTargetCube.cpp
@@ -715,7 +715,7 @@ void FNiagaraDataInterfaceProxyRenderTargetCubeProxy::PostSimulate(const FNDIGpu
 	if (ProxyData->bPreviewTexture && ProxyData->TransientRDGTexture)
 	{
 		FNiagaraGpuComputeDebugInterface GpuComputeDebugInterface = Context.GetComputeDispatchInterface().GetGpuComputeDebugInterface();
-		GpuComputeDebugInterface.AddTexture(Context.GetGraphBuilder(), Context.GetSystemInstanceID(), SourceDIName, ProxyData->TransientRDGTexture);
+		GpuComputeDebugInterface.AddTexture(Context.GetGraphBuilder(), Context.GetSystemInstanceID(), GetSourceDIName(), ProxyData->TransientRDGTexture);
 	}
 #endif
 
diff --git a/Engine/Plugins/FX/Niagara/Source/Niagara/Private/NiagaraDataInterfaceRenderTargetVolume.cpp b/Engine/Plugins/FX/Niagara/Source/Niagara/Private/NiagaraDataInterfaceRenderTargetVolume.cpp
--- a/Engine/Plugins/FX/Niagara/Source/Niagara/Private/NiagaraDataInterfaceRenderTargetVolume.cpp
+++ b/Engine/Plugins/FX/Niagara/Source/Niagara/Private/NiagaraDataInterfaceRenderTargetVolume.cpp
@@ -720,7 +720,7 @@ bool UNiagaraDataInterfaceRenderTargetVolume::SimCacheEndWrite(UObject* StorageO
 	{	
 		UAnimatedSparseVolumeTexture* CurrCache = CastChecked<UAnimatedSparseVolumeTexture>(StorageObject);
 		CurrCache->EndInitialize();
-		CurrCache->PostLoad();
+		CurrCache->ConditionalPostLoad();
 	}
 
 	return true;
@@ -1135,7 +1135,7 @@ void FNiagaraDataInterfaceProxyRenderTargetVolumeProxy::PostSimulate(const FNDIG
 	if (ProxyData->bPreviewTexture && ProxyData->TransientRDGTexture)
 	{
 		FNiagaraGpuComputeDebugInterface GpuComputeDebugInterface = Context.GetComputeDispatchInterface().GetGpuComputeDebugInterface();
-		GpuComputeDebugInterface.AddTexture(Context.GetGraphBuilder(), Context.GetSystemInstanceID(), SourceDIName, ProxyData->TransientRDGTexture);
+		GpuComputeDebugInterface.AddTexture(Context.GetGraphBuilder(), Context.GetSystemInstanceID(), GetSourceDIName(), ProxyData->TransientRDGTexture);
 	}
 #endif
 
diff --git a/Engine/Plugins/FX/Niagara/Source/Niagara/Private/NiagaraGpuComputeDispatch.cpp b/Engine/Plugins/FX/Niagara/Source/Niagara/Private/NiagaraGpuComputeDispatch.cpp
--- a/Engine/Plugins/FX/Niagara/Source/Niagara/Private/NiagaraGpuComputeDispatch.cpp
+++ b/Engine/Plugins/FX/Niagara/Source/Niagara/Private/NiagaraGpuComputeDispatch.cpp
@@ -605,7 +605,7 @@ void FNiagaraGpuComputeDispatch::DumpDebugFrame()
 					for (const auto& Pair : Tick->DIInstanceData->InterfaceProxiesToOffsets)
 					{
 						FNiagaraDataInterfaceProxy* Proxy = Pair.Key;
-						UE_LOG(LogNiagara, Warning, TEXT("Proxy(%s)"), *Proxy->SourceDIName.ToString());
+						UE_LOG(LogNiagara, Warning, TEXT("Proxy(%s)"), *Proxy->GetSourceDIName().ToString());
 					}
 				}
 			}
@@ -657,7 +657,7 @@ void FNiagaraGpuComputeDispatch::DumpDebugFrame()
 				}
 				Builder.Appendf(TEXT("Source(%p 0x%08x %d) "), SimStageData.Source, SimStageData.SourceCountOffset, SimStageData.SourceNumInstances);
 				Builder.Appendf(TEXT("Destination(%p 0x%08x %d) "), SimStageData.Destination, SimStageData.DestinationCountOffset, SimStageData.DestinationNumInstances);
-				Builder.Appendf(TEXT("Iteration(%d | %s) "), SimStageData.IterationIndex, SimStageData.AlternateIterationSource ? *SimStageData.AlternateIterationSource->SourceDIName.ToString() : TEXT("Particles"));
+				Builder.Appendf(TEXT("Iteration(%d | %s) "), SimStageData.IterationIndex, SimStageData.AlternateIterationSource ? *SimStageData.AlternateIterationSource->GetSourceDIName().ToString() : TEXT("Particles"));
 				Builder.Appendf(TEXT("DispatchElementCount(%d, %d, %d) "), SimStageData.DispatchArgs.ElementCount.X, SimStageData.DispatchArgs.ElementCount.Y, SimStageData.DispatchArgs.ElementCount.Z);
 				UE_LOG(LogNiagara, Warning, TEXT("%s"), Builder.ToString());
 			}
@@ -1538,7 +1538,7 @@ void FNiagaraGpuComputeDispatch::DispatchStage(FRDGBuilder& GraphBuilder, const
 				DispatchCount = SimStageData.DispatchArgs.ElementCount;
 			
 				// Verify the number of elements isn't higher that what we can handle
-				checkf(uint64(DispatchCount.X) * uint64(DispatchCount.Y) * uint64(DispatchCount.Z) < uint64(TNumericLimits<int32>::Max()), TEXT("DispatchCount(%d, %d, %d) for IterationInterface(%s) overflows an int32 this is not allowed"), DispatchCount.X, DispatchCount.Y, DispatchCount.Z, *SimStageData.AlternateIterationSource->SourceDIName.ToString());
+				checkf(uint64(DispatchCount.X) * uint64(DispatchCount.Y) * uint64(DispatchCount.Z) < uint64(TNumericLimits<int32>::Max()), TEXT("DispatchCount(%d, %d, %d) for IterationInterface(%s) overflows an int32 this is not allowed"), DispatchCount.X, DispatchCount.Y, DispatchCount.Z, *SimStageData.AlternateIterationSource->GetSourceDIName().ToString());
 
 				// Data interfaces such as grids / render targets can choose to dispatch in either the correct dimensionality for the target (i.e. RT2D would choose 2D)
 				// or run in linear mode if performance is not beneficial due to increased waves.  It is also possible the we may choose to override on the simulation stage.
diff --git a/Engine/Plugins/FX/Niagara/Source/Niagara/Private/NiagaraSystemInstance.cpp b/Engine/Plugins/FX/Niagara/Source/Niagara/Private/NiagaraSystemInstance.cpp
--- a/Engine/Plugins/FX/Niagara/Source/Niagara/Private/NiagaraSystemInstance.cpp
+++ b/Engine/Plugins/FX/Niagara/Source/Niagara/Private/NiagaraSystemInstance.cpp
@@ -1652,8 +1652,7 @@ void FNiagaraSystemInstance::InitDataInterfaces()
 						Interface->SetUsedWithGPUScript(true);
 						if(FNiagaraDataInterfaceProxy* Proxy = Interface->GetProxy())
 						{
-							// We need to store the name of each DI source variable here so that we can look it up later when looking for the iteration interface.
-							Proxy->SourceDIName = Var.GetName();
+							Proxy->SetSourceDIName(Var.GetName());
 						}
 					}
 					else if (bIsCPUScript)
