- Remove parent emitter post loading
- As we don't do any merging until UpdateSystemAfterLoad, there is no value to conditionally post loading them, and it may be causing a bug we are finding hard to track down
- Can be disabled with fx.Niagara.Rollback.NiagaraEmitterPostLoadCleanup

#Synced-CL 51711051

diff --git a/Engine/Plugins/FX/Niagara/Source/Niagara/Private/NiagaraEmitter.cpp b/Engine/Plugins/FX/Niagara/Source/Niagara/Private/NiagaraEmitter.cpp
--- a/Engine/Plugins/FX/Niagara/Source/Niagara/Private/NiagaraEmitter.cpp
+++ b/Engine/Plugins/FX/Niagara/Source/Niagara/Private/NiagaraEmitter.cpp
@@ -47,6 +47,14 @@ const FString InitialNotSynchronizedReason("Emitter created");
 const FGuid UNiagaraEmitter::EmitterMergeMessageId(0xDAF26E73, 0x4D1B416B, 0x815FA6C2, 0x6D5D0A75);
 #endif
 
+static bool GNiagaraEmitterPostLoadCleanup = true;
+static FAutoConsoleVariableRef CVarRollback_NiagaraEmitterPostLoadCleanup(
+	TEXT("fx.Niagara.Rollback.NiagaraEmitterPostLoadCleanup"),
+	GNiagaraEmitterPostLoadCleanup,
+	TEXT("When true use a leaner version of post load"),
+	ECVF_Default
+);
+
 static int32 GbForceNiagaraCompileOnLoad = 0;
 static FAutoConsoleVariableRef CVarForceNiagaraCompileOnLoad(
 	TEXT("fx.ForceCompileOnLoad"),
@@ -809,21 +817,23 @@ void FVersionedNiagaraEmitterData::PostLoad(UNiagaraEmitter& Emitter, int32 Niag
 			GraphSource->PostLoadFromEmitter(FVersionedNiagaraEmitter(&Emitter, Version.VersionGuid));
 		}
 
-		// Prepare for emitter inheritance.
-		if (GetParent().Emitter != nullptr)
+		if (GNiagaraEmitterPostLoadCleanup == false)
 		{
-			GetParent().Emitter->ConditionalPostLoad();
-		}
-		if (GetParentAtLastMerge().Emitter != nullptr)
-		{
-			GetParentAtLastMerge().Emitter->ConditionalPostLoad();
-		}
+			if (GetParent().Emitter != nullptr)
+			{
+				GetParent().Emitter->ConditionalPostLoad();
+			}
+			if (GetParentAtLastMerge().Emitter != nullptr)
+			{
+				GetParentAtLastMerge().Emitter->ConditionalPostLoad();
+			}
 
-		for (auto ParentScratchPadScript : ParentScratchPads->Scripts)
-		{
-			if (ParentScratchPadScript)
+			for (auto ParentScratchPadScript : ParentScratchPads->Scripts)
 			{
-				ParentScratchPadScript->ConditionalPostLoad();
+				if (ParentScratchPadScript)
+				{
+					ParentScratchPadScript->ConditionalPostLoad();
+				}
 			}
 		}
 
@@ -835,7 +845,7 @@ void FVersionedNiagaraEmitterData::PostLoad(UNiagaraEmitter& Emitter, int32 Niag
 			}
 		}
 
-		if (Emitter.IsSynchronizedWithParent() == false && IsRunningCommandlet())
+		if (IsRunningCommandlet() && Emitter.IsSynchronizedWithParent() == false)
 		{
 			// Modify here so that the asset will be marked dirty when using the resave commandlet.  This will be ignored during regular post load.
 			Emitter.Modify();
