- Add SourceMode to the spline data interface

#Synced-CL 51711051

diff --git a/Engine/Plugins/FX/Niagara/Source/Niagara/Classes/NiagaraDataInterfaceSpline.h b/Engine/Plugins/FX/Niagara/Source/Niagara/Classes/NiagaraDataInterfaceSpline.h
--- a/Engine/Plugins/FX/Niagara/Source/Niagara/Classes/NiagaraDataInterfaceSpline.h
+++ b/Engine/Plugins/FX/Niagara/Source/Niagara/Classes/NiagaraDataInterfaceSpline.h
@@ -11,6 +11,24 @@
 
 class UNiagaraDataInterfaceSpline;
 
+UENUM()
+enum class ENDISpline_SourceMode : uint8
+{
+	// Default will use all possible sources in the following order
+	// SplineUserParameter - if change will reset the system
+	// SoftSourceActor - looks for the first spline component on the root actor
+	// AttachParent - looks for the first spline component on the attached parent
+	Default,
+
+	// AttachParent - looks for the first spline component on the attached parent
+	// Will check every frame to see if it has changed
+	AttachParentOnly,
+
+	// SplineUserParameter - if change will reset the system
+	// Will check every frame to see if it has changed
+	ParameterBindingOnly,
+};
+
 /** Proxy data for splines */
 struct FNDISpline_InstanceData_RenderThread
 {
@@ -18,7 +36,7 @@ struct FNDISpline_InstanceData_RenderThread
 	{
 		Reset();
 	}
-	
+
 	FMatrix44f SplineTransform;
 	FMatrix44f SplineTransformRotationMat;
 	FMatrix44f SplineTransformInverse;
@@ -31,7 +49,7 @@ struct FNDISpline_InstanceData_RenderThread
 	float SplineDistanceStep;
 	float InvSplineDistanceStep;
 	int32 MaxIndex;
-	
+
 	FReadBuffer SplinePositionsLUT;
 	FReadBuffer SplineScalesLUT;
 	FReadBuffer SplineRotationsLUT;
@@ -41,7 +59,7 @@ struct FNDISpline_InstanceData_RenderThread
 
 
 struct FNiagaraDataInterfaceProxySpline : public FNiagaraDataInterfaceProxy
-{	
+{
 	virtual void ConsumePerInstanceDataFromGameThread(void* PerInstanceData, const FNiagaraSystemInstanceID& Instance) override {}
 	virtual int32 PerInstanceDataPassedToRenderThreadSize() const override { return 0; }
 
@@ -55,7 +73,7 @@ USTRUCT()
 struct FNiagaraDataInterfaceSplineLUT
 {
 	GENERATED_BODY()
-public:	
+public:
 	UPROPERTY()
 	TArray<FVector> Positions;
 	UPROPERTY()
@@ -79,50 +97,41 @@ public:
 
 struct FNDISpline_InstanceData
 {
-	FNDISpline_InstanceData()
-		: CachedUserParam(nullptr)
-		, bSyncedGPUCopy(false)
-	{}
 	//Cached ptr to component we sample from. 
 	TWeakObjectPtr<USplineComponent> Component;
 
-	UObject* CachedUserParam;
-
 	/** A binding to the user ptr we're reading the mesh from (if we are). */
 	FNiagaraParameterDirectBinding<UObject*> UserParamBinding;
-
+	UObject* CachedUserParam = nullptr;
 
 	//Cached ComponentToWorld Rotation
-	FQuat TransformQuat;
+	FQuat TransformQuat = FQuat::Identity;
 	//Cached ComponentToWorld.
-	FMatrix Transform;
+	FMatrix Transform = FMatrix::Identity;
 	//InverseTranspose of above for transforming normals/tangents.
-	FMatrix TransformInverseTransposed;
-	FTransform ComponentTransform;
+	FMatrix TransformInverseTransposed = FMatrix::Identity;
+	FTransform ComponentTransform = FTransform::Identity;
 	FNiagaraLWCConverter LwcConverter;
 
-	FVector DefaultUpVector;
+	FVector DefaultUpVector = FVector::UpVector;
 	FSplineCurves SplineCurves;
 	FNiagaraDataInterfaceSplineLUT SplineLUT;
 
-	bool bSyncedGPUCopy;
-	
+	bool bSyncedGPUCopy = false;
 
 	// We cache the version of the current spline curves so that we can reset the curves structure if we're using the LUT
 	TOptional<uint32> SplineCurvesVersion;
 
-	inline bool ResetRequired(UNiagaraDataInterfaceSpline* Interface, FNiagaraSystemInstance* SystemInstance) const;
-	
 	template<typename UseLUT>
 	float GetSplineLength() const;
-	
+
 	bool IsValid() const;
 
 
 
-	
+
 	template<typename UseLUT>
-	FVector GetLocationAtDistanceAlongSpline(float Distance, ESplineCoordinateSpace::Type CoordinateSpace) const;	
+	FVector GetLocationAtDistanceAlongSpline(float Distance, ESplineCoordinateSpace::Type CoordinateSpace) const;
 	template<typename UseLUT>
 	FQuat GetQuaternionAtDistanceAlongSpline(float Distance, ESplineCoordinateSpace::Type CoordinateSpace) const;
 	template<typename UseLUT>
@@ -168,8 +177,12 @@ class UNiagaraDataInterfaceSpline : public UNiagaraDataInterface
 	GENERATED_UCLASS_BODY()
 
 public:
-	/** The source actor from which to sample.  Note that this can only be set when used as a user variable on a component in the world.*/
-	UPROPERTY(EditAnywhere, Category = "Mesh", meta = (DisplayName = "Source Actor"))
+	/** Controls how the spline source is located each frame. */
+	UPROPERTY(EditAnywhere, Category = "Spline")
+	ENDISpline_SourceMode SourceMode = ENDISpline_SourceMode::Default;
+
+	/** The source actor from which to sample. Note that this can only be set when used as a user variable on a component in the world. */
+	UPROPERTY(EditAnywhere, Category = "Spline", meta = (DisplayName = "Source Actor", EditCondition = "SourceMode == ENDISpline_SourceMode::Default"))
 	TSoftObjectPtr<AActor> SoftSourceActor;
 
 #if WITH_EDITORONLY_DATA
@@ -177,16 +190,16 @@ public:
 	TObjectPtr<AActor> Source_DEPRECATED;
 #endif
 
-	/** Reference to a user parameter if we're reading one. This should  be an Object user parameter that is either a USplineComponent or an AActor containing a USplineComponent. */
-	UPROPERTY(EditAnywhere, Category = "Spline")
+	/** Reference to a user parameter if we're reading one. This should be an Object user parameter that is either a USplineComponent or an AActor containing a USplineComponent. */
+	UPROPERTY(EditAnywhere, Category = "Spline", meta = (EditCondition = "SourceMode == ENDISpline_SourceMode::Default || SourceMode == ENDISpline_SourceMode::ParameterBindingOnly"))
 	FNiagaraUserParameterBinding SplineUserParameter;
 
 	UPROPERTY(EditAnywhere, Category = "Spline")
 	bool bUseLUT;
 
-	UPROPERTY(EditAnywhere, Category = "Spline", Meta=(EditCondition="bUseLuT"))
+	UPROPERTY(EditAnywhere, Category = "Spline", Meta = (EditCondition = "bUseLuT"))
 	int32 NumLUTSteps;
-	
+
 	//UObject Interface
 	NIAGARA_API virtual void PostInitProperties() override;
 	NIAGARA_API virtual void PostLoad() override;
@@ -197,7 +210,7 @@ public:
 	NIAGARA_API virtual void DestroyPerInstanceData(void* PerInstanceData, FNiagaraSystemInstance* SystemInstance)override;
 	NIAGARA_API virtual bool PerInstanceTick(void* PerInstanceData, FNiagaraSystemInstance* SystemInstance, float DeltaSeconds) override;
 	NIAGARA_API virtual int32 PerInstanceDataSize()const override;
-	NIAGARA_API virtual void GetVMExternalFunction(const FVMExternalFunctionBindingInfo& BindingInfo, void* InstanceData, FVMExternalFunction &OutFunc) override;
+	NIAGARA_API virtual void GetVMExternalFunction(const FVMExternalFunctionBindingInfo& BindingInfo, void* InstanceData, FVMExternalFunction& OutFunc) override;
 	NIAGARA_API virtual bool Equals(const UNiagaraDataInterface* Other) const override;
 	virtual bool CanExecuteOnTarget(ENiagaraSimTarget Target)const override { return true; }
 	virtual bool HasPreSimulateTick() const override { return true; }
@@ -206,6 +219,8 @@ public:
 #endif
 	//UNiagaraDataInterface Interface End
 
+	USplineComponent* ResolveSplineComponent(FNiagaraSystemInstance* SystemInstance, FNDISpline_InstanceData* InstanceData) const;
+
 	template<typename UseLUT, typename TransformHandlerType, typename SplineSampleType>
 	void SampleSplinePositionByUnitDistance(FVectorVMExternalFunctionContext& Context);
 	template<typename UseLUT, typename TransformHandlerType, typename SplineSampleType>
@@ -220,7 +235,7 @@ public:
 	void SampleSplineTangentByUnitDistance(FVectorVMExternalFunctionContext& Context);
 	template<typename UseLUT>
 	void FindClosestUnitDistanceFromPositionWS(FVectorVMExternalFunctionContext& Context);
-	
+
 	NIAGARA_API void GetLocalToWorld(FVectorVMExternalFunctionContext& Context);
 	NIAGARA_API void GetLocalToWorldInverseTransposed(FVectorVMExternalFunctionContext& Context);
 
@@ -250,7 +265,7 @@ protected:
 	NIAGARA_API virtual void SetShaderParameters(const FNiagaraDataInterfaceSetShaderParametersContext& Context) const override;
 
 private:
-	
+
 	NIAGARA_API void WriteTransform(const FMatrix44f& ToWrite, FVectorVMExternalFunctionContext& Context);
 
 	TMap<FNiagaraSystemInstanceID, FNDISpline_InstanceData*> SystemInstancesToProxyData_GT;
diff --git a/Engine/Plugins/FX/Niagara/Source/Niagara/Private/NiagaraDataInterfaceSpline.cpp b/Engine/Plugins/FX/Niagara/Source/Niagara/Private/NiagaraDataInterfaceSpline.cpp
--- a/Engine/Plugins/FX/Niagara/Source/Niagara/Private/NiagaraDataInterfaceSpline.cpp
+++ b/Engine/Plugins/FX/Niagara/Source/Niagara/Private/NiagaraDataInterfaceSpline.cpp
@@ -406,9 +406,9 @@ bool UNiagaraDataInterfaceSpline::CopyToInternal(UNiagaraDataInterface* Destinat
 	}
 
 	UNiagaraDataInterfaceSpline* OtherTyped = CastChecked<UNiagaraDataInterfaceSpline>(Destination);
+	OtherTyped->SourceMode = SourceMode;
 	OtherTyped->SoftSourceActor = SoftSourceActor;
 	OtherTyped->SplineUserParameter = SplineUserParameter;
-	
 	OtherTyped->bUseLUT = bUseLUT;
 	OtherTyped->NumLUTSteps = NumLUTSteps;
 	
@@ -535,7 +535,11 @@ bool UNiagaraDataInterfaceSpline::Equals(const UNiagaraDataInterface* Other) con
 		return false;
 	}
 	const UNiagaraDataInterfaceSpline* OtherTyped = CastChecked<const UNiagaraDataInterfaceSpline>(Other);
-	return OtherTyped->SoftSourceActor == SoftSourceActor && OtherTyped->SplineUserParameter == SplineUserParameter && OtherTyped->bUseLUT == bUseLUT && OtherTyped->NumLUTSteps == NumLUTSteps;
+	return OtherTyped->SourceMode == SourceMode
+		&& OtherTyped->SoftSourceActor == SoftSourceActor
+		&& OtherTyped->SplineUserParameter == SplineUserParameter
+		&& OtherTyped->bUseLUT == bUseLUT
+		&& OtherTyped->NumLUTSteps == NumLUTSteps;
 }
 
 #if WITH_NIAGARA_DEBUGGER
@@ -620,17 +624,9 @@ int32 UNiagaraDataInterfaceSpline::PerInstanceDataSize()const
 bool UNiagaraDataInterfaceSpline::InitPerInstanceData(void* PerInstanceData, FNiagaraSystemInstance* SystemInstance)
 {
 	FNDISpline_InstanceData* InstData = new (PerInstanceData) FNDISpline_InstanceData();
-	SystemInstancesToProxyData_GT.Emplace(SystemInstance->GetId(), InstData);	
-
-	InstData->Component.Reset();
-	InstData->TransformQuat = FQuat::Identity;
-	InstData->Transform = FMatrix::Identity;
-	InstData->TransformInverseTransposed = FMatrix::Identity;
-	InstData->ComponentTransform = FTransform::Identity;
-	InstData->DefaultUpVector = FVector::UpVector;
-	InstData->bSyncedGPUCopy = false;
+	SystemInstancesToProxyData_GT.Emplace(SystemInstance->GetId(), InstData);
 
-	InstData->SplineLUT.Reset();
+	InstData->UserParamBinding.Init(SystemInstance->GetInstanceParameters(), SplineUserParameter.Parameter);
 
 	return true;
 }
@@ -657,72 +653,100 @@ void UNiagaraDataInterfaceSpline::DestroyPerInstanceData(void* PerInstanceData,
 	);
 }
 
-bool UNiagaraDataInterfaceSpline::PerInstanceTick(void* PerInstanceData, FNiagaraSystemInstance* SystemInstance, float DeltaSeconds)
+USplineComponent* UNiagaraDataInterfaceSpline::ResolveSplineComponent(FNiagaraSystemInstance* SystemInstance, FNDISpline_InstanceData* InstanceData) const
 {
-	check(SystemInstance);
-	FNDISpline_InstanceData* InstData = (FNDISpline_InstanceData*)PerInstanceData;
+	const bool bTryParameterBinding = SourceMode == ENDISpline_SourceMode::Default || SourceMode == ENDISpline_SourceMode::ParameterBindingOnly;
+	const bool bTrySourceActor = SourceMode == ENDISpline_SourceMode::Default;
+	const bool bTryAttachParent = SourceMode == ENDISpline_SourceMode::Default || SourceMode == ENDISpline_SourceMode::AttachParentOnly;
 
-	if (InstData && InstData->ResetRequired(this, SystemInstance))
+	// Parameter binding
+	UObject* ParameterBindingValue = InstanceData->UserParamBinding.GetValue();
+	if (bTryParameterBinding && ParameterBindingValue)
 	{
-		return true;
+		if (USplineComponent* BoundComponent = Cast<USplineComponent>(ParameterBindingValue))
+		{
+			if (IsValid(BoundComponent))
+			{
+				return BoundComponent;
+			}
+		}
+		else if (AActor* Actor = Cast<AActor>(ParameterBindingValue))
+		{
+			if (USplineComponent* ActorComponent = Actor->FindComponentByClass<USplineComponent>())
+			{
+				return ActorComponent;
+			}
+		}
+		else
+		{
+			UE_LOG(LogNiagara, Warning, TEXT("Spline data interface has an invalid object parameter type. Only USplineComponent or AActor are supported."));
+			UE_LOG(LogNiagara, Warning, TEXT("Invalid Parameter : %s"), *ParameterBindingValue->GetFullName());
+			UE_LOG(LogNiagara, Warning, TEXT("Niagara Component : %s"), *GetFullNameSafe(SystemInstance ? Cast<UNiagaraComponent>(SystemInstance->GetAttachComponent()) : nullptr));
+			UE_LOG(LogNiagara, Warning, TEXT("System : %s"), *GetFullNameSafe(SystemInstance ? SystemInstance->GetSystem() : nullptr));
+		}
 	}
 
-	if (!InstData)
+	// Source actor
+	if (bTrySourceActor)
 	{
-		return true;
+		if (AActor* SourceActor = SoftSourceActor.Get())
+		{
+			return SourceActor->FindComponentByClass<USplineComponent>();
+		}
 	}
 
-	USplineComponent* SplineComponent = InstData->Component.Get();
-	if (SplineComponent == nullptr)
+	// Attach parent
+	if (bTryAttachParent && SystemInstance)
 	{
-		if (SplineUserParameter.Parameter.IsValid() && InstData && SystemInstance != nullptr)
+		if (USceneComponent* AttachComp = SystemInstance->GetAttachComponent())
 		{
-			// Initialize the binding and retrieve the object. If a valid object is bound, we'll try and retrieve the Spline component from it.
-			// If it's not valid yet, we'll reset and do this again when/if a valid object is set on the binding
-			UObject* UserParamObject = InstData->UserParamBinding.Init(SystemInstance->GetInstanceParameters(), SplineUserParameter.Parameter);
-			InstData->CachedUserParam = UserParamObject;
-			if (UserParamObject)
-			{
-				if (USplineComponent* UserSplineComp = Cast<USplineComponent>(UserParamObject))
-				{
-					if (IsValid(UserSplineComp))
-					{
-						SplineComponent = UserSplineComp;
-					}
-				}
-				else if (AActor* UserParamActor = Cast<AActor>(UserParamObject))
-				{
-					SplineComponent = UserParamActor->FindComponentByClass<USplineComponent>();
-				}
-				else
-				{
-					//We have a valid, non-null UObject parameter type but it is not a type we can use to get a skeletal Spline from. 
-					UE_LOG(LogNiagara, Warning, TEXT("Spline data interface using object parameter with invalid type. Spline Data Interfaces can only get a valid Spline from SplineComponents or Actors."));
-					UE_LOG(LogNiagara, Warning, TEXT("Invalid Parameter : %s"), *UserParamObject->GetFullName());
-					UE_LOG(LogNiagara, Warning, TEXT("Niagara Component : %s"), *GetFullNameSafe(Cast<UNiagaraComponent>(SystemInstance->GetAttachComponent())));
-					UE_LOG(LogNiagara, Warning, TEXT("System : %s"), *GetFullNameSafe(SystemInstance->GetSystem()));
-				}
-			}
-			else
+			if (AActor* Owner = AttachComp->GetAttachmentRootActor())
 			{
-				// The binding exists, but no object is bound. Not warning here in case the user knows what they're doing.
+				return Owner->FindComponentByClass<USplineComponent>();
 			}
 		}
-		else if (SoftSourceActor != nullptr)
-		{
-			SplineComponent = SoftSourceActor->FindComponentByClass<USplineComponent>();
-		}
-		else if (USceneComponent* AttachComp = SystemInstance->GetAttachComponent())
+	}
+
+	return nullptr;
+}
+
+bool UNiagaraDataInterfaceSpline::PerInstanceTick(void* PerInstanceData, FNiagaraSystemInstance* SystemInstance, float DeltaSeconds)
+{
+	check(SystemInstance && PerInstanceData);
+	FNDISpline_InstanceData* InstData = static_cast<FNDISpline_InstanceData*>(PerInstanceData);
+
+	// Reolve the spline parameter
+	USplineComponent* SplineComponent = nullptr;
+	if (SourceMode == ENDISpline_SourceMode::Default)
+	{
+		if (SplineUserParameter.Parameter.IsValid())
 		{
-			if (AActor* Owner = AttachComp->GetAttachmentRootActor())
+			if (InstData->CachedUserParam && InstData->UserParamBinding.GetValue() != InstData->CachedUserParam)
 			{
-				SplineComponent = Owner->FindComponentByClass<USplineComponent>();
+				return true;
 			}
 		}
+
+		SplineComponent = InstData->Component.Get();
+		if (SplineComponent == nullptr)
+		{
+			SplineComponent = ResolveSplineComponent(SystemInstance, InstData);
+			InstData->Component = SplineComponent;
+		}
+	}
+	else
+	{
+		SplineComponent = ResolveSplineComponent(SystemInstance, InstData);
+	}
+
+	// Handle the parameter changing
+	if (SplineComponent != InstData->Component.Get())
+	{
 		InstData->Component = SplineComponent;
+		InstData->SplineCurvesVersion.Reset();
 	}
 
-	//Re-evaluate source in case it's changed?
+	// Re-evaluate source in case it's changed?
 	if (SplineComponent != nullptr)
 	{
 		FTransform SplineTransform = SplineComponent->GetComponentToWorld();
@@ -884,23 +908,6 @@ void FNiagaraDataInterfaceSplineLUT::FindNeighborKeys(float InDistance, int32& P
 	Alpha = FMath::Frac(Key);
 }
 
-bool FNDISpline_InstanceData::ResetRequired(UNiagaraDataInterfaceSpline* Interface, FNiagaraSystemInstance* SystemInstance) const
-{
-	
-	if (Interface->SplineUserParameter.Parameter.IsValid())
-	{
-		// Reset if the user object ptr has been changed to look at a new object
-		if (UserParamBinding.GetValue() != CachedUserParam)
-		{
-			return true;
-		}
-	}
-	
-
-	return false;
-}
-
-
 template<>
 float FNDISpline_InstanceData::GetSplineLength<TIntegralConstant<bool, false>>() const
 {
