-  Fix pending instances being missed with async sim cache reading

#Synced-CL 51711051

diff --git a/Engine/Plugins/FX/Niagara/Source/Niagara/Private/NiagaraSystemSimulation.cpp b/Engine/Plugins/FX/Niagara/Source/Niagara/Private/NiagaraSystemSimulation.cpp
--- a/Engine/Plugins/FX/Niagara/Source/Niagara/Private/NiagaraSystemSimulation.cpp
+++ b/Engine/Plugins/FX/Niagara/Source/Niagara/Private/NiagaraSystemSimulation.cpp
@@ -1565,14 +1565,20 @@ void FNiagaraSystemSimulation::SimCacheTick_GameThread(UNiagaraSimCache* SimCach
 	WaitForInstancesTickComplete();
 
 	// Get pointer to instances
-	TArray<FNiagaraSystemInstance*>& SystemInstances = GetSystemInstances(ENiagaraSystemInstanceState::Running);
+	TArray<FNiagaraSystemInstance*>& RunningSystemInstances = GetSystemInstances(ENiagaraSystemInstanceState::Running);
+	TArray<FNiagaraSystemInstance*>& PendingSpawnSystemInstances = GetSystemInstances(ENiagaraSystemInstanceState::PendingSpawn);
 
 	//-OPT: On cooked builds we do not need to do this every frame
 	if (SimCache->CanRead(System) == false)
 	{
-		for (int32 i=SystemInstances.Num() - 1; i >= 0; --i)
+		for (int32 i=RunningSystemInstances.Num() - 1; i >= 0; --i)
+		{
+			RunningSystemInstances[i]->Complete(false);
+		}
+
+		for (int32 i=PendingSpawnSystemInstances.Num() - 1; i >= 0; --i)
 		{
-			SystemInstances[i]->Complete(false);
+			PendingSpawnSystemInstances[i]->Complete(false);
 		}
 		return;
 	}
@@ -1589,45 +1595,55 @@ void FNiagaraSystemSimulation::SimCacheTick_GameThread(UNiagaraSimCache* SimCach
 			std::atomic<int> DebugCounter = 0;
 		#endif
 		};
-		TUniquePtr<FFinalizeTickBatch> FinalizeTickBatch(new FFinalizeTickBatch());
+		TSharedPtr<FFinalizeTickBatch> FinalizeTickBatch(new FFinalizeTickBatch());
 
-		FinalizeTickBatch->Instances.Reserve(SystemInstances.Num());
-		for (FNiagaraSystemInstance* SystemInstance : SystemInstances)
+		// Create events to signal old task system
+		ConcurrentTickGraphEvent = FGraphEvent::CreateGraphEvent();
+		AllWorkCompleteGraphEvent = FGraphEvent::CreateGraphEvent();
+
+		for (FNiagaraSystemInstance* SystemInstance : RunningSystemInstances)
 		{
 			SystemInstance->ConcurrentTickGraphEvent = ConcurrentTickGraphEvent;
-			FinalizeTickBatch->Instances.Add(SystemInstance);
-
-			FNiagaraSystemInstanceFinalizeRef FinalizeRef(&FinalizeTickBatch->Instances.Last());
-		#if DO_CHECK
-			FinalizeRef.SetDebugCounter(&FinalizeTickBatch->DebugCounter);
-		#endif
-			SystemInstance->SetPendingFinalize(FinalizeRef);
 		}
 
-		// Create events to signal old task system
-		ConcurrentTickGraphEvent = FGraphEvent::CreateGraphEvent();
-		AllWorkCompleteGraphEvent = FGraphEvent::CreateGraphEvent();
+		for (FNiagaraSystemInstance* SystemInstance : PendingSpawnSystemInstances)
+		{
+			SystemInstance->ConcurrentTickGraphEvent = ConcurrentTickGraphEvent;
+		}
 
 		// Execute concurrent work
 		UE::Tasks::FTask ConcurrentTask = UE::Tasks::Launch(
 			TEXT("Niagara::SimCacheTick_Concurrent"),
-			[this, SimCache, DeltaSeconds]()
+			[this, SimCache, DeltaSeconds, FinalizeTickBatch]()
 			{
 				TRACE_CPUPROFILER_EVENT_SCOPE(Niagara::SimCacheTick_Concurrent);
 
 				SimCacheTick_Concurrent(SimCache, DeltaSeconds);
 				ConcurrentTickGraphEvent->DispatchSubsequents();
+
+				TArray<FNiagaraSystemInstance*>& RunningSystemInstances = GetSystemInstances(ENiagaraSystemInstanceState::Running);
+				FinalizeTickBatch->Instances.Reserve(RunningSystemInstances.Num());
+				for (FNiagaraSystemInstance* SystemInstance : RunningSystemInstances)
+				{
+					FinalizeTickBatch->Instances.Add(SystemInstance);
+
+					FNiagaraSystemInstanceFinalizeRef FinalizeRef(&FinalizeTickBatch->Instances.Last());
+				#if DO_CHECK
+					FinalizeRef.SetDebugCounter(&FinalizeTickBatch->DebugCounter);
+				#endif
+					SystemInstance->SetPendingFinalize(FinalizeRef);
+				}
 			}
 		);
 
 		// Execute finalize work
 		UE::Tasks::Launch(
 			TEXT("Niagara::SimCacheTick_Finalize"),
-			[this, TickBatch=MoveTemp(FinalizeTickBatch)]()
+			[this, FinalizeTickBatch]()
 			{
 				TRACE_CPUPROFILER_EVENT_SCOPE(Niagara::SimCacheTick_Finalize);
 
-				for (FNiagaraSystemInstance* Instance : TickBatch->Instances)
+				for (FNiagaraSystemInstance* Instance : FinalizeTickBatch->Instances)
 				{
 					if (Instance)
 					{
@@ -1635,7 +1651,7 @@ void FNiagaraSystemSimulation::SimCacheTick_GameThread(UNiagaraSimCache* SimCach
 					}
 				}
 			#if DO_CHECK
-				ensureMsgf(TickBatch->DebugCounter.load() == 0, TEXT("Finalize batch is complete but counter is %d when it should be zero"), TickBatch->DebugCounter.load());
+				ensureMsgf(FinalizeTickBatch->DebugCounter.load() == 0, TEXT("Finalize batch is complete but counter is %d when it should be zero"), FinalizeTickBatch->DebugCounter.load());
 			#endif
 				AllWorkCompleteGraphEvent->DispatchSubsequents();
 			},
@@ -1648,9 +1664,9 @@ void FNiagaraSystemSimulation::SimCacheTick_GameThread(UNiagaraSimCache* SimCach
 	{
 		SimCacheTick_Concurrent(SimCache, DeltaSeconds);
 
-		for ( int32 i=SystemInstances.Num()-1; i >= 0; --i)
+		for ( int32 i=RunningSystemInstances.Num()-1; i >= 0; --i)
 		{
-			SystemInstances[i]->SimCacheTickFinalize_GameThread();
+			RunningSystemInstances[i]->SimCacheTickFinalize_GameThread();
 		}
 	}
 }
