- Additional information to debug nullptr node

#Synced-CL 51711051

diff --git a/Engine/Plugins/FX/Niagara/Source/NiagaraEditor/Private/NiagaraNode.cpp b/Engine/Plugins/FX/Niagara/Source/NiagaraEditor/Private/NiagaraNode.cpp
--- a/Engine/Plugins/FX/Niagara/Source/NiagaraEditor/Private/NiagaraNode.cpp
+++ b/Engine/Plugins/FX/Niagara/Source/NiagaraEditor/Private/NiagaraNode.cpp
@@ -81,12 +81,26 @@ void UNiagaraNode::PostLoad()
 						AssetPath = OwningGraph->GetPathName();
 					}
 
+					FString OwningNodeDetailed;
 					if (UEdGraphNode* OwningNode = Pin->GetOwningNode())
 					{
 						NodeName = FString::Printf(TEXT("%s:%s"), *GetNameSafe(OwningNode->GetClass()), *OwningNode->GetNodeTitle(ENodeTitleType::FullTitle).ToString());
+
+						OwningNodeDetailed.Appendf(TEXT("\nOwningNodePath - %s"), *GetFullNameSafe(OwningNode));
+						for (int32 i=0; i < OwningNode->Pins.Num(); ++i)
+						{
+							if (UEdGraphPin* OwningNodePin = OwningNode->Pins[i])
+							{
+								OwningNodeDetailed.Appendf(TEXT("\nPin[%d] - %s | %s"), i, *OwningNodePin->GetDisplayName().ToString(), *OwningNodePin->GetFName().ToString());
+							}
+							else
+							{
+								OwningNodeDetailed.Appendf(TEXT("\nPin[%d] - nullptr"), i);
+							}
+						}
 					}
 
-					UE_LOG(LogNiagaraEditor, Log, TEXT("Connected pin of Pin '%s' in '%s' of node '%s' in asset '%s' is nullptr!"), *PinName, *ScriptUsageName, *NodeName, *AssetPath);
+					UE_LOG(LogNiagaraEditor, Warning, TEXT("Connected pin of Pin '%s' in '%s' of node '%s' in asset '%s' is nullptr!%s"), *PinName, *ScriptUsageName, *NodeName, *AssetPath, *OwningNodeDetailed);
 				}
 			}
 		}
diff --git a/Engine/Plugins/FX/Niagara/Source/NiagaraEditor/Private/NiagaraScriptMergeManager.cpp b/Engine/Plugins/FX/Niagara/Source/NiagaraEditor/Private/NiagaraScriptMergeManager.cpp
--- a/Engine/Plugins/FX/Niagara/Source/NiagaraEditor/Private/NiagaraScriptMergeManager.cpp
+++ b/Engine/Plugins/FX/Niagara/Source/NiagaraEditor/Private/NiagaraScriptMergeManager.cpp
@@ -246,7 +246,12 @@ FNiagaraStackFunctionMergeAdapter::FNiagaraStackFunctionMergeAdapter(const FVers
 
 	FString UniqueEmitterName = InOwningEmitter.Emitter->GetUniqueEmitterName();
 
-	FCompileConstantResolver ConstantResolver(InOwningEmitter, FNiagaraStackGraphUtilities::GetEmitterOutputNodeForStackNode(*FunctionCallNode)->GetUsage());
+	UNiagaraNodeOutput* FunctionCallOutputNode = FNiagaraStackGraphUtilities::GetEmitterOutputNodeForStackNode(*FunctionCallNode);
+	if (FunctionCallOutputNode == nullptr)
+	{
+		UE_LOG(LogNiagaraEditor, Warning, TEXT("FunctionCallOutputNode is nullptr in FNiagaraStackFunctionMergeAdapter, results may be incorrect"));
+	}
+	FCompileConstantResolver ConstantResolver = FunctionCallOutputNode ? FCompileConstantResolver(InOwningEmitter, FunctionCallOutputNode->GetUsage()) : FCompileConstantResolver();
 
 	// Collect explicit overrides set via parameter map set nodes.
 	TSet<FName> AliasedInputsAdded;
diff --git a/Engine/Plugins/FX/Niagara/Source/NiagaraEditor/Private/ViewModels/Stack/NiagaraStackGraphUtilities.cpp b/Engine/Plugins/FX/Niagara/Source/NiagaraEditor/Private/ViewModels/Stack/NiagaraStackGraphUtilities.cpp
--- a/Engine/Plugins/FX/Niagara/Source/NiagaraEditor/Private/ViewModels/Stack/NiagaraStackGraphUtilities.cpp
+++ b/Engine/Plugins/FX/Niagara/Source/NiagaraEditor/Private/ViewModels/Stack/NiagaraStackGraphUtilities.cpp
@@ -454,7 +454,7 @@ OutputNodeType* GetEmitterOutputNodeForStackNodeInternal(InputNodeType& StackNod
 		{
 			for (UEdGraphPin* LinkedPin : OutputPin->LinkedTo)
 			{
-				InputNodeType* LinkedNiagaraNode = Cast<UNiagaraNode>(LinkedPin->GetOwningNode());
+				InputNodeType* LinkedNiagaraNode = LinkedPin ? Cast<UNiagaraNode>(LinkedPin->GetOwningNode()) : nullptr;
 				if (LinkedNiagaraNode != nullptr && NodesSeen.Contains(LinkedNiagaraNode) == false)
 				{
 					NodesSeen.Add(LinkedNiagaraNode);
