- Material dynamic parameter default values show up as tooltips and right click option to set to default value

#Synced-CL 51762535

diff --git a/Engine/Plugins/FX/Niagara/Source/NiagaraEditor/Private/Customizations/Stack/NiagaraStackObjectPropertyCustomization_StatelessModule_DynamicMaterialParameters.cpp b/Engine/Plugins/FX/Niagara/Source/NiagaraEditor/Private/Customizations/Stack/NiagaraStackObjectPropertyCustomization_StatelessModule_DynamicMaterialParameters.cpp
--- a/Engine/Plugins/FX/Niagara/Source/NiagaraEditor/Private/Customizations/Stack/NiagaraStackObjectPropertyCustomization_StatelessModule_DynamicMaterialParameters.cpp
+++ b/Engine/Plugins/FX/Niagara/Source/NiagaraEditor/Private/Customizations/Stack/NiagaraStackObjectPropertyCustomization_StatelessModule_DynamicMaterialParameters.cpp
@@ -7,10 +7,16 @@
 #include "Containers/StaticBitArray.h"
 #include "Materials/MaterialExpressionDynamicParameter.h"
 #include "Stateless/NiagaraStatelessEmitter.h"
+#include "Stateless/Modules/NiagaraStatelessModule_DynamicMaterialParameters.h"
 #include "ViewModels/NiagaraEmitterHandleViewModel.h"
 #include "ViewModels/NiagaraEmitterViewModel.h"
 #include "ViewModels/NiagaraSystemViewModel.h"
 #include "Materials/Material.h"
+#include "Framework/MultiBox/MultiBoxBuilder.h"
+#include "PropertyHandle.h"
+#include "ScopedTransaction.h"
+
+#define LOCTEXT_NAMESPACE "NiagaraStackObjectPropertyCustomization_StatelessModule_DynamicMaterialParameters"
 
 FNiagaraStackObjectPropertyCustomization_StatelessModule_DynamicMaterialParameters::
 FNiagaraStackObjectPropertyCustomization_StatelessModule_DynamicMaterialParameters()
@@ -53,10 +59,16 @@ TOptional<TSharedPtr<SWidget>> FNiagaraStackObjectPropertyCustomization_Stateles
 
 	if(DisplayName.IsSet())
 	{
-		TSharedPtr<SWidget> Widget = SNew(STextBlock).Text(DisplayName.GetValue());
+		TOptional<float> MaterialDefault = TryGetDefaultValueForDynamicMaterialParameter(EmitterHandleViewModel, ParameterIndex, ParameterChannel);
+
+		TSharedPtr<SWidget> Widget = SNew(STextBlock)
+			.Text(DisplayName.GetValue())
+			.ToolTipText(MaterialDefault.IsSet()
+				? FText::FormatOrdered(LOCTEXT("MaterialDefaultTooltipFormat", "Material Default: {0}"), FText::AsNumber(MaterialDefault.GetValue()))
+				: FText::GetEmpty());
 		return Widget;
 	}
-	
+
 	return TOptional<TSharedPtr<SWidget>>();
 }
 
@@ -178,6 +190,133 @@ TOptional<FText> FNiagaraStackObjectPropertyCustomization_StatelessModule_Dynami
 	return TOptional<FText>();
 }
 
+TOptional<float> FNiagaraStackObjectPropertyCustomization_StatelessModule_DynamicMaterialParameters::TryGetDefaultValueForDynamicMaterialParameter(TSharedPtr<FNiagaraEmitterHandleViewModel> EmitterHandleViewModel, int32 ParameterIndex, int32 ParameterChannel) const
+{
+	UNiagaraStatelessEmitter* OwningStatelessEmitter = EmitterHandleViewModel->GetEmitterHandle()->GetStatelessEmitter();
+
+	FNiagaraEmitterInstance* ThisEmitterInstance = EmitterHandleViewModel->GetEmitterViewModel()->GetSimulation().Pin().Get();
+	TArray<UMaterial*> Materials = GetMaterialsFromEmitter(*OwningStatelessEmitter, ThisEmitterInstance);
+
+	for (UMaterial* Material : Materials)
+	{
+		if (Material == nullptr)
+		{
+			continue;
+		}
+
+		TArray<UMaterialExpression*> Expressions;
+		Material->GetAllExpressionsInMaterialAndFunctionsOfType<UMaterialExpression>(Expressions);
+		for (UMaterialExpression* Expression : Expressions)
+		{
+			const UMaterialExpressionDynamicParameter* DynamicParameterExpression = Cast<UMaterialExpressionDynamicParameter>(Expression);
+			if (DynamicParameterExpression == nullptr || DynamicParameterExpression->ParameterIndex != static_cast<uint32>(ParameterIndex))
+			{
+				continue;
+			}
+
+			return DynamicParameterExpression->DefaultValue.Component(ParameterChannel);
+		}
+	}
+
+	return TOptional<float>();
+}
+
+void FNiagaraStackObjectPropertyCustomization_StatelessModule_DynamicMaterialParameters::GenerateContextMenuActions(UNiagaraStackPropertyRow* StackPropertyRow, FMenuBuilder& MenuBuilder) const
+{
+	Super::GenerateContextMenuActions(StackPropertyRow, MenuBuilder);
+
+	TSharedRef<FDetailTreeNode> DetailTreeNode = StaticCastSharedRef<FDetailTreeNode>(StackPropertyRow->GetDetailTreeNode());
+
+	int32 ParameterIndex = INDEX_NONE;
+	int32 ParameterChannel = INDEX_NONE;
+	GetParameterIndexAndChannel(DetailTreeNode, ParameterIndex, ParameterChannel);
+
+	if (ParameterIndex == INDEX_NONE || ParameterChannel == INDEX_NONE)
+	{
+		return;
+	}
+
+	TSharedPtr<FNiagaraSystemViewModel> SystemViewModel = StackPropertyRow->GetSystemViewModel();
+	TSharedPtr<FNiagaraEmitterHandleViewModel> EmitterHandleViewModel = SystemViewModel->GetEmitterHandleViewModelForEmitter(StackPropertyRow->GetEmitterViewModel()->GetEmitter());
+
+	TOptional<float> MaterialDefault = TryGetDefaultValueForDynamicMaterialParameter(EmitterHandleViewModel, ParameterIndex, ParameterChannel);
+	if (!MaterialDefault.IsSet())
+	{
+		return;
+	}
+
+	const float DefaultValue = MaterialDefault.GetValue();
+	TWeakObjectPtr<UNiagaraStackPropertyRow> WeakPropertyRow(StackPropertyRow);
+
+	FUIAction SetToDefaultAction(
+		FExecuteAction::CreateLambda([WeakPropertyRow, ParameterIndex, ParameterChannel, DefaultValue]()
+		{
+			UNiagaraStackPropertyRow* PropertyRow = WeakPropertyRow.Get();
+			if (PropertyRow == nullptr)
+			{
+				return;
+			}
+
+			TSharedPtr<IPropertyHandle> PropertyHandle = PropertyRow->GetDetailTreeNode()->CreatePropertyHandle();
+			if (!PropertyHandle.IsValid() || !PropertyHandle->IsValidHandle())
+			{
+				return;
+			}
+
+			TArray<UObject*> OuterObjects;
+			PropertyHandle->GetOuterObjects(OuterObjects);
+
+			for (UObject* OuterObject : OuterObjects)
+			{
+				UNiagaraStatelessModule_DynamicMaterialParameters* Module = Cast<UNiagaraStatelessModule_DynamicMaterialParameters>(OuterObject);
+				if (Module == nullptr)
+				{
+					continue;
+				}
+
+				FNiagaraStatelessDynamicParameterSet* ParameterSet = nullptr;
+				switch (ParameterIndex)
+				{
+					case 0: ParameterSet = &Module->Parameter0; break;
+					case 1: ParameterSet = &Module->Parameter1; break;
+					case 2: ParameterSet = &Module->Parameter2; break;
+					case 3: ParameterSet = &Module->Parameter3; break;
+					default: break;
+				}
+
+				if (ParameterSet == nullptr)
+				{
+					continue;
+				}
+
+				FNiagaraDistributionFloat* Distribution = nullptr;
+				switch (ParameterChannel)
+				{
+					case 0: Distribution = &ParameterSet->XChannelDistribution; break;
+					case 1: Distribution = &ParameterSet->YChannelDistribution; break;
+					case 2: Distribution = &ParameterSet->ZChannelDistribution; break;
+					case 3: Distribution = &ParameterSet->WChannelDistribution; break;
+					default: break;
+				}
+
+				if (Distribution != nullptr)
+				{
+					FScopedTransaction Transaction(LOCTEXT("SetToMaterialDefault", "Set Channel to Material Default"));
+					PropertyHandle->NotifyPreChange();
+					Distribution->InitConstant(DefaultValue);
+					PropertyHandle->NotifyPostChange(EPropertyChangeType::ValueSet);
+				}
+			}
+		})
+	);
+
+	MenuBuilder.AddMenuEntry(
+		FText::FormatOrdered(LOCTEXT("SetMaterialDefaultFormat", "Set to Material Default ({0})"), FText::AsNumber(DefaultValue)),
+		LOCTEXT("SetMaterialDefaultFormatTooltip", "Sets the channel constant to the default value defined in the material's Dynamic Parameter node"),
+		FSlateIcon(),
+		SetToDefaultAction);
+}
+
 void FNiagaraStackObjectPropertyCustomization_StatelessModule_DynamicMaterialParameters::GetChannelUsedBitMask(FExpressionInput* Input, TStaticBitArray<4>& ChannelUsedMask) const
 {
 	if (Input->Expression)
@@ -258,3 +397,5 @@ void FNiagaraStackObjectPropertyCustomization_StatelessModule_DynamicMaterialPar
 		}
 	}
 }
+
+#undef LOCTEXT_NAMESPACE
diff --git a/Engine/Plugins/FX/Niagara/Source/NiagaraEditor/Private/Customizations/Stack/NiagaraStackObjectPropertyCustomization_StatelessModule_DynamicMaterialParameters.h b/Engine/Plugins/FX/Niagara/Source/NiagaraEditor/Private/Customizations/Stack/NiagaraStackObjectPropertyCustomization_StatelessModule_DynamicMaterialParameters.h
--- a/Engine/Plugins/FX/Niagara/Source/NiagaraEditor/Private/Customizations/Stack/NiagaraStackObjectPropertyCustomization_StatelessModule_DynamicMaterialParameters.h
+++ b/Engine/Plugins/FX/Niagara/Source/NiagaraEditor/Private/Customizations/Stack/NiagaraStackObjectPropertyCustomization_StatelessModule_DynamicMaterialParameters.h
@@ -20,16 +20,20 @@ struct FExpressionInput;
 class FNiagaraStackObjectPropertyCustomization_StatelessModule_DynamicMaterialParameters
 	: public FNiagaraStackObjectPropertyCustomization
 {
+	using Super = FNiagaraStackObjectPropertyCustomization;
+
 public:
 	UE_API FNiagaraStackObjectPropertyCustomization_StatelessModule_DynamicMaterialParameters();
 
 	static UE_API TSharedRef<FNiagaraStackObjectPropertyCustomization> MakeInstance();
 	
 	UE_API virtual TOptional<TSharedPtr<SWidget>> GenerateNameWidget(UNiagaraStackPropertyRow* PropertyRow) const override;
+	UE_API virtual void GenerateContextMenuActions(UNiagaraStackPropertyRow* PropertyRow, FMenuBuilder& MenuBuilder) const override;
 
 private:
 	UE_API TOptional<FText> TryGetDisplayNameForDynamicMaterialParameter(TSharedPtr<FNiagaraEmitterHandleViewModel> EmitterHandleViewModel, int32 ParameterIndex, int32 ParameterChannel) const;
-	
+	UE_API TOptional<float> TryGetDefaultValueForDynamicMaterialParameter(TSharedPtr<FNiagaraEmitterHandleViewModel> EmitterHandleViewModel, int32 ParameterIndex, int32 ParameterChannel) const;
+
 	UE_API void GetChannelUsedBitMask(FExpressionInput* Input, TStaticBitArray<4>& ChannelUsedMask) const;
 	UE_API TArray<UMaterial*> GetMaterialsFromEmitter(const UNiagaraStatelessEmitter& InEmitter, const FNiagaraEmitterInstance* InEmitterInstance) const;
 	UE_API void GetParameterIndexAndChannel(TSharedRef<FDetailTreeNode> DetailTreeNode, int32& OutParameterIndex, int32& OutParameterChannel) const;
diff --git a/Engine/Plugins/FX/Niagara/Source/NiagaraEditor/Private/ViewModels/Stack/NiagaraStackFunctionInput.cpp b/Engine/Plugins/FX/Niagara/Source/NiagaraEditor/Private/ViewModels/Stack/NiagaraStackFunctionInput.cpp
--- a/Engine/Plugins/FX/Niagara/Source/NiagaraEditor/Private/ViewModels/Stack/NiagaraStackFunctionInput.cpp
+++ b/Engine/Plugins/FX/Niagara/Source/NiagaraEditor/Private/ViewModels/Stack/NiagaraStackFunctionInput.cpp
@@ -294,7 +294,16 @@ bool UNiagaraStackFunctionInput::GetShouldShowInStack() const
 FText UNiagaraStackFunctionInput::GetTooltipText() const
 {
 	FText Description = SummaryViewTooltipOverride.IsSet() && !SummaryViewTooltipOverride->Get().IsEmptyOrWhitespace() ? SummaryViewTooltipOverride->Get() : InputMetaData.IsSet() ? InputMetaData->Description : FText::GetEmpty();
-	return FNiagaraEditorUtilities::FormatVariableDescription(Description, GetDisplayName(), InputType.GetNameText());
+	FText BaseTooltip = FNiagaraEditorUtilities::FormatVariableDescription(Description, GetDisplayName(), InputType.GetNameText());
+
+	// Special handling to resolve dynamic parameter default value
+	TOptional<float> MaterialDefault = GetMaterialDynamicParamDefaultValue();
+	if (MaterialDefault.IsSet())
+	{
+		return FText::Format(LOCTEXT("MaterialDefaultFormat", "{0}\nMaterial Default: {1}"), BaseTooltip, FText::AsNumber(MaterialDefault.GetValue()));
+	}
+
+	return BaseTooltip;
 }
 
 bool UNiagaraStackFunctionInput::GetIsEnabled() const
@@ -811,6 +820,32 @@ FText UNiagaraStackFunctionInput::GetValueToolTip() const
 	return ValueToolTipCache.GetValue();
 }
 
+TArray<FNiagaraStackFunctionInputActions> UNiagaraStackFunctionInput::GetCustomActions()
+{
+	TArray<FNiagaraStackFunctionInputActions> Actions;
+
+	TOptional<float> MaterialDefault = GetMaterialDynamicParamDefaultValue();
+	if (MaterialDefault.IsSet())
+	{
+		const float DefaultValue = MaterialDefault.GetValue();
+		Actions.Emplace(
+			FText::FormatOrdered(LOCTEXT("SetToMaterialDefault", "Set to Material Default ({0})"), FText::AsNumber(DefaultValue)),
+			LOCTEXT("SetToMaterialDefaultToolTip", "Sets this input to a constant matching the default value defined in the material's Dynamic Parameter node."),
+			FExecuteAction::CreateLambda(
+				[this, DefaultValue]()
+				{
+					TSharedRef<FStructOnScope> ValueStruct = MakeShared<FStructOnScope>(GetInputType().GetStruct());
+					float* FloatData = reinterpret_cast<float*>(ValueStruct->GetStructMemory());
+					*FloatData = DefaultValue;
+					SetLocalValue(ValueStruct);
+				}
+			)
+		);
+	}
+
+	return Actions;
+}
+
 UNiagaraStackEntry::FStackIssueFixDelegate UNiagaraStackFunctionInput::GetUpgradeDynamicInputVersionFix()
 {
 	if (!InputValues.DynamicNode.IsValid())
@@ -1368,6 +1403,85 @@ FString UNiagaraStackFunctionInput::ResolveDisplayNameArgument(const FString& In
 	return InArg.Replace(TEXT("MaterialDynamic"), TEXT("")) + TEXT(" (Parameter not used in materials.)");
 }
 
+TOptional<float> UNiagaraStackFunctionInput::GetMaterialDynamicParamDefaultValue() const
+{
+	if (!InputMetaData.IsSet())
+	{
+		return TOptional<float>();
+	}
+
+	const FString* FoundDisplayNameArg0 = InputMetaData->PropertyMetaData.Find(TEXT("DisplayNameArg0"));
+	if (FoundDisplayNameArg0 == nullptr || !FoundDisplayNameArg0->StartsWith(TEXT("MaterialDynamicParam")))
+	{
+		return TOptional<float>();
+	}
+
+	// Parse ParamIdx (channel: 0-3) and ParamSlotIdx (parameter index: 0-3) from the suffix.
+	// The format is MaterialDynamicParam{ParamIdx}{X}{ParamSlotIdx} where X is an ignored middle char.
+	const FString& InArg = *FoundDisplayNameArg0;
+	FString Suffix = InArg.Right(3);
+	int32 ParamIdx;
+	LexFromString(ParamIdx, *Suffix.Left(1));
+	int32 ParamSlotIdx;
+	LexFromString(ParamSlotIdx, *Suffix.Right(1));
+
+	if (ParamIdx < 0 || ParamIdx > 3 || ParamSlotIdx < 0 || ParamSlotIdx > 3)
+	{
+		return TOptional<float>();
+	}
+
+	TArray<UMaterial*> Materials;
+	TSharedPtr<FNiagaraEmitterViewModel> ThisEmitterViewModel = GetEmitterViewModel();
+	if (ThisEmitterViewModel.IsValid())
+	{
+		FNiagaraEmitterInstance* ThisEmitterInstance = ThisEmitterViewModel->GetSimulation().Pin().Get();
+		if (ThisEmitterViewModel->GetEmitter().Emitter != nullptr)
+		{
+			for (UNiagaraRendererProperties* RenderProperties : ThisEmitterViewModel->GetEmitter().GetEmitterData()->GetRenderers())
+			{
+				TArray<UMaterialInterface*> UsedMaterialInterfaces;
+				RenderProperties->GetUsedMaterials(ThisEmitterInstance, UsedMaterialInterfaces);
+				for (UMaterialInterface* UsedMaterialInterface : UsedMaterialInterfaces)
+				{
+					if (UsedMaterialInterface != nullptr)
+					{
+						UMaterial* UsedMaterial = UsedMaterialInterface->GetBaseMaterial();
+						if (UsedMaterial != nullptr)
+						{
+							Materials.AddUnique(UsedMaterial);
+							break;
+						}
+					}
+				}
+			}
+		}
+	}
+
+	for (UMaterial* Material : Materials)
+	{
+		if (Material == nullptr)
+		{
+			continue;
+		}
+
+		TArray<UMaterialExpression*> Expressions;
+		Material->GetAllExpressionsInMaterialAndFunctionsOfType<UMaterialExpression>(Expressions);
+		for (UMaterialExpression* Expression : Expressions)
+		{
+			const UMaterialExpressionDynamicParameter* DynParam = Cast<UMaterialExpressionDynamicParameter>(Expression);
+			if (DynParam == nullptr || DynParam->ParameterIndex != static_cast<uint32>(ParamSlotIdx))
+			{
+				continue;
+			}
+
+			// DefaultValue is FLinearColor; ParamIdx maps to RGBA channels (0=R, 1=G, 2=B, 3=A)
+			return DynParam->DefaultValue.Component(ParamIdx);
+		}
+	}
+
+	return TOptional<float>();
+}
+
 void UNiagaraStackFunctionInput::RefreshValues()
 {
 	if (ensureMsgf(IsStaticParameter() || InputParameterHandle.IsModuleHandle(), TEXT("Function inputs can only be generated for module parameters.")) == false)
diff --git a/Engine/Plugins/FX/Niagara/Source/NiagaraEditor/Public/Customizations/NiagaraStackObjectPropertyCustomization.h b/Engine/Plugins/FX/Niagara/Source/NiagaraEditor/Public/Customizations/NiagaraStackObjectPropertyCustomization.h
--- a/Engine/Plugins/FX/Niagara/Source/NiagaraEditor/Public/Customizations/NiagaraStackObjectPropertyCustomization.h
+++ b/Engine/Plugins/FX/Niagara/Source/NiagaraEditor/Public/Customizations/NiagaraStackObjectPropertyCustomization.h
@@ -6,6 +6,7 @@
 #include "UObject/Object.h"
 #include "Widgets/SWidget.h"
 
+class FMenuBuilder;
 class UNiagaraStackPropertyRow;
 
 /**
@@ -16,7 +17,10 @@ class FNiagaraStackObjectPropertyCustomization : public TSharedFromThis<FNiagara
 {
 public:
 	virtual ~FNiagaraStackObjectPropertyCustomization() = default;
-	
+
 	/** If specified, will generate a custom name widget for a given property row. Only applicable for rows that generate typical name & value widgets. */
 	virtual TOptional<TSharedPtr<SWidget>> GenerateNameWidget(UNiagaraStackPropertyRow* PropertyRow) const { return TOptional<TSharedPtr<SWidget>>(); }
+
+	/** If specified, adds custom entries to the right-click context menu for a given property row. */
+	virtual void GenerateContextMenuActions(UNiagaraStackPropertyRow* PropertyRow, FMenuBuilder& MenuBuilder) const {}
 };
diff --git a/Engine/Plugins/FX/Niagara/Source/NiagaraEditor/Public/ViewModels/Stack/NiagaraStackFunctionInput.h b/Engine/Plugins/FX/Niagara/Source/NiagaraEditor/Public/ViewModels/Stack/NiagaraStackFunctionInput.h
--- a/Engine/Plugins/FX/Niagara/Source/NiagaraEditor/Public/ViewModels/Stack/NiagaraStackFunctionInput.h
+++ b/Engine/Plugins/FX/Niagara/Source/NiagaraEditor/Public/ViewModels/Stack/NiagaraStackFunctionInput.h
@@ -29,6 +29,13 @@ class UNiagaraClipboardFunction;
 class UNiagaraScriptVariable;
 class FNiagaraPlaceholderDataInterfaceHandle;
 
+struct FNiagaraStackFunctionInputActions
+{
+	FText			Label;
+	FText			ToolTip;
+	FExecuteAction	ExecuteAction;
+};
+
 /** Represents a single module input in the module stack view model. */
 UCLASS(MinimalAPI)
 class UNiagaraStackFunctionInput : public UNiagaraStackItemContent
@@ -155,6 +162,9 @@ public:
 	/** Gets the tooltip that should be shown for the value of this input. */
 	NIAGARAEDITOR_API FText GetValueToolTip() const;
 
+	/** Get a list of custom actions that can be executed on the input from the right click context menu */
+	NIAGARAEDITOR_API TArray<FNiagaraStackFunctionInputActions> GetCustomActions();
+
 	/** Gets the tooltip that should be shown for the value of this input. */
 	NIAGARAEDITOR_API FText GetCollapsedStateText() const;
 
@@ -376,6 +386,7 @@ protected:
 	NIAGARAEDITOR_API bool UpdateRapidIterationParametersForAffectedScripts(const uint8* Data);
 	NIAGARAEDITOR_API bool RemoveRapidIterationParametersForAffectedScripts(bool bUpdateGraphGuidsForAffected = false);
 	NIAGARAEDITOR_API FString ResolveDisplayNameArgument(const FString& InArg) const;
+	NIAGARAEDITOR_API TOptional<float> GetMaterialDynamicParamDefaultValue() const;
 	NIAGARAEDITOR_API FStackIssueFixDelegate GetUpgradeDynamicInputVersionFix();
 
 	/** Sets whether or not this input was created by another input to display and edit a shared linked value. */
diff --git a/Engine/Plugins/FX/Niagara/Source/NiagaraEditorWidgets/Private/Stack/NiagaraStackPropertyRowUtilities.cpp b/Engine/Plugins/FX/Niagara/Source/NiagaraEditorWidgets/Private/Stack/NiagaraStackPropertyRowUtilities.cpp
--- a/Engine/Plugins/FX/Niagara/Source/NiagaraEditorWidgets/Private/Stack/NiagaraStackPropertyRowUtilities.cpp
+++ b/Engine/Plugins/FX/Niagara/Source/NiagaraEditorWidgets/Private/Stack/NiagaraStackPropertyRowUtilities.cpp
@@ -2,6 +2,8 @@
 
 #include "NiagaraStackPropertyRowUtilities.h"
 
+#include "Customizations/NiagaraStackObjectPropertyCustomization.h"
+#include "Framework/MultiBox/MultiBoxBuilder.h"
 #include "HAL/PlatformApplicationMisc.h"
 #include "IDetailTreeNode.h"
 #include "PropertyHandle.h"
@@ -23,6 +25,12 @@ void FNiagaraStackPropertyRowUtilities::OnFillPropertyRowContextMenu(FMenuBuilde
 		return;
 	}
 
+	// Allow the property customization to contribute context menu actions regardless of copy/paste support.
+	if (TSharedPtr<const FNiagaraStackObjectPropertyCustomization> Customization = PropertyRow->GetCustomization())
+	{
+		Customization->GenerateContextMenuActions(PropertyRow, MenuBuilder);
+	}
+
 	if (PropertyRow->SupportsCopy() && PropertyRow->SupportsPaste())
 	{
 		// If the row supports stack copy, don't add the actions from the property editor.
diff --git a/Engine/Plugins/FX/Niagara/Source/NiagaraEditorWidgets/Private/Stack/SNiagaraStackFunctionInputName.cpp b/Engine/Plugins/FX/Niagara/Source/NiagaraEditorWidgets/Private/Stack/SNiagaraStackFunctionInputName.cpp
--- a/Engine/Plugins/FX/Niagara/Source/NiagaraEditorWidgets/Private/Stack/SNiagaraStackFunctionInputName.cpp
+++ b/Engine/Plugins/FX/Niagara/Source/NiagaraEditorWidgets/Private/Stack/SNiagaraStackFunctionInputName.cpp
@@ -165,6 +165,20 @@ void SNiagaraStackFunctionInputName::FillRowContextMenu(class FMenuBuilder& Menu
 		}
 		MenuBuilder.EndSection();
 	}
+
+	{
+		TArray<FNiagaraStackFunctionInputActions> CustomActions = FunctionInput->GetCustomActions();
+		if (CustomActions.Num() > 0)
+		{
+			MenuBuilder.BeginSection("InputActions", LOCTEXT("InputActions", "Input Actions"));
+			for (const FNiagaraStackFunctionInputActions& Action : CustomActions)
+			{
+				MenuBuilder.AddMenuEntry(Action.Label, Action.ToolTip, FSlateIcon(), FUIAction(Action.ExecuteAction));
+			}
+			MenuBuilder.EndSection();
+		}
+	}
+
 	if (FunctionInput->GetValueMode() == UNiagaraStackFunctionInput::EValueMode::Dynamic && GbEnableExperimentalInlineDynamicInputs)
 	{
 		MenuBuilder.BeginSection("View", LOCTEXT("ViewHeader", "View"));
