- Move shared dynamic parameter functions into shared utils class

#rnx

#Synced-CL 51762535

diff --git a/Engine/Plugins/FX/Niagara/Source/NiagaraEditor/Private/Customizations/Stack/NiagaraStackObjectPropertyCustomization_StatelessModule_DynamicMaterialParameters.cpp b/Engine/Plugins/FX/Niagara/Source/NiagaraEditor/Private/Customizations/Stack/NiagaraStackObjectPropertyCustomization_StatelessModule_DynamicMaterialParameters.cpp
--- a/Engine/Plugins/FX/Niagara/Source/NiagaraEditor/Private/Customizations/Stack/NiagaraStackObjectPropertyCustomization_StatelessModule_DynamicMaterialParameters.cpp
+++ b/Engine/Plugins/FX/Niagara/Source/NiagaraEditor/Private/Customizations/Stack/NiagaraStackObjectPropertyCustomization_StatelessModule_DynamicMaterialParameters.cpp
@@ -1,6 +1,7 @@
 ﻿// Copyright Epic Games, Inc. All Rights Reserved.
 
 #include "NiagaraStackObjectPropertyCustomization_StatelessModule_DynamicMaterialParameters.h"
+#include "Utilities/NiagaraDynamicMaterialParameterUtils.h"
 #include "Widgets/Text/STextBlock.h"
 #include "ViewModels/Stack/NiagaraStackPropertyRow.h"
 #include "DetailTreeNode.h"
@@ -75,10 +76,8 @@ TOptional<TSharedPtr<SWidget>> FNiagaraStackObjectPropertyCustomization_Stateles
 TOptional<FText> FNiagaraStackObjectPropertyCustomization_StatelessModule_DynamicMaterialParameters::TryGetDisplayNameForDynamicMaterialParameter(TSharedPtr<FNiagaraEmitterHandleViewModel> EmitterHandleViewModel, int32 ParameterIndex, int32 ParameterChannel) const
 {
 	UNiagaraStatelessEmitter* OwningStatelessEmitter = EmitterHandleViewModel->GetEmitterHandle()->GetStatelessEmitter();
-
-	TArray<UMaterial*> Materials;
 	FNiagaraEmitterInstance* ThisEmitterInstance = EmitterHandleViewModel->GetEmitterViewModel()->GetSimulation().Pin().Get();
-	Materials = GetMaterialsFromEmitter(*OwningStatelessEmitter, ThisEmitterInstance);
+	TArray<UMaterial*> Materials = NiagaraDynamicMaterialParameterUtils::GetMaterialsFromRenderers(OwningStatelessEmitter->GetRenderers(), ThisEmitterInstance);
 	
 	TMap<UMaterialExpressionDynamicParameter*, TStaticBitArray<4>> DynamicParameterExpressionToOutputMaskMap;
 	TArray<FExpressionInput*> ExpressionInputsToProcess;
@@ -193,32 +192,9 @@ TOptional<FText> FNiagaraStackObjectPropertyCustomization_StatelessModule_Dynami
 TOptional<float> FNiagaraStackObjectPropertyCustomization_StatelessModule_DynamicMaterialParameters::TryGetDefaultValueForDynamicMaterialParameter(TSharedPtr<FNiagaraEmitterHandleViewModel> EmitterHandleViewModel, int32 ParameterIndex, int32 ParameterChannel) const
 {
 	UNiagaraStatelessEmitter* OwningStatelessEmitter = EmitterHandleViewModel->GetEmitterHandle()->GetStatelessEmitter();
-
 	FNiagaraEmitterInstance* ThisEmitterInstance = EmitterHandleViewModel->GetEmitterViewModel()->GetSimulation().Pin().Get();
-	TArray<UMaterial*> Materials = GetMaterialsFromEmitter(*OwningStatelessEmitter, ThisEmitterInstance);
-
-	for (UMaterial* Material : Materials)
-	{
-		if (Material == nullptr)
-		{
-			continue;
-		}
-
-		TArray<UMaterialExpression*> Expressions;
-		Material->GetAllExpressionsInMaterialAndFunctionsOfType<UMaterialExpression>(Expressions);
-		for (UMaterialExpression* Expression : Expressions)
-		{
-			const UMaterialExpressionDynamicParameter* DynamicParameterExpression = Cast<UMaterialExpressionDynamicParameter>(Expression);
-			if (DynamicParameterExpression == nullptr || DynamicParameterExpression->ParameterIndex != static_cast<uint32>(ParameterIndex))
-			{
-				continue;
-			}
-
-			return DynamicParameterExpression->DefaultValue.Component(ParameterChannel);
-		}
-	}
-
-	return TOptional<float>();
+	TArray<UMaterial*> Materials = NiagaraDynamicMaterialParameterUtils::GetMaterialsFromRenderers(OwningStatelessEmitter->GetRenderers(), ThisEmitterInstance);
+	return NiagaraDynamicMaterialParameterUtils::GetDynamicParameterDefaultValue(Materials, ParameterIndex, ParameterChannel);
 }
 
 void FNiagaraStackObjectPropertyCustomization_StatelessModule_DynamicMaterialParameters::GenerateContextMenuActions(UNiagaraStackPropertyRow* StackPropertyRow, FMenuBuilder& MenuBuilder) const
@@ -352,30 +328,6 @@ void FNiagaraStackObjectPropertyCustomization_StatelessModule_DynamicMaterialPar
 	}
 }
 
-TArray<UMaterial*> FNiagaraStackObjectPropertyCustomization_StatelessModule_DynamicMaterialParameters::GetMaterialsFromEmitter(
-	const UNiagaraStatelessEmitter& InEmitter, const FNiagaraEmitterInstance* InEmitterInstance) const
-{
-	TArray<UMaterial*> ResultMaterials;
-	for (UNiagaraRendererProperties* RenderProperties : InEmitter.GetRenderers())
-	{
-		TArray<UMaterialInterface*> UsedMaterialInteraces;
-		RenderProperties->GetUsedMaterials(InEmitterInstance, UsedMaterialInteraces);
-		for (UMaterialInterface* UsedMaterialInterface : UsedMaterialInteraces)
-		{
-			if (UsedMaterialInterface != nullptr)
-			{
-				UMaterial* UsedMaterial = UsedMaterialInterface->GetBaseMaterial();
-				if (UsedMaterial != nullptr)
-				{
-					ResultMaterials.AddUnique(UsedMaterial);
-					break;
-				}
-			}
-		}
-	}
-
-	return ResultMaterials;
-}
 
 void FNiagaraStackObjectPropertyCustomization_StatelessModule_DynamicMaterialParameters::GetParameterIndexAndChannel(TSharedRef<FDetailTreeNode> DetailTreeNode, int32& OutParameterIndex, int32& OutParameterChannel) const
 {
diff --git a/Engine/Plugins/FX/Niagara/Source/NiagaraEditor/Private/Customizations/Stack/NiagaraStackObjectPropertyCustomization_StatelessModule_DynamicMaterialParameters.h b/Engine/Plugins/FX/Niagara/Source/NiagaraEditor/Private/Customizations/Stack/NiagaraStackObjectPropertyCustomization_StatelessModule_DynamicMaterialParameters.h
--- a/Engine/Plugins/FX/Niagara/Source/NiagaraEditor/Private/Customizations/Stack/NiagaraStackObjectPropertyCustomization_StatelessModule_DynamicMaterialParameters.h
+++ b/Engine/Plugins/FX/Niagara/Source/NiagaraEditor/Private/Customizations/Stack/NiagaraStackObjectPropertyCustomization_StatelessModule_DynamicMaterialParameters.h
@@ -9,13 +9,10 @@
 #define UE_API NIAGARAEDITOR_API
 
 class FNiagaraEmitterHandleViewModel;
-class FNiagaraEmitterInstance;
-class UNiagaraStatelessEmitter;
 class UNiagaraStackPropertyRow;
-class UMaterial;
 class FDetailTreeNode;
 struct FExpressionInput;
-
+class FMenuBuilder;
 
 class FNiagaraStackObjectPropertyCustomization_StatelessModule_DynamicMaterialParameters
 	: public FNiagaraStackObjectPropertyCustomization
@@ -35,7 +32,6 @@ private:
 	UE_API TOptional<float> TryGetDefaultValueForDynamicMaterialParameter(TSharedPtr<FNiagaraEmitterHandleViewModel> EmitterHandleViewModel, int32 ParameterIndex, int32 ParameterChannel) const;
 
 	UE_API void GetChannelUsedBitMask(FExpressionInput* Input, TStaticBitArray<4>& ChannelUsedMask) const;
-	UE_API TArray<UMaterial*> GetMaterialsFromEmitter(const UNiagaraStatelessEmitter& InEmitter, const FNiagaraEmitterInstance* InEmitterInstance) const;
 	UE_API void GetParameterIndexAndChannel(TSharedRef<FDetailTreeNode> DetailTreeNode, int32& OutParameterIndex, int32& OutParameterChannel) const;
 
 private:
diff --git a/Engine/Plugins/FX/Niagara/Source/NiagaraEditor/Private/Utilities/NiagaraDynamicMaterialParameterUtils.cpp b/Engine/Plugins/FX/Niagara/Source/NiagaraEditor/Private/Utilities/NiagaraDynamicMaterialParameterUtils.cpp
new file mode 100644
--- /dev/null
+++ b/Engine/Plugins/FX/Niagara/Source/NiagaraEditor/Private/Utilities/NiagaraDynamicMaterialParameterUtils.cpp
@@ -0,0 +1,58 @@
+// Copyright Epic Games, Inc. All Rights Reserved.
+
+#include "NiagaraDynamicMaterialParameterUtils.h"
+#include "Materials/Material.h"
+#include "Materials/MaterialExpressionDynamicParameter.h"
+#include "NiagaraRendererProperties.h"
+#include "NiagaraEmitterInstance.h"
+
+TArray<UMaterial*> NiagaraDynamicMaterialParameterUtils::GetMaterialsFromRenderers(TArrayView<UNiagaraRendererProperties* const> Renderers, const FNiagaraEmitterInstance* EmitterInstance)
+{
+	TArray<UMaterial*> ResultMaterials;
+	for (UNiagaraRendererProperties* RenderProperties : Renderers)
+	{
+		TArray<UMaterialInterface*> UsedMaterialInterfaces;
+		RenderProperties->GetUsedMaterials(EmitterInstance, UsedMaterialInterfaces);
+		for (UMaterialInterface* UsedMaterialInterface : UsedMaterialInterfaces)
+		{
+			if (UsedMaterialInterface != nullptr)
+			{
+				UMaterial* UsedMaterial = UsedMaterialInterface->GetBaseMaterial();
+				if (UsedMaterial != nullptr)
+				{
+					ResultMaterials.AddUnique(UsedMaterial);
+					break;
+				}
+			}
+		}
+	}
+	return ResultMaterials;
+}
+
+TOptional<float> NiagaraDynamicMaterialParameterUtils::GetDynamicParameterDefaultValue(TArrayView<UMaterial* const> Materials, int32 InParameterIndex, int32 InParameterChannel)
+{
+	check(InParameterChannel >= 0 && InParameterChannel < 4);
+
+	for (UMaterial* Material : Materials)
+	{
+		if (Material == nullptr)
+		{
+			continue;
+		}
+
+		TArray<UMaterialExpression*> Expressions;
+		Material->GetAllExpressionsInMaterialAndFunctionsOfType<UMaterialExpression>(Expressions);
+		for (UMaterialExpression* Expression : Expressions)
+		{
+			const UMaterialExpressionDynamicParameter* DynamicParameterExpression = Cast<UMaterialExpressionDynamicParameter>(Expression);
+			if (DynamicParameterExpression == nullptr || DynamicParameterExpression->ParameterIndex != static_cast<uint32>(InParameterIndex))
+			{
+				continue;
+			}
+
+			return DynamicParameterExpression->DefaultValue.Component(InParameterChannel);
+		}
+	}
+
+	return TOptional<float>();
+}
diff --git a/Engine/Plugins/FX/Niagara/Source/NiagaraEditor/Private/Utilities/NiagaraDynamicMaterialParameterUtils.h b/Engine/Plugins/FX/Niagara/Source/NiagaraEditor/Private/Utilities/NiagaraDynamicMaterialParameterUtils.h
new file mode 100644
--- /dev/null
+++ b/Engine/Plugins/FX/Niagara/Source/NiagaraEditor/Private/Utilities/NiagaraDynamicMaterialParameterUtils.h
@@ -0,0 +1,25 @@
+// Copyright Epic Games, Inc. All Rights Reserved.
+
+#pragma once
+
+#include "CoreMinimal.h"
+
+class UMaterial;
+class UNiagaraRendererProperties;
+class FNiagaraEmitterInstance;
+
+namespace NiagaraDynamicMaterialParameterUtils
+{
+	/**
+	 * Collects unique base UMaterial* objects from an array of renderer properties.
+	 * Materials are resolved via UMaterialInterface::GetBaseMaterial and deduplicated.
+	 */
+	TArray<UMaterial*> GetMaterialsFromRenderers(TArrayView<UNiagaraRendererProperties* const> Renderers, const FNiagaraEmitterInstance* EmitterInstance);
+
+	/**
+	 * Searches the given materials for a UMaterialExpressionDynamicParameter node
+	 * whose ParameterIndex matches InParameterIndex, and returns the default value
+	 * for InParameterChannel (0=R, 1=G, 2=B, 3=A) from the node's DefaultValue.
+	 */
+	TOptional<float> GetDynamicParameterDefaultValue(TArrayView<UMaterial* const> Materials, int32 InParameterIndex, int32 InParameterChannel);
+}
diff --git a/Engine/Plugins/FX/Niagara/Source/NiagaraEditor/Private/ViewModels/Stack/NiagaraStackFunctionInput.cpp b/Engine/Plugins/FX/Niagara/Source/NiagaraEditor/Private/ViewModels/Stack/NiagaraStackFunctionInput.cpp
--- a/Engine/Plugins/FX/Niagara/Source/NiagaraEditor/Private/ViewModels/Stack/NiagaraStackFunctionInput.cpp
+++ b/Engine/Plugins/FX/Niagara/Source/NiagaraEditor/Private/ViewModels/Stack/NiagaraStackFunctionInput.cpp
@@ -1,6 +1,7 @@
 // Copyright Epic Games, Inc. All Rights Reserved.
 
 #include "ViewModels/Stack/NiagaraStackFunctionInput.h"
+#include "Utilities/NiagaraDynamicMaterialParameterUtils.h"
 
 #include "AssetRegistry/AssetRegistryModule.h"
 #include "Containers/StaticBitArray.h"
@@ -1437,49 +1438,12 @@ TOptional<float> UNiagaraStackFunctionInput::GetMaterialDynamicParamDefaultValue
 		FNiagaraEmitterInstance* ThisEmitterInstance = ThisEmitterViewModel->GetSimulation().Pin().Get();
 		if (ThisEmitterViewModel->GetEmitter().Emitter != nullptr)
 		{
-			for (UNiagaraRendererProperties* RenderProperties : ThisEmitterViewModel->GetEmitter().GetEmitterData()->GetRenderers())
-			{
-				TArray<UMaterialInterface*> UsedMaterialInterfaces;
-				RenderProperties->GetUsedMaterials(ThisEmitterInstance, UsedMaterialInterfaces);
-				for (UMaterialInterface* UsedMaterialInterface : UsedMaterialInterfaces)
-				{
-					if (UsedMaterialInterface != nullptr)
-					{
-						UMaterial* UsedMaterial = UsedMaterialInterface->GetBaseMaterial();
-						if (UsedMaterial != nullptr)
-						{
-							Materials.AddUnique(UsedMaterial);
-							break;
-						}
-					}
-				}
-			}
-		}
-	}
-
-	for (UMaterial* Material : Materials)
-	{
-		if (Material == nullptr)
-		{
-			continue;
-		}
-
-		TArray<UMaterialExpression*> Expressions;
-		Material->GetAllExpressionsInMaterialAndFunctionsOfType<UMaterialExpression>(Expressions);
-		for (UMaterialExpression* Expression : Expressions)
-		{
-			const UMaterialExpressionDynamicParameter* DynParam = Cast<UMaterialExpressionDynamicParameter>(Expression);
-			if (DynParam == nullptr || DynParam->ParameterIndex != static_cast<uint32>(ParamSlotIdx))
-			{
-				continue;
-			}
-
-			// DefaultValue is FLinearColor; ParamIdx maps to RGBA channels (0=R, 1=G, 2=B, 3=A)
-			return DynParam->DefaultValue.Component(ParamIdx);
+			Materials = NiagaraDynamicMaterialParameterUtils::GetMaterialsFromRenderers(
+				ThisEmitterViewModel->GetEmitter().GetEmitterData()->GetRenderers(), ThisEmitterInstance);
 		}
 	}
 
-	return TOptional<float>();
+	return NiagaraDynamicMaterialParameterUtils::GetDynamicParameterDefaultValue(Materials, ParamSlotIdx, ParamIdx);
 }
 
 void UNiagaraStackFunctionInput::RefreshValues()
