- Add stack errors if a mesh / sprite renderer info is reading a dangling renderer properties

#Synced-CL 51886665

diff --git a/Engine/Plugins/FX/Niagara/Source/Niagara/Private/NiagaraDataInterfaceMeshRendererInfo.cpp b/Engine/Plugins/FX/Niagara/Source/Niagara/Private/NiagaraDataInterfaceMeshRendererInfo.cpp
--- a/Engine/Plugins/FX/Niagara/Source/Niagara/Private/NiagaraDataInterfaceMeshRendererInfo.cpp
+++ b/Engine/Plugins/FX/Niagara/Source/Niagara/Private/NiagaraDataInterfaceMeshRendererInfo.cpp
@@ -4,9 +4,14 @@
 #include "NiagaraMeshRendererProperties.h"
 #include "NiagaraRenderer.h"
 #include "NiagaraShaderParametersBuilder.h"
+
 #include "Engine/StaticMesh.h"
 #include "RenderingThread.h"
 
+#if WITH_EDITOR
+#include "ScopedTransaction.h"
+#endif
+
 #include UE_INLINE_GENERATED_CPP_BY_NAME(NiagaraDataInterfaceMeshRendererInfo)
 
 #define LOCTEXT_NAMESPACE "NiagaraDataInterfaceMeshRendererInfo"
@@ -258,8 +263,7 @@ bool UNiagaraDataInterfaceMeshRendererInfo::UpgradeFunctionCall(FNiagaraFunction
 	return bWasChanged;
 }
 
-void UNiagaraDataInterfaceMeshRendererInfo::GetFeedback(UNiagaraSystem* InAsset, UNiagaraComponent* InComponent, TArray<FNiagaraDataInterfaceError>& OutErrors, TArray<FNiagaraDataInterfaceFeedback>& OutWarnings,
-	TArray<FNiagaraDataInterfaceFeedback>& OutInfo)
+void UNiagaraDataInterfaceMeshRendererInfo::GetFeedback(UNiagaraSystem* InAsset, UNiagaraComponent* InComponent, TArray<FNiagaraDataInterfaceError>& OutErrors, TArray<FNiagaraDataInterfaceFeedback>& OutWarnings, TArray<FNiagaraDataInterfaceFeedback>& OutInfo)
 {
 	if (MeshRenderer == nullptr)
 	{
@@ -270,23 +274,37 @@ void UNiagaraDataInterfaceMeshRendererInfo::GetFeedback(UNiagaraSystem* InAsset,
 		);
 		OutWarnings.Add(NoMeshRendererSelectedWarning);
 	}
-	else 
+	else if (!MeshRenderer->GetIsEnabled())
 	{
-		if (!MeshRenderer->GetIsEnabled())
-		{
-			FNiagaraDataInterfaceFeedback MeshRendererDisabledWarning(
-				LOCTEXT("RendererDisabledWarning", "The selected Mesh Renderer is disabled"),
-				LOCTEXT("RendererDisabledWarningSummary", "Mesh Renderer is disabled"),
-				FNiagaraDataInterfaceFix::CreateLambda(
-					[this]()
-					{
-						MeshRenderer->SetIsEnabled(true);
-						return true;
-					}
-				)
-			);
-			OutWarnings.Add(MeshRendererDisabledWarning);
-		}
+		FNiagaraDataInterfaceFeedback MeshRendererDisabledWarning(
+			LOCTEXT("RendererDisabledWarning", "The selected Mesh Renderer is disabled"),
+			LOCTEXT("RendererDisabledWarningSummary", "Mesh Renderer is disabled"),
+			FNiagaraDataInterfaceFix::CreateLambda(
+				[this]()
+				{
+					MeshRenderer->SetIsEnabled(true);
+					return true;
+				}
+			)
+		);
+		OutWarnings.Add(MeshRendererDisabledWarning);
+	}
+	// Check for the renderer being referenced but no longer referenced by it's parent emitter as that's likely a invalid dangling reference
+	else if ( !MeshRenderer->IsContainedInEmitterRendererArray() )
+	{
+		OutWarnings.Emplace(
+			LOCTEXT("OrphanedRendererWarning", "The selected Mesh Renderer is invalid"),
+			LOCTEXT("OrphanedRendererWarningSummary", "Mesh Renderer is orphaned and no longer valid"),
+			FNiagaraDataInterfaceFix::CreateLambda(
+				[this]()
+				{
+					const FScopedTransaction Transaction(LOCTEXT("ClearMeshRenderer", "Clear Mesh Renderer"));
+					Modify();
+					MeshRenderer = nullptr;
+					return true;
+				}
+			)
+		);
 	}
 }
 
diff --git a/Engine/Plugins/FX/Niagara/Source/Niagara/Private/NiagaraDataInterfaceSpriteRendererInfo.cpp b/Engine/Plugins/FX/Niagara/Source/Niagara/Private/NiagaraDataInterfaceSpriteRendererInfo.cpp
--- a/Engine/Plugins/FX/Niagara/Source/Niagara/Private/NiagaraDataInterfaceSpriteRendererInfo.cpp
+++ b/Engine/Plugins/FX/Niagara/Source/Niagara/Private/NiagaraDataInterfaceSpriteRendererInfo.cpp
@@ -5,6 +5,10 @@
 #include "NiagaraShaderParametersBuilder.h"
 #include "RenderingThread.h"
 
+#if WITH_EDITOR
+#include "ScopedTransaction.h"
+#endif
+
 #include UE_INLINE_GENERATED_CPP_BY_NAME(NiagaraDataInterfaceSpriteRendererInfo)
 
 #define LOCTEXT_NAMESPACE "NiagaraDataInterfaceSpriteRendererInfo"
@@ -282,6 +286,22 @@ void UNiagaraDataInterfaceSpriteRendererInfo::GetFeedback(UNiagaraSystem* InAsse
 			)
 		);
 	}
+	else if ( !SpriteRenderer->IsContainedInEmitterRendererArray() )
+	{
+		OutWarnings.Emplace(
+			LOCTEXT("OrphanedRendererWarning", "The selected Sprite Renderer is invalid"),
+			LOCTEXT("OrphanedRendererWarningSummary", "Sprite Renderer is orphaned and no longer valid"),
+			FNiagaraDataInterfaceFix::CreateLambda(
+				[this]()
+				{
+					const FScopedTransaction Transaction(LOCTEXT("ClearSpriteRenderer", "Clear Sprite Renderer"));
+					Modify();
+					SpriteRenderer = nullptr;
+					return true;
+				}
+			)
+		);
+	}
 }
 
 #endif // WITH_EDITOR
diff --git a/Engine/Plugins/FX/Niagara/Source/Niagara/Private/NiagaraRendererProperties.cpp b/Engine/Plugins/FX/Niagara/Source/Niagara/Private/NiagaraRendererProperties.cpp
--- a/Engine/Plugins/FX/Niagara/Source/Niagara/Private/NiagaraRendererProperties.cpp
+++ b/Engine/Plugins/FX/Niagara/Source/Niagara/Private/NiagaraRendererProperties.cpp
@@ -1400,6 +1400,28 @@ FVersionedNiagaraEmitterBase UNiagaraRendererProperties::GetOuterEmitterBase() c
 	return FVersionedNiagaraEmitterBase();
 }
 
+#if WITH_EDITORONLY_DATA
+bool UNiagaraRendererProperties::IsContainedInEmitterRendererArray() const
+{
+	FVersionedNiagaraEmitterBase VersionedEmitter = GetOuterEmitterBase();
+	if (VersionedEmitter.Emitter)
+	{
+		if (UNiagaraEmitter* Emitter = Cast<UNiagaraEmitter>(VersionedEmitter.Emitter))
+		{
+			if (FVersionedNiagaraEmitterData* EmitterData = Emitter->GetEmitterData(VersionedEmitter.Version))
+			{
+				return EmitterData->GetRenderers().Contains(this);
+			}
+		}
+		else if (UNiagaraStatelessEmitter* StatelessEmitter = Cast<UNiagaraStatelessEmitter>(VersionedEmitter.Emitter))
+		{
+			return StatelessEmitter->GetRenderers().Contains(this);
+		}
+	}
+	return false;
+}
+#endif
+
 bool UNiagaraRendererProperties::NeedsPreciseMotionVectors() const
 {
 	if (MotionVectorSetting == ENiagaraRendererMotionVectorSetting::AutoDetect)
diff --git a/Engine/Plugins/FX/Niagara/Source/Niagara/Public/NiagaraRendererProperties.h b/Engine/Plugins/FX/Niagara/Source/Niagara/Public/NiagaraRendererProperties.h
--- a/Engine/Plugins/FX/Niagara/Source/Niagara/Public/NiagaraRendererProperties.h
+++ b/Engine/Plugins/FX/Niagara/Source/Niagara/Public/NiagaraRendererProperties.h
@@ -457,6 +457,12 @@ public:
 	[[nodiscard]] NIAGARA_API FVersionedNiagaraEmitter GetOuterEmitter() const;
 	[[nodiscard]] NIAGARA_API FVersionedNiagaraEmitterBase GetOuterEmitterBase() const;
 
+#if WITH_EDITORONLY_DATA
+	// Checks to see if th renderer is contained with it's owning emitters renderer array
+	// If not this is likely an orphaned or invalid renderer
+	[[nodiscard]] bool IsContainedInEmitterRendererArray() const;
+#endif
+
 	/** Platforms on which this renderer is enabled. */
 	UPROPERTY(EditAnywhere, Category = "Scalability", meta=(DisplayInScalabilityContext))
 	FNiagaraPlatformSet Platforms;
