- Fix CIS issue with WITH_PARTICLE_PERF_STATS defined to 0

#Synced-CL 51894006

diff --git a/Engine/Plugins/FX/Niagara/Source/Niagara/Private/NiagaraTrace.cpp b/Engine/Plugins/FX/Niagara/Source/Niagara/Private/NiagaraTrace.cpp
--- a/Engine/Plugins/FX/Niagara/Source/Niagara/Private/NiagaraTrace.cpp
+++ b/Engine/Plugins/FX/Niagara/Source/Niagara/Private/NiagaraTrace.cpp
@@ -83,6 +83,7 @@ namespace NiagaraInsightsPrivate
 		UE_TRACE_EVENT_FIELD(bool,					bExternalCompletion)
 	UE_TRACE_EVENT_END()
 
+#if WITH_PARTICLE_PERF_STATS
 	class FNiagaraInsightsStatsListener : public FParticlePerfStatsListener
 	{
 	public:
@@ -158,10 +159,12 @@ namespace NiagaraInsightsPrivate
 	};
 
 	TSharedPtr<FNiagaraInsightsStatsListener> GStatsListener;
+#endif //WITH_PARTICLE_PERF_STATS
 }
 
 void FNiagaraInsights::Update()
 {
+#if WITH_PARTICLE_PERF_STATS
 	if (IsChannelActive())
 	{
 		if (NiagaraInsightsPrivate::GStatsListener == nullptr)
@@ -178,6 +181,7 @@ void FNiagaraInsights::Update()
 			NiagaraInsightsPrivate::GStatsListener = nullptr;
 		}
 	}
+#endif //WITH_PARTICLE_PERF_STATS
 }
 
 bool FNiagaraInsights::IsChannelActive()
