- Fix for functions being using in Custom HLSL where there are no parameters failed to compile
- Previously a , was always added resulting in invalid hlsl

#Synced-CL 51808504

diff --git a/Engine/Plugins/FX/Niagara/Source/NiagaraEditor/Private/NiagaraHlslTranslator.cpp b/Engine/Plugins/FX/Niagara/Source/NiagaraEditor/Private/NiagaraHlslTranslator.cpp
--- a/Engine/Plugins/FX/Niagara/Source/NiagaraEditor/Private/NiagaraHlslTranslator.cpp
+++ b/Engine/Plugins/FX/Niagara/Source/NiagaraEditor/Private/NiagaraHlslTranslator.cpp
@@ -9061,7 +9061,22 @@ void TNiagaraHlslTranslator<GraphBridge>::ProcessCustomHlsl(const FString& InCus
 						// the tokens end after the parenthesis, we'll be inserting at the end of the array.
 						Tokens.Insert(LexToString(Info.UserPtrIdx), TokenIndex++);
 
-						if (Sig.Inputs.Num() > 1 || Sig.Outputs.Num() > 0)
+						// We only insert a comma if there are valid parameters, i.e. non-class, non-parameter map arugments
+						// The first input argument is the instance data so we don't evaluate that
+						// For example a impure function may have the instance data / exec pin / data interface as inputs, then an exec pin for output, none of which will translate to real parameters
+						const auto IsRealHlslArg = [](const FNiagaraTypeDefinition& Type) { return !Type.GetClass() && Type != FNiagaraTypeDefinition::GetParameterMapDef(); };
+
+						bool bHasRemainingArgs = false;
+						for (int32 j = 1; j < Sig.Inputs.Num() && !bHasRemainingArgs; ++j)
+						{
+							bHasRemainingArgs = IsRealHlslArg(Sig.Inputs[j].GetType());
+						}
+						for (int32 j = 0; j < Sig.Outputs.Num() && !bHasRemainingArgs; ++j)
+						{
+							bHasRemainingArgs = IsRealHlslArg(Sig.Outputs[j].GetType());
+						}
+
+						if (bHasRemainingArgs)
 						{
 							// If there are other arguments, insert a comma and a space. These are separators, so they need to be different tokens.
 							Tokens.Insert(TEXT(","), TokenIndex++);
