- Allow skipping syncing the RT / updating static draw lists when Niagara modifies a MIC.
- Has been seen to remove several 100+ms stalls when loading levels

#Synced-CL 51894006

diff --git a/Engine/Plugins/FX/Niagara/Source/Niagara/Private/NiagaraRendererProperties.cpp b/Engine/Plugins/FX/Niagara/Source/Niagara/Private/NiagaraRendererProperties.cpp
--- a/Engine/Plugins/FX/Niagara/Source/Niagara/Private/NiagaraRendererProperties.cpp
+++ b/Engine/Plugins/FX/Niagara/Source/Niagara/Private/NiagaraRendererProperties.cpp
@@ -632,17 +632,17 @@ bool UNiagaraRendererProperties::UpdateMaterialStaticParameters(const FNiagaraRe
 	return false;
 }
 
-void UNiagaraRendererProperties::UpdateMaterialParametersMIC(const FNiagaraRendererMaterialParameters& MaterialParameters, TObjectPtr<UMaterialInterface>& InOutMaterial, TObjectPtr<UMaterialInstanceConstant>& InOutMIC)
+void UNiagaraRendererProperties::UpdateMaterialParametersMIC(const FNiagaraRendererMaterialParameters& MaterialParameters, TObjectPtr<UMaterialInterface>& InOutMaterial, TObjectPtr<UMaterialInstanceConstant>& InOutMIC, FMaterialUpdateContext* OptionalMaterialUpdateContext)
 {
 	TArray<UMaterialInterface*> LocalMaterials = { InOutMaterial };
 	TArray<TObjectPtr<UMaterialInstanceConstant>> LocalMICs = { InOutMIC };
 
-	UpdateMaterialParametersMIC(MaterialParameters, LocalMaterials, LocalMICs);
+	UpdateMaterialParametersMIC(MaterialParameters, LocalMaterials, LocalMICs, OptionalMaterialUpdateContext);
 
 	InOutMIC = LocalMICs.IsEmpty() ? nullptr : LocalMICs[0];
 }
 
-void UNiagaraRendererProperties::UpdateMaterialParametersMIC(const FNiagaraRendererMaterialParameters& MaterialParameters, TArrayView<UMaterialInterface*> Materials, TArray<TObjectPtr<UMaterialInstanceConstant>>& InOutMICs)
+void UNiagaraRendererProperties::UpdateMaterialParametersMIC(const FNiagaraRendererMaterialParameters& MaterialParameters, TArrayView<UMaterialInterface*> Materials, TArray<TObjectPtr<UMaterialInstanceConstant>>& InOutMICs, FMaterialUpdateContext* OptionalMaterialUpdateContext)
 {
 	// Create a pool of unique MICs that we can potentially reuse
 	TArray<UMaterialInstanceConstant*> MICPool;
@@ -670,6 +670,8 @@ void UNiagaraRendererProperties::UpdateMaterialParametersMIC(const FNiagaraRende
 		return;
 	}
 
+	FMaterialUpdateContext DefaultMaterialUpdateContext(FMaterialUpdateContext::EOptions::SkipRecacheStaticPermutations | FMaterialUpdateContext::EOptions::SkipUpdateStaticDrawLists);
+
 	// Loop over each material to see if we need to generate a MIC for it
 	for ( int i=0; i < Materials.Num(); ++i )
 	{
@@ -724,7 +726,7 @@ void UNiagaraRendererProperties::UpdateMaterialParametersMIC(const FNiagaraRende
 				MIC->SetParentEditorOnly(Material);
 			}
 
-			MIC->UpdateStaticPermutation(MaterialParameterSet);
+			MIC->UpdateStaticPermutation(MaterialParameterSet, OptionalMaterialUpdateContext ? OptionalMaterialUpdateContext : &DefaultMaterialUpdateContext);
 			MIC->PostEditChange();
 
 			InOutMICs.SetNum(i + 1);
diff --git a/Engine/Plugins/FX/Niagara/Source/Niagara/Public/NiagaraRendererProperties.h b/Engine/Plugins/FX/Niagara/Source/Niagara/Public/NiagaraRendererProperties.h
--- a/Engine/Plugins/FX/Niagara/Source/Niagara/Public/NiagaraRendererProperties.h
+++ b/Engine/Plugins/FX/Niagara/Source/Niagara/Public/NiagaraRendererProperties.h
@@ -19,6 +19,7 @@ class FVertexFactoryType;
 class UMaterial;
 class UMaterialInterface;
 class UMaterialInstanceConstant;
+class FMaterialUpdateContext;
 class UTexture;
 class FNiagaraEmitterInstance;
 class SWidget;
@@ -515,9 +516,13 @@ protected:
 	UE_DEPRECATED(5.5, "This helper function has been deprecated.  Use BuildMaterialStaticParameterSet instead.")
 	NIAGARA_API bool UpdateMaterialStaticParameters(const FNiagaraRendererMaterialParameters& MaterialParameters, UMaterialInstanceConstant* MIC);
 
-	/** Utility function to updates MICs. */
-	NIAGARA_API void UpdateMaterialParametersMIC(const FNiagaraRendererMaterialParameters& MaterialParameters, TObjectPtr<UMaterialInterface>& InOutMaterial, TObjectPtr<UMaterialInstanceConstant>& InOutMIC);
-	NIAGARA_API void UpdateMaterialParametersMIC(const FNiagaraRendererMaterialParameters& MaterialParameters, TArrayView<UMaterialInterface*> Materials, TArray<TObjectPtr<UMaterialInstanceConstant>>& InOutMICs);
+	/**
+	Utility function to updates MICs.
+	If your renderer uses static draw lists and you are unsure if the data could still be inside the renderer make sure to pass a valid material update context.
+	The default is to skip updating anything as Niagara system should be recreated as part of the caller to update the MICs.
+	*/
+	NIAGARA_API void UpdateMaterialParametersMIC(const FNiagaraRendererMaterialParameters& MaterialParameters, TObjectPtr<UMaterialInterface>& InOutMaterial, TObjectPtr<UMaterialInstanceConstant>& InOutMIC, FMaterialUpdateContext* OptionalMaterialUpdateContext = nullptr);
+	NIAGARA_API void UpdateMaterialParametersMIC(const FNiagaraRendererMaterialParameters& MaterialParameters, TArrayView<UMaterialInterface*> Materials, TArray<TObjectPtr<UMaterialInstanceConstant>>& InOutMICs, FMaterialUpdateContext* OptionalMaterialUpdateContext = nullptr);
 
 	NIAGARA_API int32 GetDynamicParameterChannelMask(const FVersionedNiagaraEmitterData* EmitterData, FName BindingName, int32 DefaultChannelMask) const;
 	NIAGARA_API int32 GetDynamicParameterCombinedChannelMask(FName Parameter0Name, FName Parameter1Name, FName Parameter2Name, FName Parameter3Name) const;
diff --git a/Engine/Source/Runtime/Engine/Private/Materials/MaterialShared.cpp b/Engine/Source/Runtime/Engine/Private/Materials/MaterialShared.cpp
--- a/Engine/Source/Runtime/Engine/Private/Materials/MaterialShared.cpp
+++ b/Engine/Source/Runtime/Engine/Private/Materials/MaterialShared.cpp
@@ -4990,6 +4990,7 @@ FMaterialUpdateContext::FMaterialUpdateContext(uint32 Options, EShaderPlatform I
 
 	bSyncWithRenderingThread = (Options & EOptions::SyncWithRenderingThread) != 0;
 	bSkipRecacheStaticPermutations = (Options & EOptions::SkipRecacheStaticPermutations) != 0;
+	bSkipUpdateStaticDrawLists = (Options & EOptions::SkipUpdateStaticDrawLists) != 0;
 
 	if (bReregisterComponents)
 	{
@@ -5059,7 +5060,7 @@ FMaterialUpdateContext::~FMaterialUpdateContext()
 	TArray<const FMaterial*> MaterialResourcesToUpdate;
 	TArray<UMaterialInstance*> InstancesToUpdate;
 
-	bool bUpdateStaticDrawLists = !ComponentReregisterContext && !ComponentRecreateRenderStateContext && FApp::CanEverRender();
+	const bool bUpdateStaticDrawLists = !bSkipUpdateStaticDrawLists && !ComponentReregisterContext && !ComponentRecreateRenderStateContext && FApp::CanEverRender();
 
 	// If static draw lists must be updated, gather material resources from all updated materials.
 	if (bUpdateStaticDrawLists)
diff --git a/Engine/Source/Runtime/Engine/Public/MaterialShared.h b/Engine/Source/Runtime/Engine/Public/MaterialShared.h
--- a/Engine/Source/Runtime/Engine/Public/MaterialShared.h
+++ b/Engine/Source/Runtime/Engine/Public/MaterialShared.h
@@ -3475,6 +3475,8 @@ class FMaterialUpdateContext
 	bool bSyncWithRenderingThread;
 	/** True if the SkipRecacheStaticPermutations option was specified. */
 	bool bSkipRecacheStaticPermutations;
+	/** True if the SkipUpdateStaticDrawLists option was specified. */
+	bool bSkipUpdateStaticDrawLists;
 
 public:
 
@@ -3490,11 +3492,17 @@ public:
 			 * material exposed to the rendering thread. You may omit this flag if
 			 * you have already flushed rendering commands.
 			 */
-			SyncWithRenderingThread = 0x2,
+			SyncWithRenderingThread = 0x02,
 			/* Recreates only the render state for all components (mutually exclusive with ReregisterComponents) */
-			RecreateRenderStates = 0x4,
+			RecreateRenderStates = 0x04,
 			/** Skip recache of material instance permutations. Only safe to use if you know that no shader code or static switches are changing. */
 			SkipRecacheStaticPermutations = 0x08,
+			/**
+			 * Skip updating static draw lists.
+			 * Has no meaning if ReregisterComponents or RecreateRenderStates is set as it is already skipped.
+			 * Only safe to use this if you know that no static draw lists could be using the material.
+			*/
+			SkipUpdateStaticDrawLists = 0x10,
 			/** Default options: Recreate render state, sync with rendering thread. */
 			Default = RecreateRenderStates | SyncWithRenderingThread,
 		};
