Niagara AI Assistant improvements.

Extended the System tool to provide access to much more of Niagara's internals.
Inlcuding allow reading and writing of data interfaces, dynamic inputs and a range of other data.
Split system tool into a schema, topology and data layer to reduce the overall token count being used and expose Niagara data in a more structured and understandabl way for the EDA.
Improved error reporting and failure cases. Now using RaiseScriptError rather than returning custom error messages.

A couple of additions to UNiagaraStackFunctionInput
* SetDataInterfaceValueExternal - Needed a way to setup a DI on an input and get it back immediately for initialization from EDA data..
* GetVisibleConditionVariable - Exposing visible and editable state for the EDA.

Misc:

Exposed a few internal properties for BP edit to expose to the EDA. Ideally we'd have a more restricted flag here.
A few minor fixes and extensions to Niagara internals.
A few minor comment fixes that were throwing off the EDA in some places.
Moved asset builder from ET to Niagara proper.
Moved ToolsetLibrary header to public as this is meant to be a utility function library for toolsets to use.

#Synced-CL 52000591

diff --git a/Engine/Plugins/FX/Niagara/Niagara.uplugin b/Engine/Plugins/FX/Niagara/Niagara.uplugin
--- a/Engine/Plugins/FX/Niagara/Niagara.uplugin
+++ b/Engine/Plugins/FX/Niagara/Niagara.uplugin
@@ -18,49 +18,52 @@
 		{
 			"Name": "PythonScriptPlugin",
 			"Enabled": true
+		},
+		{
+			"Name": "ToolsetRegistry",
+			"Enabled": true
 		}
 	],
-	"Modules" :
-	[
+	"Modules": [
 		{
-		  "Name": "NiagaraCore",
-		  "Type": "Runtime",
-		  "LoadingPhase": "PreDefault"
+			"Name": "NiagaraCore",
+			"Type": "Runtime",
+			"LoadingPhase": "PreDefault"
 		},
 		{
-		  "Name": "Niagara",
-		  "Type": "Runtime",
-		  "LoadingPhase": "PreDefault"
+			"Name": "Niagara",
+			"Type": "Runtime",
+			"LoadingPhase": "PreDefault"
 		},
 		{
-			"Name" : "NiagaraAnimNotifies",
-			"Type" : "Runtime",
-			"LoadingPhase" : "PreDefault"
+			"Name": "NiagaraAnimNotifies",
+			"Type": "Runtime",
+			"LoadingPhase": "PreDefault"
 		},
 		{
-			"Name" : "NiagaraShader",
-			"Type" : "Runtime",
-			"LoadingPhase" : "PostConfigInit"
+			"Name": "NiagaraShader",
+			"Type": "Runtime",
+			"LoadingPhase": "PostConfigInit"
 		},
 		{
-			"Name" : "NiagaraVertexFactories",
-			"Type" : "Runtime",
-			"LoadingPhase" : "PostConfigInit"
+			"Name": "NiagaraVertexFactories",
+			"Type": "Runtime",
+			"LoadingPhase": "PostConfigInit"
 		},
 		{
-			"Name" : "NiagaraBlueprintNodes",
-			"Type" : "UncookedOnly",
-			"LoadingPhase" : "Default"
+			"Name": "NiagaraBlueprintNodes",
+			"Type": "UncookedOnly",
+			"LoadingPhase": "Default"
 		},
 		{
-			"Name" : "NiagaraEditor",
-			"Type" : "Editor",
-			"LoadingPhase" : "Default"
+			"Name": "NiagaraEditor",
+			"Type": "Editor",
+			"LoadingPhase": "Default"
 		},
 		{
-			"Name" : "NiagaraEditorWidgets",
-			"Type" : "Editor",
-			"LoadingPhase" : "Default"
+			"Name": "NiagaraEditorWidgets",
+			"Type": "Editor",
+			"LoadingPhase": "Default"
 		}
 	] 
 }
diff --git a/Engine/Plugins/FX/Niagara/Source/Niagara/Classes/NiagaraDataInterfaceSpriteRendererInfo.h b/Engine/Plugins/FX/Niagara/Source/Niagara/Classes/NiagaraDataInterfaceSpriteRendererInfo.h
--- a/Engine/Plugins/FX/Niagara/Source/Niagara/Classes/NiagaraDataInterfaceSpriteRendererInfo.h
+++ b/Engine/Plugins/FX/Niagara/Source/Niagara/Classes/NiagaraDataInterfaceSpriteRendererInfo.h
@@ -60,7 +60,7 @@ protected:
 	NIAGARA_API void VMGetFacingMode(FVectorVMExternalFunctionContext& Context);
 	NIAGARA_API void VMGetSubUVDetails(FVectorVMExternalFunctionContext& Context);
 
-	/** The name of the sprite renderer */
+	/** The sprite renderer to use. */
 	UPROPERTY(EditAnywhere, Category = "Source")
 	TObjectPtr<UNiagaraSpriteRendererProperties> SpriteRenderer;
 
diff --git a/Engine/Plugins/FX/Niagara/Source/Niagara/Classes/NiagaraEmitter.h b/Engine/Plugins/FX/Niagara/Source/Niagara/Classes/NiagaraEmitter.h
--- a/Engine/Plugins/FX/Niagara/Source/Niagara/Classes/NiagaraEmitter.h
+++ b/Engine/Plugins/FX/Niagara/Source/Niagara/Classes/NiagaraEmitter.h
@@ -419,6 +419,7 @@ struct FVersionedNiagaraEmitterData
 	NIAGARA_API bool UsesScript(const UNiagaraScript* Script) const;
 	NIAGARA_API bool IsReadyToRun() const;
 	const TArray<UNiagaraRendererProperties*>& GetRenderers() const { return RendererProperties; }
+	UNiagaraRendererProperties* GetRenderer(int32 RendererIndex)const { return RendererProperties.IsValidIndex(RendererIndex) ? RendererProperties[RendererIndex] : nullptr; }
 	NIAGARA_API void GetScripts(TArray<UNiagaraScript*>& OutScripts, bool bCompilableOnly = true, bool bEnabledOnly = false) const;
 	NIAGARA_API UNiagaraScript* GetScript(ENiagaraScriptUsage Usage, FGuid UsageId) const;
 	UNiagaraScript* GetGPUComputeScript() { return GPUComputeScript; }
diff --git a/Engine/Plugins/FX/Niagara/Source/Niagara/Private/NiagaraParameterStore.cpp b/Engine/Plugins/FX/Niagara/Source/Niagara/Private/NiagaraParameterStore.cpp
--- a/Engine/Plugins/FX/Niagara/Source/Niagara/Private/NiagaraParameterStore.cpp
+++ b/Engine/Plugins/FX/Niagara/Source/Niagara/Private/NiagaraParameterStore.cpp
@@ -1526,7 +1526,7 @@ void FNiagaraParameterStore::GetParameter_InternalUseOnly(const FNiagaraVariable
 		if (const FVector* PosData = GetPositionParameterValue(Key.GetName()))
 		{
 			PosValue = *PosData;
-			Value.SetBytesValue(Key.GetType(), PosValue);
+			Value.SetBytesValue(FNiagaraTypeHelper::GetVectorDef(), PosValue);
 		}
 	}
 	else
diff --git a/Engine/Plugins/FX/Niagara/Source/Niagara/Public/NiagaraCommon.h b/Engine/Plugins/FX/Niagara/Source/Niagara/Public/NiagaraCommon.h
--- a/Engine/Plugins/FX/Niagara/Source/Niagara/Public/NiagaraCommon.h
+++ b/Engine/Plugins/FX/Niagara/Source/Niagara/Public/NiagaraCommon.h
@@ -1689,7 +1689,7 @@ namespace FNiagaraUtilities
 	NIAGARA_API TOptional<ETextureRenderTargetFormat> BufferFormatToRenderTargetFormat(ENiagaraGpuBufferFormat NiagaraFormat, EPixelFormatCapabilities RequiredCapabilities);
 };
 
-USTRUCT()
+USTRUCT(BlueprintInternalUseOnly)
 struct FNiagaraUserParameterBinding
 {
 	GENERATED_USTRUCT_BODY()
@@ -1698,6 +1698,7 @@ struct FNiagaraUserParameterBinding
 
 	NIAGARA_API FNiagaraUserParameterBinding(const FNiagaraTypeDefinition& InMaterialDef);
 
+	/** Link to a User Parameter to control a value within a Niagara System. */
 	UPROPERTY(EditAnywhere, Category = "User Parameter")
 	FNiagaraVariable Parameter;
 
diff --git a/Engine/Plugins/FX/Niagara/Source/Niagara/Public/NiagaraRendererProperties.h b/Engine/Plugins/FX/Niagara/Source/Niagara/Public/NiagaraRendererProperties.h
--- a/Engine/Plugins/FX/Niagara/Source/Niagara/Public/NiagaraRendererProperties.h
+++ b/Engine/Plugins/FX/Niagara/Source/Niagara/Public/NiagaraRendererProperties.h
@@ -477,7 +477,7 @@ public:
 	UPROPERTY(EditAnywhere, Category = "Rendering")
 	ENiagaraRendererMotionVectorSetting MotionVectorSetting = ENiagaraRendererMotionVectorSetting::AutoDetect;
 
-	UPROPERTY()
+	UPROPERTY(EditAnywhere, Category = "Rendering")
 	bool bIsEnabled;
 
 	UPROPERTY(EditAnywhere, Category = "Scalability")
diff --git a/Engine/Plugins/FX/Niagara/Source/Niagara/Public/NiagaraTypes.h b/Engine/Plugins/FX/Niagara/Source/Niagara/Public/NiagaraTypes.h
--- a/Engine/Plugins/FX/Niagara/Source/Niagara/Public/NiagaraTypes.h
+++ b/Engine/Plugins/FX/Niagara/Source/Niagara/Public/NiagaraTypes.h
@@ -660,7 +660,7 @@ enum class ENiagaraOrientationAxis : uint32
 	ZAxis UMETA(DisplayName = "Z Axis"),
 };
 
-USTRUCT()
+USTRUCT(BlueprintInternalUseOnly)
 struct FNiagaraTypeDefinition
 {
 	GENERATED_USTRUCT_BODY()
diff --git a/Engine/Plugins/FX/Niagara/Source/Niagara/Public/NiagaraVariableMetaData.h b/Engine/Plugins/FX/Niagara/Source/Niagara/Public/NiagaraVariableMetaData.h
--- a/Engine/Plugins/FX/Niagara/Source/Niagara/Public/NiagaraVariableMetaData.h
+++ b/Engine/Plugins/FX/Niagara/Source/Niagara/Public/NiagaraVariableMetaData.h
@@ -11,7 +11,7 @@
 class IPropertyHandle;
 
 /** Defines options for conditionally editing and showing script inputs in the UI. */
-USTRUCT()
+USTRUCT(BlueprintInternalUseOnly)
 struct FNiagaraInputConditionMetadata
 {
 	GENERATED_USTRUCT_BODY()
diff --git a/Engine/Plugins/FX/Niagara/Source/NiagaraEditor/Internal/AssetBuildUtilities/NiagaraAssetBuilder.h b/Engine/Plugins/FX/Niagara/Source/NiagaraEditor/Internal/AssetBuildUtilities/NiagaraAssetBuilder.h
new file mode 100644
--- /dev/null
+++ b/Engine/Plugins/FX/Niagara/Source/NiagaraEditor/Internal/AssetBuildUtilities/NiagaraAssetBuilder.h
@@ -0,0 +1,100 @@
+// Copyright Epic Games, Inc. All Rights Reserved.
+
+#pragma once
+
+#if WITH_EDITOR
+
+#include "NiagaraTypes.h"
+
+class UNiagaraEmitter;
+class UNiagaraSystem;
+
+/** Niagara Emitter builder which creates a new Emitter asset. */
+class FNiagaraEmitterAssetBuilder
+{
+public:
+	explicit FNiagaraEmitterAssetBuilder(const FString& InPackagePath, const FString& InAssetName)
+		: PackagePath(InPackagePath), AssetName(InAssetName) {}
+
+	/**
+	 * Provide an Emitter as the source.
+	 *
+	 * @param Emitter - NiagaraEmitter to be as the base for the new Emitter
+	 *
+	 * @return reference to this
+	 */
+	FNiagaraEmitterAssetBuilder& WithEmitterToCopy(UNiagaraEmitter* Emitter);
+
+	/**
+	 * Build the NiagaraEmitter.
+	 *
+	 * @param OutError - String for any errors during creation
+	 *
+	 * @return the newly created NiagaraEmitter
+	 * 
+	 * @note Copied over from `UNiagaraEmitterFactoryNew` to mimic asset creation as close as possible
+	 */
+	UNiagaraEmitter* BuildEmitter(FString& OutError);
+
+private:
+	/** Package path of the new NiagaraEmitter asset. */
+	FString PackagePath;
+
+	/** Name of the Niagara Emitter. */
+	FString AssetName;
+
+	/** Source Emitter to use for the new Emitter. */
+	UNiagaraEmitter* EmitterToCopy{ nullptr };
+};
+
+/** Niagara System builder which creates a new System asset. */
+class FNiagaraSystemAssetBuilder
+{
+public:
+	explicit FNiagaraSystemAssetBuilder(const FString& InPackagePath, const FString& InAssetName)
+		: PackagePath(InPackagePath), AssetName(InAssetName) {}
+
+	/**
+	 * Provide a System as the source.
+	 *
+	 * @param System - NiagaraSystem to be as the base for the new System
+	 *
+	 * @return reference to this
+	 */
+	FNiagaraSystemAssetBuilder& WithSystemToCopy(UNiagaraSystem* System);
+
+	/**
+	 * Add an Emitter to the new System.
+	 *
+	 * @param Emitter - NiagaraEmitter to be added
+	 *
+	 * @return reference to this
+	 */
+	FNiagaraSystemAssetBuilder& WithEmitter(FVersionedNiagaraEmitter&& Emitter);
+
+	/**
+	 * Build the NiagaraSystem.
+	 *
+	 * @param OutError - String for any errors during creation
+	 *
+	 * @return the newly created NiagaraSystem
+	 *
+	 * @note Copied over from `UNiagaraSystemFactoryNew` to mimic asset creation as close as possible
+	 */
+	UNiagaraSystem* BuildSystem(FString& OutError);
+
+private:
+	/** Package path of the new NiagaraSystem asset. */
+	FString PackagePath;
+
+	/** Name of the Niagara System. */
+	FString AssetName;
+
+	/** Source System to use for the new System. */
+	UNiagaraSystem* SystemToCopy{ nullptr };
+
+	/** Array of Emitters to be added to the new System. */
+	TArray<FVersionedNiagaraEmitter> EmittersToAddToNewSystem;
+};
+
+#endif //WITH_EDITOR
\ No newline at end of file
diff --git a/Engine/Plugins/FX/Niagara/Source/NiagaraEditor/NiagaraEditor.Build.cs b/Engine/Plugins/FX/Niagara/Source/NiagaraEditor/NiagaraEditor.Build.cs
--- a/Engine/Plugins/FX/Niagara/Source/NiagaraEditor/NiagaraEditor.Build.cs
+++ b/Engine/Plugins/FX/Niagara/Source/NiagaraEditor/NiagaraEditor.Build.cs
@@ -10,15 +10,15 @@ public class NiagaraEditor : ModuleRules
 		PrivateDependencyModuleNames.AddRange(
 			new string[] {
                 "RHI",
-                "Core", 
-				"CoreUObject", 
+                "Core",
+				"CoreUObject",
 				"CurveEditor",
 				"DerivedDataCache",
 				"ApplicationCore",
                 "InputCore",
 				"RenderCore",
 				"Json",
-				"Slate", 
+				"Slate",
 				"SlateCore",
 				"Kismet",
                 "EditorStyle",
@@ -56,6 +56,7 @@ public class NiagaraEditor : ModuleRules
 				"EngineSettings",
 				"EditorConfig",
 				"MaterialEditor",
+				"JsonUtilities",
 			}
 		);
 
@@ -90,9 +91,10 @@ public class NiagaraEditor : ModuleRules
                 "AppFramework",
 				"ToolWidgets",
 				"Sequencer",
-				"DataHierarchyEditor", 
+				"DataHierarchyEditor",
 				"MediaAssets",
-				"UserAssetTagsEditor"
+				"UserAssetTagsEditor",
+				"ToolsetRegistry",
             }
         );
 
diff --git a/Engine/Plugins/FX/Niagara/Source/NiagaraEditor/Private/AssetBuildUtilities/NiagaraAssetBuilder.cpp b/Engine/Plugins/FX/Niagara/Source/NiagaraEditor/Private/AssetBuildUtilities/NiagaraAssetBuilder.cpp
new file mode 100644
--- /dev/null
+++ b/Engine/Plugins/FX/Niagara/Source/NiagaraEditor/Private/AssetBuildUtilities/NiagaraAssetBuilder.cpp
@@ -0,0 +1,207 @@
+// Copyright Epic Games, Inc. All Rights Reserved.
+
+#include "AssetBuildUtilities/NiagaraAssetBuilder.h"
+
+#if WITH_EDITOR
+
+#include "AssetToolsModule.h"
+#include "NiagaraEditorUtilities.h"
+#include "NiagaraEmitter.h"
+#include "NiagaraEmitterEditorData.h"
+#include "NiagaraEmitterFactoryNew.h"
+#include "NiagaraSettings.h"
+#include "NiagaraSystem.h"
+#include "NiagaraSystemFactoryNew.h"
+#include "ObjectTools.h"
+
+namespace {
+
+UPackage* CreateNiagaraPackage(const FString& PackageName)
+{
+	UPackage* Package = CreatePackage(*(PackageName));
+	if (!IsValid(Package))
+	{
+		return nullptr;
+	}
+
+	EObjectFlags Flags = RF_Public | RF_Standalone | RF_Transactional;
+	Package->SetFlags(Flags);
+
+	return Package;
+}
+
+} //anonymous
+
+FNiagaraEmitterAssetBuilder& FNiagaraEmitterAssetBuilder::WithEmitterToCopy(UNiagaraEmitter* Emitter)
+{
+	EmitterToCopy = Emitter;
+	return *this;
+}
+
+UNiagaraEmitter* FNiagaraEmitterAssetBuilder::BuildEmitter(FString& OutError)
+{
+	const FString PackageName = (PackagePath / AssetName);
+	UPackage* Package = CreateNiagaraPackage(PackageName);
+	if (!IsValid(Package))
+	{
+		OutError = TEXT("Failed to create new package for Niagara Emitter.");
+		return nullptr;
+	}
+
+	const EObjectFlags Flags = Package->GetFlags();
+
+	if (!IsValid(EmitterToCopy))
+	{
+		FSoftObjectPath DefaultEmptyEmitter = GetDefault<UNiagaraEditorSettings>()->DefaultEmptyEmitter;
+		if (DefaultEmptyEmitter.IsValid() && DefaultEmptyEmitter.IsAsset())
+		{
+			UObject* DefaultEmitter = DefaultEmptyEmitter.TryLoad();
+			EmitterToCopy = Cast<UNiagaraEmitter>(DefaultEmitter);
+		}
+	}
+
+	UNiagaraEmitter* NewNiagaraEmitter;
+	if (IsValid(EmitterToCopy))
+	{
+		const bool bUseInheritance = EmitterToCopy->bIsInheritable;
+		if (bUseInheritance)
+		{
+			NewNiagaraEmitter = UNiagaraEmitter::CreateWithParentAndOwner(FVersionedNiagaraEmitter(EmitterToCopy, EmitterToCopy->GetExposedVersion().VersionGuid), Package, *AssetName, Flags);
+		}
+		else
+		{
+			NewNiagaraEmitter = Cast<UNiagaraEmitter>(StaticDuplicateObject(EmitterToCopy, Package, *AssetName, Flags, UNiagaraEmitter::StaticClass()));
+			NewNiagaraEmitter->SetUniqueEmitterName(AssetName);
+			NewNiagaraEmitter->DisableVersioning(EmitterToCopy->GetExposedVersion().VersionGuid);
+		}
+
+		NewNiagaraEmitter->bIsInheritable = true;
+		NewNiagaraEmitter->TemplateAssetDescription = FText();
+		NewNiagaraEmitter->Category = FText();
+	}
+	else
+	{
+		// Create an empty emitter, source, and graph.
+		NewNiagaraEmitter = NewObject<UNiagaraEmitter>(Package, UNiagaraEmitter::StaticClass(), *AssetName, Flags | RF_Transactional);
+
+		// Figure out how an empty emitter is created
+		UNiagaraEmitterFactoryNew::InitializeEmitter(NewNiagaraEmitter, false);
+	}
+
+	NewNiagaraEmitter->ForEachVersionData([](FVersionedNiagaraEmitterData& EmitterVersionData) {
+		UNiagaraEmitterEditorData* EmitterEditorData = CastChecked<UNiagaraEmitterEditorData>(EmitterVersionData.GetEditorData());
+		EmitterEditorData->SetShowSummaryView(EmitterVersionData.AddEmitterDefaultViewState == ENiagaraEmitterDefaultSummaryState::Summary ? true : false);
+	});
+
+	Package->SetAssetAccessSpecifier(EAssetAccessSpecifier::Public);
+
+	// Notify the asset registry
+	FAssetRegistryModule::AssetCreated(NewNiagaraEmitter);
+
+	// Mark the package dirty...
+	Package->MarkPackageDirty();
+
+	return NewNiagaraEmitter;
+}
+
+FNiagaraSystemAssetBuilder& FNiagaraSystemAssetBuilder::WithSystemToCopy(UNiagaraSystem* System)
+{
+	SystemToCopy = System;
+	return *this;
+}
+	
+	
+FNiagaraSystemAssetBuilder& FNiagaraSystemAssetBuilder::WithEmitter(FVersionedNiagaraEmitter&& Emitter)
+{
+	EmittersToAddToNewSystem.Emplace(Emitter);
+	return *this;
+}
+
+UNiagaraSystem* FNiagaraSystemAssetBuilder::BuildSystem(FString& OutError)
+{
+	const UNiagaraSettings* Settings = GetDefault<UNiagaraSettings>();
+	if (!IsValid(Settings))
+	{
+		OutError = TEXT("No default Niagara Editor Settings found.");
+		return nullptr;
+	}
+
+	const FString PackageName = (PackagePath / AssetName);
+	UPackage* Package = CreateNiagaraPackage(PackageName);
+	if (!IsValid(Package))
+	{
+		OutError = TEXT("Failed to create new package for Niagara System.");
+		return nullptr;
+	}
+
+	const EObjectFlags Flags = Package->GetFlags();
+	UNiagaraSystem* NewNiagaraSystem;
+	if (IsValid(SystemToCopy))
+	{
+		if (SystemToCopy->IsReadyToRun() == false)
+		{
+			SystemToCopy->WaitForCompilationComplete();
+		}
+
+		NewNiagaraSystem = Cast<UNiagaraSystem>(StaticDuplicateObject(SystemToCopy, Package, *AssetName, Flags, UNiagaraSystem::StaticClass()));
+		NewNiagaraSystem->TemplateAssetDescription = FText();
+		NewNiagaraSystem->Category = FText();
+
+		// If the new system doesn't have a thumbnail image, check the thumbnail map of the original asset's UPackage
+		if (NewNiagaraSystem->ThumbnailImage == nullptr)
+		{
+			FString ObjectFullName = SystemToCopy->GetFullName();
+			FName ObjectName = FName(ObjectFullName);
+			FString PackageFullName;
+			ThumbnailTools::QueryPackageFileNameForObject(ObjectFullName, PackageFullName);
+			FThumbnailMap ThumbnailMap;
+			ThumbnailTools::ConditionallyLoadThumbnailsFromPackage(PackageFullName, { FName(ObjectFullName) }, ThumbnailMap);
+
+			// there should always be a dummy thumbnail in here
+			if (ThumbnailMap.Contains(ObjectName))
+			{
+				FObjectThumbnail Thumbnail = ThumbnailMap[ObjectName];
+				// we only want to copy the thumbnail over if it's not a dummy
+				if (Thumbnail.GetImageWidth() != 0 && Thumbnail.GetImageHeight() != 0)
+				{
+					ThumbnailTools::CacheThumbnail(NewNiagaraSystem->GetFullName(), &Thumbnail, NewNiagaraSystem->GetOutermost());
+				}
+			}
+		}
+	}
+	else if (EmittersToAddToNewSystem.Num() > 0)
+	{
+		NewNiagaraSystem = NewObject<UNiagaraSystem>(Package, *AssetName, Flags | RF_Transactional);
+		UNiagaraSystemFactoryNew::InitializeSystem(NewNiagaraSystem, true);
+
+		for (const FVersionedNiagaraEmitter& EmitterToAddToNewSystem : EmittersToAddToNewSystem)
+		{
+			FNiagaraEditorUtilities::AddEmitterToSystem(*NewNiagaraSystem, *EmitterToAddToNewSystem.Emitter, EmitterToAddToNewSystem.Version);
+		}
+	}
+	else
+	{
+		NewNiagaraSystem = NewObject<UNiagaraSystem>(Package, *AssetName, Flags | RF_Transactional);
+		UNiagaraSystemFactoryNew::InitializeSystem(NewNiagaraSystem, true);
+	}
+
+	if (NewNiagaraSystem->GetEffectType() == nullptr && Settings->GetDefaultEffectType())
+	{
+		UNiagaraEffectType* DefaultEffectType = Settings->GetDefaultEffectType();
+		NewNiagaraSystem->SetEffectType(DefaultEffectType);
+	}
+
+	NewNiagaraSystem->RequestCompile(false);
+
+	Package->SetAssetAccessSpecifier(EAssetAccessSpecifier::Public);
+
+	// Notify the asset registry
+	FAssetRegistryModule::AssetCreated(NewNiagaraSystem);
+
+	// Mark the package dirty...
+	Package->MarkPackageDirty();
+
+	return NewNiagaraSystem;
+}
+
+#endif
\ No newline at end of file
diff --git a/Engine/Plugins/FX/Niagara/Source/NiagaraEditor/Private/NiagaraExternalSystemEditorUtilities.cpp b/Engine/Plugins/FX/Niagara/Source/NiagaraEditor/Private/NiagaraExternalSystemEditorUtilities.cpp
--- a/Engine/Plugins/FX/Niagara/Source/NiagaraEditor/Private/NiagaraExternalSystemEditorUtilities.cpp
+++ b/Engine/Plugins/FX/Niagara/Source/NiagaraEditor/Private/NiagaraExternalSystemEditorUtilities.cpp
@@ -9,124 +9,1517 @@
 
 #include "ScopedTransaction.h"
 
+#include "AssetBuildUtilities/NiagaraAssetBuilder.h"
+#include "NiagaraStackQuery.h"
+#include "NiagaraStackEntryEnumerable.h"
+#include "PropertyHandle.h"
+#include "ViewModels/NiagaraEmitterHandleViewModel.h"
+#include "ViewModels/NiagaraEmitterViewModel.h"
+#include "ViewModels/NiagaraSystemViewModel.h"
+#include "ViewModels/Stack/NiagaraStackEntry.h"
+#include "ViewModels/Stack/NiagaraStackInputCategory.h"
+#include "ViewModels/Stack/NiagaraStackFunctionInput.h"
+#include "ViewModels/Stack/NiagaraStackValueCollection.h"
+#include "ViewModels/Stack/NiagaraStackModuleItem.h"
+#include "ViewModels/Stack/NiagaraStackObject.h"
+#include "ViewModels/Stack/NiagaraStackPropertyRow.h"
+#include "ViewModels/Stack/NiagaraStackRoot.h"
+#include "ViewModels/Stack/NiagaraStackScriptItemGroup.h"
+#include "ViewModels/Stack/NiagaraStackViewModel.h"
+
+#include "NiagaraNodeOutput.h"
+
+#include "NiagaraEmitter.h"
+
+#include "ToolsetRegistry/ToolsetLibrary.h"
+#include "JsonSchema/JsonSchemaGenerator.h"
+#include "JsonObjectConverter.h"
+#include "Serialization/JsonSerializer.h"
+#include "Serialization/JsonReader.h"
+#include "Serialization/JsonWriter.h"
+
+#include "UObject/PropertyAccessUtil.h"
+
+#include "NiagaraEditorSettings.h"
+#include "NiagaraScriptSource.h"
+#include "NiagaraNodeCustomHlsl.h"
+#include "NiagaraNodeParameterMapSet.h"
+
+#include "NiagaraEmitterFactoryNew.h"
+
 #define LOCTEXT_NAMESPACE "FNiagaraExternalSystemEditorUtilities"
 
-FNiagaraEmitterHandle* FNiagaraExternalSystemEditorUtilities::GetEmitterHandle(UNiagaraSystem& System, FName EmitterName)
+//////////////////////////////////////////////////////////////////////////
+// Helper functions for property access and JSON conversion
+
+namespace
+{
+	/// Checks if a property should be included in JSON conversion
+	/// Includes: CPF_Edit, CPF_BlueprintVisible, CPF_BlueprintAssignable, or "EDAEdit" metadata
+	/// Excludes: CPF_Deprecated
+	static bool ShouldExportProperty(const FProperty* Property)
+	{
+		if (Property->HasAnyPropertyFlags(CPF_Deprecated))
+		{
+			return false;
+		}
+
+		return Property->HasAnyPropertyFlags(CPF_Edit | CPF_BlueprintVisible | CPF_BlueprintAssignable) ||
+		       Property->HasMetaData("EDAEdit");
+	}
+
+	/// Converts a struct to JSON, filtering properties via ShouldExportProperty
+	static TSharedPtr<FJsonObject> StructToFilteredJsonObject(const UStruct* StructDefinition, const void* Struct)
+	{
+		TSharedPtr<FJsonObject> JsonObject = MakeShared<FJsonObject>();
+
+		for (TFieldIterator<FProperty> It(StructDefinition); It; ++It)
+		{
+			FProperty* Property = *It;
+			if (ShouldExportProperty(Property))
+			{
+				const void* PropertyValue = Property->ContainerPtrToValuePtr<void>(Struct);
+				TSharedPtr<FJsonValue> JsonValue = FJsonObjectConverter::UPropertyToJsonValue(
+					Property,
+					PropertyValue,
+					0,
+					0);
+				if (JsonValue.IsValid())
+				{
+					JsonObject->SetField(Property->GetName(), JsonValue);
+				}
+			}
+		}
+
+		return JsonObject;
+	}
+
+	/// Gets all editable property names from a struct (excludes deprecated)
+	static TArray<FName> GetAllPropertyNames(const UStruct* Struct)
+	{
+		TArray<FName> PropertyNames;
+
+		//TODO: Ideally we could using the ToolsetLibrary here but that currently would require a bunch of hoop jumping to properly get the set of Names for accessible properties.
+		EPropertyFlags CheckFlags = CPF_BlueprintVisible | CPF_BlueprintAssignable;
+		if (GEditor)
+		{
+			CheckFlags |= CPF_Edit;
+		}
+
+		for (TFieldIterator<FProperty> It(Struct); It; ++It)
+		{
+			const FProperty* Property = *It;
+			if (Property->HasAnyPropertyFlags(CheckFlags) && !Property->HasAnyPropertyFlags(CPF_Deprecated))
+			{
+				PropertyNames.Add(Property->GetFName());
+			}
+		}
+		return PropertyNames;
+	}
+
+	/// Gets all editable properties from a UObject as JSON string
+	static FString GetAllObjectProperties(const UObject* Object)
+	{
+		if (!Object)
+		{
+			return FString();
+		}
+		TArray<FName> PropertyNames = GetAllPropertyNames(Object->GetClass());
+		return UToolsetLibrary::GetObjectProperties(Object, PropertyNames);
+	}
+
+	/// Converts a stack item reference to a readable path for error messages
+	/// Example: "System(MySystem) / Emitter(Smoke) / Module(SpawnRate) / Input(Count)"
+	static FString GetStackItemPath(const FNiagaraExt_StackItemReference& Ref)
+	{
+		TStringBuilder<512> PathBuilder;
+		bool bNeedsDelimiter = false;
+
+		auto AppendPart = [&PathBuilder, &bNeedsDelimiter](const TCHAR* Part)
+		{
+			if (bNeedsDelimiter)
+			{
+				PathBuilder << TEXT(" / ");
+			}
+			PathBuilder << Part;
+			bNeedsDelimiter = true;
+		};
+
+		if (Ref.System)
+		{
+			PathBuilder << TEXT("System(") << Ref.System->GetName() << TEXT(")");
+			bNeedsDelimiter = true;
+		}
+
+		if (Ref.EmitterName != NAME_None)
+		{
+			AppendPart(*FString::Printf(TEXT("Emitter(%s)"), *Ref.EmitterName.ToString()));
+		}
+
+		if (Ref.ScriptName != NAME_None)
+		{
+			AppendPart(*FString::Printf(TEXT("Script(%s)"), *Ref.ScriptName.ToString()));
+		}
+
+		if (Ref.ModuleName != NAME_None)
+		{
+			AppendPart(*FString::Printf(TEXT("Module(%s)"), *Ref.ModuleName.ToString()));
+		}
+
+		if (Ref.RendererIndex != INDEX_NONE)
+		{
+			AppendPart(*FString::Printf(TEXT("Renderer[%d]"), Ref.RendererIndex));
+		}
+
+		if (Ref.InputNameStack.Num() > 0)
+		{
+			TStringBuilder<256> InputPath;
+			for (int32 i = 0; i < Ref.InputNameStack.Num(); ++i)
+			{
+				if (i > 0)
+				{
+					InputPath << TEXT(".");
+				}
+				InputPath << Ref.InputNameStack[i];
+			}
+			AppendPart(*FString::Printf(TEXT("Input(%s)"), InputPath.ToString()));
+		}
+
+		return PathBuilder.Len() > 0 ? PathBuilder.ToString() : TEXT("<empty reference>");
+	}
+
+	/// Sets properties on a UObject from JSON string, adding errors to OutErrors if failed
+	static bool SetAllObjectProperties(UObject* Object, const FString& PropertiesJson, TArray<FText>& OutErrors)
+	{
+		if (!Object || PropertiesJson.IsEmpty())
+		{
+			return false;
+		}
+
+		bool bSuccess = UToolsetLibrary::SetObjectProperties(Object, PropertiesJson);
+		if (!bSuccess)
+		{
+			OutErrors.Add(FText::Format(
+				NSLOCTEXT("NiagaraExternalEdit", "Property_Set_Failed", "Failed to set properties on {0} '{1}'. Check that property names and types are correct."),
+				FText::FromString(Object->GetClass()->GetName()),
+				FText::FromString(Object->GetName())));
+		}
+		return bSuccess;
+	}
+
+	/// Gets combined EmitterHandle and EmitterData properties as JSON
+	/// Merges: Handle properties (Name, bIsEnabled) + EmitterData struct properties
+	static FString GetEmitterHandleAndDataProperties(UNiagaraSystem* System, int32 EmitterIndex)
+	{
+		if (!System || !System->GetEmitterHandles().IsValidIndex(EmitterIndex))
+		{
+			return FString();
+		}
+
+		TSharedPtr<FJsonObject> CombinedJson = MakeShared<FJsonObject>();
+
+		const FNiagaraEmitterHandle& Handle = System->GetEmitterHandles()[EmitterIndex];
+		CombinedJson->SetStringField(TEXT("Name"), Handle.GetName().ToString());
+		CombinedJson->SetBoolField(TEXT("bIsEnabled"), Handle.GetIsEnabled());
+		CombinedJson->SetStringField(TEXT("Id"), Handle.GetId().ToString(EGuidFormats::DigitsWithHyphens));
+		CombinedJson->SetStringField(TEXT("IdName"), Handle.GetIdName().ToString());
+
+		FNiagaraEmitterHandle& MutableHandle = const_cast<FNiagaraEmitterHandle&>(Handle);
+		if (FVersionedNiagaraEmitterData* EmitterData = MutableHandle.GetEmitterData())
+		{
+			TSharedPtr<FJsonObject> EmitterDataJson = StructToFilteredJsonObject(
+				FVersionedNiagaraEmitterData::StaticStruct(),
+				EmitterData);
+
+			if (EmitterDataJson.IsValid())
+			{
+				for (const auto& Pair : EmitterDataJson->Values)
+				{
+					CombinedJson->SetField(Pair.Key, Pair.Value);
+				}
+			}
+		}
+
+		FString OutString;
+		TSharedRef<TJsonWriter<>> JsonWriter = TJsonWriterFactory<>::Create(&OutString);
+		FJsonSerializer::Serialize(CombinedJson.ToSharedRef(), JsonWriter);
+		return OutString;
+	}
+
+	/// Sets combined EmitterHandle and EmitterData properties from JSON
+	/// Handles both EmitterHandle properties (Name, bIsEnabled) and EmitterData struct properties
+	static bool SetEmitterHandleAndDataProperties(UNiagaraSystem* System, int32 EmitterIndex, const FString& PropertiesJson, TArray<FText>& OutErrors)
+	{
+		if (!System || !System->GetEmitterHandles().IsValidIndex(EmitterIndex) || PropertiesJson.IsEmpty())
+		{
+			return false;
+		}
+
+		const FNiagaraEmitterHandle& HandleRef = System->GetEmitterHandles()[EmitterIndex];
+		FNiagaraEmitterHandle& Handle = const_cast<FNiagaraEmitterHandle&>(HandleRef);
+
+		TSharedRef<TJsonReader<>> JsonReader = TJsonReaderFactory<>::Create(PropertiesJson);
+		TSharedPtr<FJsonObject> JsonObject;
+		if (!FJsonSerializer::Deserialize(JsonReader, JsonObject) || !JsonObject.IsValid())
+		{
+			OutErrors.Add(FText::Format(
+				NSLOCTEXT("NiagaraExternalEdit", "Emitter_Property_InvalidJson", "Failed to parse JSON for emitter '{0}' properties. Check JSON syntax."),
+				FText::FromName(Handle.GetName())));
+			return false;
+		}
+		bool bSuccess = true;
+
+		if (JsonObject->HasField(TEXT("Name")))
+		{
+			FString NameStr;
+			if (JsonObject->TryGetStringField(TEXT("Name"), NameStr))
+			{
+				Handle.SetName(FName(*NameStr), *System);
+			}
+		}
+
+		if (JsonObject->HasField(TEXT("bIsEnabled")))
+		{
+			bool bEnabled;
+			if (JsonObject->TryGetBoolField(TEXT("bIsEnabled"), bEnabled))
+			{
+				Handle.SetIsEnabled(bEnabled, *System, true);
+			}
+		}
+
+		if (FVersionedNiagaraEmitterData* EmitterData = Handle.GetEmitterData())
+		{
+			TSharedPtr<FJsonObject> EmitterDataJson = MakeShared<FJsonObject>(*JsonObject);
+			EmitterDataJson->RemoveField(TEXT("Name"));
+			EmitterDataJson->RemoveField(TEXT("bIsEnabled"));
+			EmitterDataJson->RemoveField(TEXT("Id"));
+			EmitterDataJson->RemoveField(TEXT("IdName"));
+
+			if (!FJsonObjectConverter::JsonObjectToUStruct(EmitterDataJson.ToSharedRef(),
+				FVersionedNiagaraEmitterData::StaticStruct(),
+				EmitterData,
+				0,
+				0))
+			{
+				OutErrors.Add(FText::Format(
+					NSLOCTEXT("NiagaraExternalEdit", "Emitter_Data_Set_Failed", "Failed to set emitter data properties on emitter '{0}'. Verify property names and types match FVersionedNiagaraEmitterData."),
+					FText::FromName(Handle.GetName())));
+				bSuccess = false;
+			}
+		}
+
+		return bSuccess;
+	}
+}
+
+//////////////////////////////////////////////////////////////////////////
+
+
+bool ScriptUsageFromName(FName ScriptUsageName, ENiagaraScriptUsage& OutUsage)
+{
+	const UEnum* UsageEnum = StaticEnum<ENiagaraScriptUsage>();
+	check(UsageEnum);
+
+	int64 Value = UsageEnum->GetValueByName(ScriptUsageName);
+	if (Value == INDEX_NONE)
+	{
+		return false;
+	}
+
+	OutUsage = (ENiagaraScriptUsage)Value;
+	return true;
+}
+
+bool NameFromScriptUsage(ENiagaraScriptUsage Usage, FName& OutScriptName)
+{
+	const UEnum* UsageEnum = StaticEnum<ENiagaraScriptUsage>();
+	check(UsageEnum);
+
+	FName Name = UsageEnum->GetNameByValue((int64)Usage);
+	if (Name == NAME_None)
+	{
+		return false;
+	}
+
+	OutScriptName = Name;
+	return true;
+}
+
+void ForEachFunctionInput(UNiagaraStackEntry* Entry, TFunction<bool(UNiagaraStackFunctionInput*)> Func)
+{
+	TNiagaraStackEntryEnumerable<UNiagaraStackEntry>(*Entry)
+		.Children().OfType<UNiagaraStackEntry>()
+		.Where([&](UNiagaraStackEntry* ChildItem)
+			{
+				if (UNiagaraStackFunctionInput* StackFunctionInput = Cast<UNiagaraStackFunctionInput>(ChildItem))
+				{
+					return Func(StackFunctionInput);
+				}
+
+				//Look in child items for non-inputs. This may be too permissive. 
+				ForEachFunctionInput(ChildItem, Func);
+				return false;
+			});
+}
+
+//////////////////////////////////////////////////////////////////////////
+
+FNiagaraExternalEditContext::FNiagaraExternalEditContext(UNiagaraSystem* System)
+{
+	if (System)
+	{
+		SystemViewModel = MakeShared<FNiagaraSystemViewModel>();
+		FNiagaraSystemViewModelOptions SystemOptions;
+		SystemOptions.bCanModifyEmittersFromTimeline = false;
+		SystemOptions.bCompileForEdit = true;
+		SystemOptions.bCanSimulate = false;
+		SystemOptions.bCanAutoCompile = false;
+		SystemOptions.bIsForDataProcessingOnly = true;
+		SystemOptions.MessageLogGuid = System->GetAssetGuid();
+		SystemViewModel->Initialize(*System, SystemOptions);
+	}
+}
+
+FNiagaraExternalEditContext::FNiagaraExternalEditContext(const FNiagaraExt_StackItemReference& ItemRef)
+{
+	UNiagaraSystem* System = ItemRef.GetSystem(*this);
+	if (System)
+	{
+		SystemViewModel = MakeShared<FNiagaraSystemViewModel>();
+		FNiagaraSystemViewModelOptions SystemOptions;
+		SystemOptions.bCanModifyEmittersFromTimeline = false;
+		SystemOptions.bCompileForEdit = true;
+		SystemOptions.bCanSimulate = false;
+		SystemOptions.bCanAutoCompile = false;
+		SystemOptions.bIsForDataProcessingOnly = true;
+		SystemOptions.MessageLogGuid = System->GetAssetGuid();
+		SystemViewModel->Initialize(*System, SystemOptions);
+	}
+}
+
+bool FNiagaraExternalEditContext::CheckSystem(const UNiagaraSystem* System)
+{
+	if (!System)
+	{
+		Error(NSLOCTEXT("NiagaraExternalEdit", "System_Null", "System is null. A valid UNiagaraSystem is required for this operation."));
+		return false;
+	}
+	return true;
+}
+
+bool FNiagaraExternalEditContext::CheckScriptAsset(const UNiagaraScript* ScriptAsset)
+{
+	if (!ScriptAsset)
+	{
+		Error(NSLOCTEXT("NiagaraExternalEdit", "Script_Asset_Null", "Script asset is null. A valid UNiagaraScript is required for this operation."));
+		return false;
+	}
+	return true;
+}
+
+
+TSharedPtr<FNiagaraSystemViewModel> FNiagaraExternalEditContext::GetSystemViewModel()
+{
+	if (SystemViewModel.IsValid() == false)
+	{
+		Errors.Add(NSLOCTEXT("NiagaraExternalEdit", "SystemViewModel_Invalid", "System view model is invalid. The system may have failed to initialize properly."));
+		return nullptr;
+	}
+
+	return SystemViewModel;
+}
+
+UNiagaraSystem* FNiagaraExternalEditContext::GetSystem()
+{
+	if (SystemViewModel)
+	{
+		UNiagaraSystem* Ret = &SystemViewModel->GetSystem();
+		CheckSystem(Ret);
+		return Ret;
+	}
+	return nullptr;
+}
+
+//////////////////////////////////////////////////////////////////////////
+
+void FNiagaraExt_StackItemReference::SetEmitter(FNiagaraEmitterHandle* Emitter)
+{
+	check(Emitter);
+	EmitterName = Emitter->GetName();
+	// Don't cache emitter pointer - it points into TArray and can become invalid
+}
+
+void FNiagaraExt_StackItemReference::SetEmitter(FName InEmitterName)
+{
+	EmitterName = InEmitterName;
+	CachedScript = nullptr;
+	CachedModule = nullptr;
+	CachedInput = nullptr;
+}
+
+void FNiagaraExt_StackItemReference::SetScript(UNiagaraStackScriptItemGroup* Script)
+{
+	check(Script);
+	CachedScript = Script;
+	NameFromScriptUsage(Script->GetScriptUsage(), ScriptName);
+	CachedModule = nullptr;
+	CachedInput = nullptr;
+}
+
+void FNiagaraExt_StackItemReference::SetScript(FName InScriptName)
+{
+	ScriptName = InScriptName;
+	CachedScript = nullptr;
+	CachedModule = nullptr;
+	CachedInput = nullptr;
+}
+
+void FNiagaraExt_StackItemReference::SetModule(UNiagaraStackModuleItem* Module)
+{
+	CachedModule = Module;
+	ModuleName = *Module->GetModuleNode().GetFunctionName();
+	CachedInput = nullptr;
+}
+
+void FNiagaraExt_StackItemReference::SetModule(FName InModuleName)
+{
+	ModuleName = InModuleName;
+	CachedModule = nullptr;
+	CachedInput = nullptr;
+}
+
+void FNiagaraExt_StackItemReference::SetInput(UNiagaraStackFunctionInput* Input)
+{
+	check(Input);
+	InputNameStack.Reset();
+	InputNameStack.Add(Input->GetInputParameterHandle().GetName());
+	CachedInput = Input;
+}
+
+void FNiagaraExt_StackItemReference::SetInput(FName Name)
+{
+	InputNameStack.Reset();
+	InputNameStack.Add(Name);
+	CachedInput = nullptr;
+}
+
+void FNiagaraExt_StackItemReference::PushInput(UNiagaraStackFunctionInput* Input)
+{
+	CachedInput = Input;
+	InputNameStack.Add(Input->GetInputParameterHandle().GetName());
+}
+
+void FNiagaraExt_StackItemReference::PushInput(FName Name)
+{
+	InputNameStack.Add(Name);
+	CachedInput = nullptr;
+}
+
+void FNiagaraExt_StackItemReference::SetRenderer(int32 InRendererIndex)
+{
+	RendererIndex = InRendererIndex;
+}
+
+UNiagaraSystem* FNiagaraExt_StackItemReference::GetSystem(FNiagaraExternalEditContext& Context)const
+{
+	if(!System)
+	{
+		Context.Error(FText::Format(
+			NSLOCTEXT("NiagaraExternalEdit", "StackRef_System_Missing", "System is missing in stack reference: {0}"),
+			FText::FromString(GetStackItemPath(*this))));
+	}
+	return System;
+}
+
+FNiagaraEmitterHandle* FNiagaraExt_StackItemReference::GetEmitter(FNiagaraExternalEditContext& Context, bool bRequired)const
+{
+	// Always search for emitter - no caching to avoid dangling pointer issues
+	if(!GetSystem(Context))
+	{
+		return nullptr;
+	}
+
+	if (EmitterName != NAME_None)
+	{
+		for (FNiagaraEmitterHandle& EmitterHandle : System->GetEmitterHandles())
+		{
+			if (EmitterHandle.GetName() == EmitterName)
+			{
+				return &EmitterHandle;
+			}
+		}
+
+		// Emitter not found - build error message
+		TArray<FString> AvailableEmitters;
+		for (const FNiagaraEmitterHandle& Handle : System->GetEmitterHandles())
+		{
+			AvailableEmitters.Add(Handle.GetName().ToString());
+		}
+		Context.Error(FText::Format(
+			NSLOCTEXT("NiagaraExternalEdit", "StackRef_Emitter_NotFound", "Emitter '{0}' not found in system '{1}'. Available emitters: [{2}]"),
+			FText::FromName(EmitterName),
+			FText::FromString(System->GetName()),
+			FText::FromString(FString::Join(AvailableEmitters, TEXT(", ")))));
+	}
+	else if(bRequired)
+	{
+		Context.Error(FText::Format(
+			NSLOCTEXT("NiagaraExternalEdit", "StackRef_Emitter_NotSpecified", "Emitter name not specified in stack reference: {0}"),
+			FText::FromString(GetStackItemPath(*this))));
+	}
+
+	return nullptr;
+}
+
+UNiagaraStackScriptItemGroup* FNiagaraExt_StackItemReference::GetScript(FNiagaraExternalEditContext& Context, bool bRequired)const
+{
+	// Check if cached script is still valid
+	if(CachedScript.IsValid())
+	{
+		return CachedScript.Get();
+	}
+
+	if (ScriptName != NAME_None)
+	{
+		ENiagaraScriptUsage Usage;
+		if (!ScriptUsageFromName(ScriptName, Usage))
+		{
+			Context.Error(FText::Format(
+				NSLOCTEXT("NiagaraExternalEdit", "StackRef_Script_InvalidName", "Invalid script name '{0}' in stack reference: {1}. Must be a valid ENiagaraScriptUsage (e.g., SystemSpawnScript, EmitterSpawnScript, ParticleUpdateScript)."),
+				FText::FromName(ScriptName),
+				FText::FromString(GetStackItemPath(*this))));
+			return nullptr;
+		}
+
+		// Check if this is an emitter script or system script
+		const bool bIsEmitterScript = (EmitterName != NAME_None);
+
+		TSharedPtr<FNiagaraSystemViewModel> SystemViewModel = Context.GetSystemViewModel();
+		if(ensure(SystemViewModel))
+		{
+			FNiagaraStackScriptItemGroupQuery ScriptGroupQuery = bIsEmitterScript ?
+				FNiagaraStackRootQuery::EmitterStackRootEntry(SystemViewModel.ToSharedRef(), EmitterName).FindScriptGroup(Usage, FGuid()) :
+				FNiagaraStackRootQuery::SystemStackRootEntry(SystemViewModel.ToSharedRef()).FindScriptGroup(Usage, FGuid());
+
+			CachedScript = ScriptGroupQuery.GetEntry();
+			if (!CachedScript.IsValid())
+			{
+				Context.Error(FText::Format(
+					NSLOCTEXT("NiagaraExternalEdit", "StackRef_Script_NotFound", "Script '{0}' not found in stack reference: {1}"),
+					FText::FromName(ScriptName),
+					FText::FromString(GetStackItemPath(*this))));
+			}
+		}
+	}
+	else if(bRequired)
+	{
+		Context.Error(FText::Format(
+			NSLOCTEXT("NiagaraExternalEdit", "StackRef_Script_NotSpecified", "Script name not specified in stack reference: {0}"),
+			FText::FromString(GetStackItemPath(*this))));
+	}
+	return CachedScript.Get();
+}
+
+UNiagaraStackModuleItem* FNiagaraExt_StackItemReference::GetModule(FNiagaraExternalEditContext& Context, bool bRequired)const
+{
+	// Check if cached module is still valid
+	if(CachedModule.IsValid())
+	{
+		return CachedModule.Get();
+	}
+
+	if (ModuleName != NAME_None)
+	{
+		UNiagaraStackScriptItemGroup* Script = GetScript(Context);
+		if (Script)
+		{
+			TNiagaraStackEntryEnumerable<UNiagaraStackScriptItemGroup>(*Script)
+				.Children().OfType<UNiagaraStackModuleItem>()
+				.Where([&](UNiagaraStackModuleItem* ModuleItem)
+					{
+						if (ModuleItem->GetModuleNode().GetFunctionName() == ModuleName)
+						{
+							CachedModule = ModuleItem;
+							return true;
+						}
+
+						return false;
+					});
+		}
+
+		if (!CachedModule.IsValid())
+		{
+			Context.Error(FText::Format(
+				NSLOCTEXT("NiagaraExternalEdit", "StackRef_Module_NotFound", "Module '{0}' not found in stack reference: {1}. Verify module exists in the specified script."),
+				FText::FromName(ModuleName),
+				FText::FromString(GetStackItemPath(*this))));
+		}
+	}
+	else if(bRequired)
+	{
+		Context.Error(FText::Format(
+			NSLOCTEXT("NiagaraExternalEdit", "StackRef_Module_NotSpecified", "Module name not specified in stack reference: {0}"),
+			FText::FromString(GetStackItemPath(*this))));
+	}
+	return CachedModule.Get();
+}
+
+UNiagaraStackFunctionInput* FNiagaraExt_StackItemReference::GetInput(FNiagaraExternalEditContext& Context, bool bRequired)const
+{
+	// Check if cached input is still valid
+	if(CachedInput.IsValid())
+	{
+		return CachedInput.Get();
+	}
+
+	if (InputNameStack.Num() > 0)
+	{
+		UNiagaraStackModuleItem* Module = GetModule(Context, true);
+		if (Module)
+		{
+			int32 NameIdx = 0;
+			TStringBuilder<128> InputNamePath;
+			auto NextInputName = [&]() -> FName
+				{
+					if (InputNameStack.IsValidIndex(NameIdx))
+					{
+						FName Name = InputNameStack[NameIdx++];
+						InputNamePath << TEXT("/") << Name;
+						return Name;
+					}
+					return NAME_None;
+				};
+
+			FName InputName = NextInputName();
+			if (InputName != NAME_None)
+			{
+				TFunction<bool(UNiagaraStackFunctionInput*)> FindInputs;
+				FindInputs = [&](UNiagaraStackFunctionInput* Input)
+					{
+						if (Input->GetInputParameterHandle().GetName() == InputName)
+						{
+							InputName = NextInputName();
+							if (InputName != NAME_None)
+							{
+								//If there are names left on the stack and we're a dynamic input then look at the child inputs.
+								if (Input->GetValueMode() == UNiagaraStackFunctionInput::EValueMode::Dynamic)
+								{
+									ForEachFunctionInput(Input, FindInputs);
+								}
+								else
+								{
+									Context.Error(FText::Format(
+										NSLOCTEXT("NiagaraExternalEdit", "StackRef_Input_NotDynamic", "Input '{0}' is not a dynamic input, but more inputs are specified in the path. Stack reference: {1}"),
+										FText::FromName(Input->GetInputParameterHandle().GetName()),
+										FText::FromString(GetStackItemPath(*this))));
+								}
+							}
+							else
+							{
+								//End of input chain and we found our match.
+								CachedInput = Input;
+							}
+							return true;
+						}
+						return false;
+					};
+
+				ForEachFunctionInput(Module, FindInputs);
+
+				if (!CachedInput.IsValid())
+				{
+					Context.Error(FText::Format(
+						NSLOCTEXT("NiagaraExternalEdit", "StackRef_Input_NotFound", "Input path '{0}' not found in stack reference: {1}"),
+						FText::FromString(InputNamePath.ToString()),
+						FText::FromString(GetStackItemPath(*this))));
+				}
+			}
+		}
+	}
+	else if(bRequired)
+	{
+		Context.Error(FText::Format(
+			NSLOCTEXT("NiagaraExternalEdit", "StackRef_Input_NotSpecified", "Input name not specified in stack reference: {0}"),
+			FText::FromString(GetStackItemPath(*this))));
+	}
+	return CachedInput.Get();
+}
+
+UNiagaraRendererProperties* FNiagaraExt_StackItemReference::GetRenderer(FNiagaraExternalEditContext& Context, bool bRequired)const
+{
+	// Check if cached renderer is still valid
+	if(RendererIndex != INDEX_NONE)
+	{
+		if (FNiagaraEmitterHandle* Emitter = GetEmitter(Context))
+		{
+			if (FVersionedNiagaraEmitterData* EmitterData = Emitter->GetEmitterData())
+			{
+				if (EmitterData->GetRenderers().IsValidIndex(RendererIndex))
+				{
+					return EmitterData->GetRenderer(RendererIndex);
+				}
+				else if (RendererIndex != INDEX_NONE)
+				{
+					const int32 RendererCount = EmitterData->GetRenderers().Num();
+					Context.Error(FText::Format(
+						NSLOCTEXT("NiagaraExternalEdit", "StackRef_Renderer_IndexOutOfBounds", "Renderer index {0} is out of bounds for emitter '{1}' (has {2} renderers). Stack reference: {3}"),
+						RendererIndex,
+						FText::FromName(Emitter->GetName()),
+						RendererCount,
+						FText::FromString(GetStackItemPath(*this))));
+				}
+			}
+			//TODO: Stateless.
+		}
+	}
+	else if(bRequired)
+	{
+		Context.Error(FText::Format(
+			NSLOCTEXT("NiagaraExternalEdit", "StackRef_Renderer_NotSpecified", "Renderer index not specified in stack reference: {0}"),
+			FText::FromString(GetStackItemPath(*this))));
+	}
+
+	return nullptr;
+}
+
+//////////////////////////////////////////////////////////////////////////
+
+UNiagaraSystem* UNiagaraExternalEditUtilities::CreateNiagaraSystem(const FString& AssetName, const FString& AssetPath, UNiagaraSystem* TemplateSystem, FNiagaraExternalEditContext& Context)
+{
+	// Validate asset name
+	if (AssetName.IsEmpty())
+	{
+		Context.Error(NSLOCTEXT("NiagaraExternalEdit", "CreateSystem_EmptyName", "Cannot create Niagara system: asset name is empty."));
+		return nullptr;
+	}
+
+	// Validate asset path
+	if (AssetPath.IsEmpty())
+	{
+		Context.Error(FText::Format(
+			NSLOCTEXT("NiagaraExternalEdit", "CreateSystem_EmptyPath", "Cannot create Niagara system '{0}': asset path is empty."),
+			FText::FromString(AssetName)));
+		return nullptr;
+	}
+
+	// Create system builder
+	FNiagaraSystemAssetBuilder NiagaraSystemAssetBuilder = FNiagaraSystemAssetBuilder(AssetPath, AssetName);
+
+	// If a template is provided, use it as a base. Otherwise, an empty system will be created.
+	if (TemplateSystem)
+	{
+		NiagaraSystemAssetBuilder.WithSystemToCopy(TemplateSystem);
+	}
+
+	// Build the system
+	FString ErrorStr;
+	UNiagaraSystem* System = NiagaraSystemAssetBuilder.BuildSystem(ErrorStr);
+
+	if (ErrorStr.IsEmpty() == false)
+	{
+		Context.Error(FText::Format(
+			NSLOCTEXT("NiagaraExternalEdit", "CreateSystem_BuildFailed", "Failed to create system '{0}' at path '{1}': {2}"),
+			FText::FromString(AssetName),
+			FText::FromString(AssetPath),
+			FText::FromString(ErrorStr)));
+	}
+
+	return System;
+}
+
+void UNiagaraExternalEditUtilities::GetAvailableDynamicInputs(const FNiagaraTypeDefinition& TypeDef, TArray<UNiagaraScript*>& OutDynamicInputScripts, FNiagaraExternalEditContext& Context)
+{
+	TArray<FAssetData> DynamicInputAssets;
+	FNiagaraEditorUtilities::FGetFilteredScriptAssetsOptions DynamicInputScriptFilterOptions;
+	DynamicInputScriptFilterOptions.ScriptUsageToInclude = ENiagaraScriptUsage::DynamicInput;
+	DynamicInputScriptFilterOptions.bIncludeNonLibraryScripts = false;
+	FNiagaraEditorUtilities::GetFilteredScriptAssets(DynamicInputScriptFilterOptions, DynamicInputAssets);
+
+	FPinCollectorArray InputPins;
+	TArray<UNiagaraNodeOutput*> OutputNodes;
+	auto MatchesInputType = [&](UNiagaraScript* Script)
+		{
+			UNiagaraScriptSource* DynamicInputScriptSource = Cast<UNiagaraScriptSource>(Script->GetLatestSource());
+			if (!DynamicInputScriptSource || !DynamicInputScriptSource->NodeGraph)
+			{
+				return false;
+			}
+			OutputNodes.Reset();
+			DynamicInputScriptSource->NodeGraph->GetNodesOfClass<UNiagaraNodeOutput>(OutputNodes);
+			if (OutputNodes.Num() == 1)
+			{
+				InputPins.Reset();
+				OutputNodes[0]->GetInputPins(InputPins);
+				if (InputPins.Num() == 1)
+				{
+					const UEdGraphSchema_Niagara* NiagaraSchema = GetDefault<UEdGraphSchema_Niagara>();
+					FNiagaraTypeDefinition PinType = NiagaraSchema->PinToTypeDefinition(InputPins[0]);
+					return FNiagaraEditorUtilities::AreTypesAssignable(PinType, TypeDef);
+				}
+			}
+			return false;
+		};
+
+	for (const FAssetData& DynamicInputAsset : DynamicInputAssets)
+	{
+		UNiagaraScript* DynamicInputScript = Cast<UNiagaraScript>(DynamicInputAsset.GetAsset());
+		if (DynamicInputScript != nullptr)
+		{
+			if (MatchesInputType(DynamicInputScript))
+			{
+				OutDynamicInputScripts.Add(DynamicInputScript);
+			}
+		}
+	}
+}
+
+//////////////////////////////////////////////////////////////////////////
+
+void UNiagaraExternalEditUtilities::GetSystemInfoComplete(UNiagaraSystem* System, FNiagaraExt_SystemInfo& OutInfo, FNiagaraExternalEditContext& Context)
+{
+	if (!Context.CheckSystem(System)) return;
+
+	GetSystemTopology(System, OutInfo.SystemTopology, Context);
+
+	OutInfo.UsedRenderers = Context.UsedRenderers.Array();
+	OutInfo.UsedDataInterfaces = Context.UsedDataInterfaces.Array();
+	OutInfo.UsedModules = Context.UsedModules.Array();
+	OutInfo.UsedDynamicInputs = Context.UsedDynamicInputs.Array();
+}
+
+//////////////////////////////////////////////////////////////////////////
+//Schema Layer
+
+void UNiagaraExternalEditUtilities::GetSystemSchema(FNiagaraExt_SystemSchema& OutSchema)
+{
+	OutSchema.PropertySchema = UToolsetLibrary::ListStructProperties(UNiagaraSystem::StaticClass());
+}
+
+void UNiagaraExternalEditUtilities::GetEmitterSchema(FNiagaraExt_EmitterSchema& OutSchema)
+{
+	// Combine EmitterHandle properties and EmitterData properties into one schema
+	TSharedPtr<FJsonObject> CombinedSchema = MakeShared<FJsonObject>();
+
+	// Helper to filter properties for schema generation
+	FJsonSchemaPropertyFilter PropertyFilter;
+	FJsonSchemaPropertyFilter::CustomCallback Cb;
+	Cb.BindLambda([](const FProperty* Property, const FString& ParameterDefaultString, const TSharedRef<FJsonObject>& OutputSchema)
+	{
+		return !ShouldExportProperty(Property); // Return true to skip
+	});
+	PropertyFilter.SkipFlags |= CPF_Deprecated;
+	PropertyFilter.CustomCb = &Cb;
+
+	// Get EmitterHandle struct schema (Name, bIsEnabled, Id, IdName)
+	TSharedPtr<FJsonObject> HandleSchema = FJsonSchemaGenerator::UStructToJsonSchemaObject(
+		FNiagaraEmitterHandle::StaticStruct(), PropertyFilter);
+
+	// Get EmitterData struct schema
+	TSharedPtr<FJsonObject> EmitterDataSchema = FJsonSchemaGenerator::UStructToJsonSchemaObject(
+		FVersionedNiagaraEmitterData::StaticStruct(), PropertyFilter);
+
+	// Merge both schemas
+	if (HandleSchema.IsValid())
+	{
+		const TSharedPtr<FJsonObject>* HandleProperties = nullptr;
+		if (HandleSchema->TryGetObjectField(TEXT("properties"), HandleProperties) && HandleProperties)
+		{
+			for (const auto& Pair : (*HandleProperties)->Values)
+			{
+				CombinedSchema->SetField(Pair.Key, Pair.Value);
+			}
+		}
+	}
+
+	if (EmitterDataSchema.IsValid())
+	{
+		const TSharedPtr<FJsonObject>* EmitterDataProperties = nullptr;
+		if (EmitterDataSchema->TryGetObjectField(TEXT("properties"), EmitterDataProperties) && EmitterDataProperties)
+		{
+			for (const auto& Pair : (*EmitterDataProperties)->Values)
+			{
+				CombinedSchema->SetField(Pair.Key, Pair.Value);
+			}
+		}
+	}
+
+	// Convert to JSON string
+	FString OutString;
+	TSharedRef<TJsonWriter<>> JsonWriter = TJsonWriterFactory<>::Create(&OutString);
+	FJsonSerializer::Serialize(CombinedSchema.ToSharedRef(), JsonWriter);
+	OutSchema.PropertySchema = OutString;
+}
+
+void UNiagaraExternalEditUtilities::GetRendererSchema(TSubclassOf<UNiagaraRendererProperties> RendererClass, FNiagaraExt_RendererSchema& OutSchema)
+{
+	OutSchema.PropertySchema = UToolsetLibrary::ListStructProperties(RendererClass);
+	OutSchema.RendererClass = RendererClass;
+}
+void UNiagaraExternalEditUtilities::GetDataInterfaceSchema(TSubclassOf<UNiagaraDataInterface> DataInterfaceClass, FNiagaraExt_DataInterfaceSchema& OutSchema)
+{
+	OutSchema.PropertySchema = UToolsetLibrary::ListStructProperties(DataInterfaceClass);
+	OutSchema.DataInterfaceClass = DataInterfaceClass;
+}
+
+void UNiagaraExternalEditUtilities::GetStackInputSchema(const FNiagaraExt_StackItemReference& InputReference, FNiagaraExt_StackInputSchema& OutSchema, FNiagaraExternalEditContext& Context)
+{
+	if (UNiagaraStackFunctionInput* Input = InputReference.GetInput(Context))
+	{
+		OutSchema.bSupportsExpressions = Input->SupportsCustomExpressions();
+		//OutSchema.Category = TODO: Need to gather the category from the stack or asset data.
+		OutSchema.Name = Input->GetInputParameterHandle().GetName();
+		OutSchema.Type = Input->GetInputType();
+
+		TOptional<FNiagaraVariableMetaData> MetaData = Input->GetInputMetaData();
+		if (MetaData.IsSet())
+		{
+			OutSchema.MetaData = MetaData.GetValue();
+		}
+	}
+}
+
+void UNiagaraExternalEditUtilities::GetModuleSchema(const FNiagaraExt_StackItemReference& ModuleReference, FNiagaraExt_ModuleSchema& OutSchema, FNiagaraExternalEditContext& Context)
+{
+	if (UNiagaraStackModuleItem* Module = ModuleReference.GetModule(Context))
+	{
+		OutSchema.Asset = Module->GetModuleNode().FunctionScript;
+
+		ForEachFunctionInput(Module, [&](UNiagaraStackFunctionInput* Input)
+			{
+				FNiagaraExt_StackItemReference InputRef = ModuleReference;
+				InputRef.SetInput(Input);
+				GetStackInputSchema(InputRef, OutSchema.Inputs.AddDefaulted_GetRef(), Context);
+				return false;//gather all
+			});
+
+		//TODO: Need to add output nodes to the schema/topology so that the LLM Can better understand the flow of data between modules.
+// 		const UNiagaraNodeOutput* OutputNode = FNiagaraStackGraphUtilities::GetEmitterOutputNodeForStackNode(Module->GetModuleNode());
+// 		TArray<FNiagaraVariable> OutputVariables;
+// 		TArray<FNiagaraVariable> OutputVariablesWithOriginalAliasesIntact;
+// 		FCompileConstantResolver ConstantResolver = FCompileConstantResolver();
+// 		if (OutputNode)
+// 		{
+// 			FNiagaraEmitterHandle* Emitter = ModuleReference.GetEmitter(Context);
+// 			if (Emitter && Emitter->GetEmitterData())
+// 			{
+// 				ConstantResolver = FCompileConstantResolver(*Emitter->GetEmitterData(), OutputNode->GetUsage(), ENiagaraFunctionDebugState::NoDebug);
+// 			}
+// 			else
+// 			{
+// 				ConstantResolver = FCompileConstantResolver(ModuleReference.GetSystem(), OutputNode->GetUsage(), ENiagaraFunctionDebugState::NoDebug);
+// 			}
+// 		}
+// 		FNiagaraStackGraphUtilities::GetStackFunctionOutputVariables(Module->GetModuleNode(), ConstantResolver, OutputVariables, OutputVariablesWithOriginalAliasesIntact);
+// 
+// 		for (FNiagaraVariable& Var : OutputVariables)
+// 		{
+// 			OutSchema.Outputs.Emplace(Var.GetName(), Var.GetType());
+// 		}
+	}
+}
+
+void UNiagaraExternalEditUtilities::GetDynamicInputSchema(const FNiagaraExt_StackItemReference& DynamicInputReference, FNiagaraExt_DynamicInputSchema& OutSchema, FNiagaraExternalEditContext& Context)
+{
+	if (UNiagaraStackFunctionInput* Input = DynamicInputReference.GetInput(Context))
+	{
+		ensure(Input->GetValueMode() == UNiagaraStackFunctionInput::EValueMode::Dynamic);
+
+		OutSchema.Asset = Input->GetDynamicInputNode()->FunctionScript;
+
+		ForEachFunctionInput(Input, [&](UNiagaraStackFunctionInput* Input)
+			{
+				FNiagaraExt_StackItemReference InputRef = DynamicInputReference;
+				InputRef.PushInput(Input);
+				GetStackInputSchema(InputRef, OutSchema.Inputs.AddDefaulted_GetRef(), Context);
+				return false;//Keep going to gather all inputs.
+			});
+	}
+}
+
+void UNiagaraExternalEditUtilities::GetModuleSchema(const UNiagaraScript* ModuleAsset, FNiagaraExt_ModuleSchema& OutSchema, FNiagaraExternalEditContext& Context)
+{
+	if (!Context.CheckScriptAsset(ModuleAsset)) return;
+
+	TArray<FNiagaraParameterMapHistory> Histories;
+
+	//TODO: This is all very incomplete. A cobbled together collection from stack graph utilities. Improve or ideally move to stack graph utilities or new stack adapters.
+	//TODO: Doesn't consider static variables and likely a lot of other things.
+	//TODO: Versioning?
+	if (const UNiagaraScriptSource* Source = Cast<UNiagaraScriptSource>(ModuleAsset->GetLatestSource()))
+	{
+		const UNiagaraEditorSettings* EditorSettings = GetDefault<UNiagaraEditorSettings>();
+		bool bSupportsExpressions = EditorSettings->IsVisibleClass(UNiagaraNodeCustomHlsl::StaticClass());
+
+		if (UNiagaraGraph* Graph = Source->NodeGraph)
+		{
+			if (UNiagaraExternalEditUtilities::BuildParameterMapHistoriesFromScript(ModuleAsset, Histories, Context))
+			{
+				if (Histories.Num() == 1)
+				{
+					FNiagaraParameterMapHistory& History = Histories[0];
+					TArray<const UEdGraphPin*> InputPins;
+					FNiagaraStackGraphUtilities::ExtractInputPinsFromHistory(History, Graph, FNiagaraStackGraphUtilities::ENiagaraGetStackFunctionInputPinsOptions::ModuleInputsOnly, InputPins);
+
+					for (const UEdGraphPin* InputPin : InputPins)
+					{
+						FNiagaraVariableBase Var = UEdGraphSchema_Niagara::PinToNiagaraVariable(InputPin);
+						FNiagaraExt_StackInputSchema& NewInputSchema = OutSchema.Inputs.AddDefaulted_GetRef();
+						NewInputSchema.Name = Var.GetName();
+						NewInputSchema.Type = Var.GetType();
+
+						NewInputSchema.bSupportsExpressions = bSupportsExpressions;
+
+						TOptional<FNiagaraVariableMetaData> MetaData = Graph->GetMetaData(Var);
+						if (MetaData.IsSet())
+						{
+							NewInputSchema.MetaData = MetaData.GetValue();
+							//NewInputSchema.Category = MetaData->CategoryName_DEPRECATED. Need to get this from somewhere else...
+						}
+						else
+						{
+							//NewInputSchema.Category = FText::GetEmpty();
+						}
+					}
+
+					for (int32 i = 0; i < History.Variables.Num(); i++)
+					{
+						bool bHasParameterMapSetWrite = false;
+						for (const FNiagaraParameterMapHistory::FModuleScopedPin& WritePin : History.PerVariableWriteHistory[i])
+						{
+							if (WritePin.Pin != nullptr && WritePin.Pin->GetOwningNode() != nullptr && WritePin.Pin->GetOwningNode()->IsA<UNiagaraNodeParameterMapSet>())
+							{
+								bHasParameterMapSetWrite = true;
+								break;
+							}
+						}
+
+						if (bHasParameterMapSetWrite)
+						{
+							FNiagaraVariable& Variable = History.Variables[i];
+							FNiagaraVariable& VariableWithOriginalAliasIntact = History.VariablesWithOriginalAliasesIntact[i];
+							FNiagaraExt_Variable& NewOutputVar = OutSchema.Outputs.AddDefaulted_GetRef();
+							NewOutputVar.Name = Variable.GetName();
+							NewOutputVar.Type = Variable.GetType();
+						}
+					}
+				}
+				else
+				{
+					Context.Error(FText::Format(
+						NSLOCTEXT("NiagaraExternalEdit", "Module_Graph_Malformed", "Module '{0}' has a malformed node graph. Expected 1 parameter history but found {1}."),
+						FText::FromString(ModuleAsset->GetName()),
+						Histories.Num()));
+				}
+			}
+		}
+	}
+}
+
+void UNiagaraExternalEditUtilities::GetDynamicInputSchema(const UNiagaraScript* ModuleAsset, FNiagaraExt_DynamicInputSchema& OutSchema, FNiagaraExternalEditContext& Context)
+{
+	GetModuleSchema(ModuleAsset, OutSchema, Context);
+}
+
+bool UNiagaraExternalEditUtilities::BuildParameterMapHistoriesFromScript(const UNiagaraScript* Script, TArray<FNiagaraParameterMapHistory>& OutHistories, FNiagaraExternalEditContext& Context)
+{
+	if (!Context.CheckScriptAsset(Script)) return false;
+
+	const UNiagaraScriptSourceBase* Source = Script->GetLatestSource();
+	if (!Source)
+	{
+		Context.Error(FText::Format(
+			NSLOCTEXT("NiagaraExternalEdit", "Script_Source_Missing", "Script '{0}' has no source. The script may be corrupted or not yet compiled."),
+			FText::FromString(Script->GetName())));
+		return false;
+	}
+
+	const UNiagaraScriptSource* ScriptSource = Cast<UNiagaraScriptSource>(Source);
+	if (!ScriptSource || !ScriptSource->NodeGraph)
+	{
+		Context.Error(FText::Format(
+			NSLOCTEXT("NiagaraExternalEdit", "Script_Graph_Missing", "Script '{0}' has no node graph. The script may be corrupted or in an invalid state."),
+			FText::FromString(Script->GetName())));
+		return false;
+	}
+
+	UNiagaraGraph* Graph = ScriptSource->NodeGraph;
+
+	// Create parameter map history builder
+	FNiagaraParameterMapHistoryBuilder Builder;
+	Builder.SetIgnoreDisabled(false);
+
+	// Find output nodes
+	TArray<UNiagaraNodeOutput*> OutputNodes;
+	Graph->FindOutputNodes(OutputNodes);
+
+	if (OutputNodes.Num() == 0)
+	{
+		Context.Error(FText::Format(
+			NSLOCTEXT("NiagaraExternalEdit", "Script_OutputNodes_Missing", "Script '{0}' graph has no output nodes. The script may be corrupted or incomplete."),
+			FText::FromString(Script->GetName())));
+		return false;
+	}
+
+	// Build parameter maps for each output node
+	for (UNiagaraNodeOutput* OutputNode : OutputNodes)
+	{
+		Builder.BeginUsage(OutputNode->GetUsage(), FName());
+		Builder.EnableScriptAllowList(true, OutputNode->GetUsage());
+		Builder.BuildParameterMaps(OutputNode, true);
+		Builder.EndUsage();
+	}
+
+	OutHistories = MoveTemp(Builder.Histories);
+	return OutHistories.Num() > 0;
+}
+
+//////////////////////////////////////////////////////////////////////////
+//Topology - Defines the current topology of niagara Systems, Emitters and module inputs.
+void UNiagaraExternalEditUtilities::GetSystemTopology(UNiagaraSystem* System, FNiagaraExt_SystemTopology& OutTopology, FNiagaraExternalEditContext& Context)
+{
+	if (!Context.CheckSystem(System)) return;
+
+	GetScriptStackTopology(FNiagaraExt_StackItemReference(System, NAME_None, TEXT("SystemSpawnScript")), OutTopology.SystemSpawnScript, Context);
+	GetScriptStackTopology(FNiagaraExt_StackItemReference(System, NAME_None, TEXT("SystemUpdateScript")), OutTopology.SystemUpdateScript, Context);
+
+	for (FNiagaraEmitterHandle& EmitterHandle : System->GetEmitterHandles())
+	{
+		FNiagaraExt_EmitterTopology& NewEmitterTopology = OutTopology.Emitters.AddDefaulted_GetRef();
+		GetEmitterTopology(FNiagaraExt_StackItemReference(System, EmitterHandle.GetName()), NewEmitterTopology, Context);
+	}
+
+	const FNiagaraUserRedirectionParameterStore& UserParamStore = System->GetExposedParameters();
+	TArray<FNiagaraVariable> UserVariables;
+	UserParamStore.GetUserParameters(UserVariables);
+
+	OutTopology.UserVariables.Reset(UserVariables.Num());
+	for (const FNiagaraVariable& Var : UserVariables)
+	{
+		FNiagaraExt_UserVariable& NewVarInfo = OutTopology.UserVariables.AddDefaulted_GetRef();
+		NewVarInfo.Name = Var.GetName();
+		NewVarInfo.Type = Var.GetType();
+
+		FNiagaraVariant DefaultValue;
+		UserParamStore.GetParameter_InternalUseOnly(Var, DefaultValue);
+
+		NewVarInfo.DefaultValue.Set(Var.GetType(), DefaultValue);
+
+		if (UNiagaraScriptVariable* ScriptVariable = FNiagaraEditorUtilities::UserParameters::GetScriptVariableForUserParameter(Var, *System))
+		{
+			NewVarInfo.Description = ScriptVariable->Metadata.Description;
+		}
+	}
+
+	GetSystemData(System, OutTopology.Data, Context);
+}
+
+void UNiagaraExternalEditUtilities::GetScriptStackTopology(const FNiagaraExt_StackItemReference& ScriptRef, FNiagaraExt_ScriptStackTopology& OutTopology, FNiagaraExternalEditContext& Context)
+{
+	if (UNiagaraStackScriptItemGroup* Script = ScriptRef.GetScript(Context))
+	{
+		OutTopology.ScriptName = ScriptRef.ScriptName;
+		TNiagaraStackEntryEnumerable<UNiagaraStackScriptItemGroup>(*Script)
+			.Children().OfType<UNiagaraStackModuleItem>()
+			.Where([&](UNiagaraStackModuleItem* ModuleItem)
+				{
+					FNiagaraExt_ModuleTopology& ModuleTopology = OutTopology.Modules.AddDefaulted_GetRef();
+
+					FNiagaraExt_StackItemReference ModuleRef(ScriptRef);
+					ModuleRef.SetModule(ModuleItem);
+					GetModuleTopology(ModuleRef, ModuleTopology, Context);
+					return false;//Gather all.
+				});
+	}
+}
+
+void UNiagaraExternalEditUtilities::GetEmitterTopology(const FNiagaraExt_StackItemReference& EmitterRef, FNiagaraExt_EmitterTopology& OutTopology, FNiagaraExternalEditContext& Context)
+{
+	if (FNiagaraEmitterHandle* Emitter = EmitterRef.GetEmitter(Context))
+	{
+		OutTopology.EmitterName = Emitter->GetName();
+
+		if (FVersionedNiagaraEmitterData* EmitterData = Emitter->GetEmitterData())
+		{
+			UNiagaraSystem* System = EmitterRef.GetSystem(Context);
+			GetScriptStackTopology(FNiagaraExt_StackItemReference(System, OutTopology.EmitterName, TEXT("EmitterSpawnScript")), OutTopology.EmitterSpawnScript, Context);
+			GetScriptStackTopology(FNiagaraExt_StackItemReference(System, OutTopology.EmitterName, TEXT("EmitterUpdateScript")), OutTopology.EmitterUpdateScript, Context);
+			GetScriptStackTopology(FNiagaraExt_StackItemReference(System, OutTopology.EmitterName, TEXT("ParticleSpawnScript")), OutTopology.ParticleSpawnScript, Context);
+			GetScriptStackTopology(FNiagaraExt_StackItemReference(System, OutTopology.EmitterName, TEXT("ParticleUpdateScript")), OutTopology.ParticleUpdateScript, Context);
+
+			TArray<UNiagaraRendererProperties*> Renderers = EmitterData->GetRenderers();
+			for (int32 i = 0; i < Renderers.Num(); ++i)
+			{
+				UNiagaraRendererProperties* Renderer = Renderers[i];
+				FNiagaraExt_RendererRef& RendererRef = OutTopology.Renderers.AddDefaulted_GetRef();
+				RendererRef.Index = i;
+				RendererRef.RendererClass = Renderer->GetClass();
+
+				FNiagaraExt_StackItemReference RendererStackRef = EmitterRef;
+				RendererStackRef.SetRenderer(i);
+				GetRendererData(RendererStackRef, RendererRef.Data, Context);
+
+				Context.UsedRenderers.Add(Renderer->GetClass());
+			}
+
+			GetEmitterData(EmitterRef, OutTopology.Data, Context);
+		}
+	}
+}
+
+void UNiagaraExternalEditUtilities::GetModuleTopology(const FNiagaraExt_StackItemReference& ModuleRef, FNiagaraExt_ModuleTopology& OutTopology, FNiagaraExternalEditContext& Context)
+{
+	if (UNiagaraStackModuleItem* Module = ModuleRef.GetModule(Context))
+	{
+		OutTopology.ModuleName = *Module->GetModuleNode().GetFunctionName();
+		OutTopology.Enabled = Module->GetIsEnabled();
+		OutTopology.ModuleScript = Module->GetModuleNode().FunctionScript;
+		ForEachFunctionInput(Module, [&](UNiagaraStackFunctionInput* Input)
+			{
+				FNiagaraExt_StackInputTopology& NewInputTopology = OutTopology.Inputs.AddDefaulted_GetRef();
+				FNiagaraExt_StackItemReference InputRef(ModuleRef);
+				InputRef.SetInput(Input);
+				GetStackInputTopology(InputRef, NewInputTopology, Context);				
+				return false;
+			});
+
+		Context.UsedModules.Add(OutTopology.ModuleScript);
+	}
+}
+
+void UNiagaraExternalEditUtilities::GetStackInputTopology(const FNiagaraExt_StackItemReference& StackInputRef, FNiagaraExt_StackInputTopology& OutTopology, FNiagaraExternalEditContext& Context)
+{
+	if (UNiagaraStackFunctionInput* Input = StackInputRef.GetInput(Context))
+	{
+		OutTopology.Name = Input->GetInputParameterHandle().GetName();
+		OutTopology.Type = Input->GetInputType();
+		
+		OutTopology.bIsEditable = Input->GetHasEditCondition() ? Input->GetEditConditionEnabled() : true;
+		OutTopology.bIsVisible = Input->GetHasVisibleCondition() ? Input->GetVisibleConditionEnabled() : true;
+
+		GetStackInputData(StackInputRef, OutTopology.Data, Context);
+
+		if(Input->GetValueMode() == UNiagaraStackFunctionInput::EValueMode::Dynamic)
+		{
+			ForEachFunctionInput(Input, [&](UNiagaraStackFunctionInput* DynamicInput)
+				{
+					if (DynamicInput)
+					{
+						FNiagaraExt_StackItemReference DynamicInputRef = StackInputRef;
+						DynamicInputRef.PushInput(DynamicInput);
+						FNiagaraExt_DynamicInputTopology& DynamicInputTopology = OutTopology.DynamicInputs.AddDefaulted_GetRef();
+						GetStackInputTopology(DynamicInputRef, DynamicInputTopology.GetMutable(), Context);
+					}
+					return false;
+				});
+		}
+	}
+}
+
+//////////////////////////////////////////////////////////////////////////
+//Data Layer - Access actual data values of Systems, Emitters and module inputs.
+
+void UNiagaraExternalEditUtilities::GetUserVariables(UNiagaraSystem* System, FNiagaraExt_UserVariables& OutVariables, FNiagaraExternalEditContext& Context)
 {
-	for (FNiagaraEmitterHandle& EmitterHandle : System.GetEmitterHandles())
+	if (!Context.CheckSystem(System)) return;
+
+	const FNiagaraUserRedirectionParameterStore& UserParamStore = System->GetExposedParameters();
+	TArray<FNiagaraVariable> UserVariables;
+	UserParamStore.GetUserParameters(UserVariables);
+
+	OutVariables.UserVariables.Reset(UserVariables.Num());
+	for (const FNiagaraVariable& Var : UserVariables)
 	{
-		if (EmitterHandle.GetName() == EmitterName)
+		FNiagaraExt_UserVariable& NewVarData = OutVariables.UserVariables.AddDefaulted_GetRef();
+		NewVarData.Name = Var.GetName();
+		NewVarData.Type = Var.GetType();
+
+		FNiagaraVariant DefaultValue;
+		UserParamStore.GetParameter_InternalUseOnly(Var, DefaultValue);
+
+		NewVarData.DefaultValue.Set(Var.GetType(), DefaultValue);
+
+		if (UNiagaraScriptVariable* ScriptVariable = FNiagaraEditorUtilities::UserParameters::GetScriptVariableForUserParameter(Var, *System))
 		{
-			return &EmitterHandle;
+			NewVarData.Description = ScriptVariable->Metadata.Description;
 		}
 	}
-	return nullptr;
 }
 
-FNiagaraEmitterHandle* FNiagaraExternalSystemEditorUtilities::GetEmitterHandle(UNiagaraSystem& System, const FGuid& EmitterId)
+void UNiagaraExternalEditUtilities::GetSystemData(UNiagaraSystem* System, FNiagaraExt_SystemData& OutData, FNiagaraExternalEditContext& Context)
+{
+	if (!Context.CheckSystem(System)) return;
+
+	OutData.PropertyValues = GetAllObjectProperties(System);
+}
+
+void UNiagaraExternalEditUtilities::GetEmitterData(const FNiagaraExt_StackItemReference& EmitterRef, FNiagaraExt_EmitterData& OutData, FNiagaraExternalEditContext& Context)
+{
+	if (FNiagaraEmitterHandle* Emitter = EmitterRef.GetEmitter(Context))
+	{
+		GetEmitterDataInternal(EmitterRef.GetSystem(Context), *Emitter, OutData, Context);
+	}
+}
+
+void UNiagaraExternalEditUtilities::GetEmitterDataInternal(UNiagaraSystem* System, FNiagaraEmitterHandle& Emitter, FNiagaraExt_EmitterData& OutData, FNiagaraExternalEditContext& Context)
 {
-	for (FNiagaraEmitterHandle& EmitterHandle : System.GetEmitterHandles())
+	if (!Context.CheckSystem(System)) return;
+
+	// Find the emitter index
+	int32 EmitterIndex = System->GetEmitterHandles().IndexOfByPredicate([&](const FNiagaraEmitterHandle& Handle) {
+		return Handle.GetId() == Emitter.GetId();
+	});
+
+	if (EmitterIndex != INDEX_NONE)
+	{
+		OutData.PropertyValues = GetEmitterHandleAndDataProperties(System, EmitterIndex);
+	}
+	else
 	{
-		if (EmitterHandle.GetId() == EmitterId)
+		// Build list of available emitters for better error message
+		TArray<FString> AvailableEmitters;
+		for (const FNiagaraEmitterHandle& Handle : System->GetEmitterHandles())
 		{
-			return &EmitterHandle;
+			AvailableEmitters.Add(Handle.GetName().ToString());
 		}
+		Context.Error(FText::Format(
+			NSLOCTEXT("NiagaraExternalEdit", "Emitter_Data_Get_NotFound", "Emitter with ID '{0}' not found in system '{1}'. Available emitters: [{2}]"),
+			FText::FromString(Emitter.GetId().ToString(EGuidFormats::DigitsWithHyphens)),
+			FText::FromString(System->GetName()),
+			FText::FromString(FString::Join(AvailableEmitters, TEXT(", ")))));
 	}
-	return nullptr;
 }
 
-void FNiagaraExternalSystemEditorUtilities::AddUserVariable(UNiagaraSystem& System, const FNiagaraVariableBase& Variable, const FNiagaraVariant& Variant, FString VariableDesc, FText TransactionDescription)
+void UNiagaraExternalEditUtilities::GetRendererData(const FNiagaraExt_StackItemReference& RendererRef, FNiagaraExt_RendererData& OutData, FNiagaraExternalEditContext& Context)
 {
-	if (!Variant.IsValid())
+	UNiagaraRendererProperties* Renderer = RendererRef.GetRenderer(Context, true);
+	FNiagaraEmitterHandle* Emitter = RendererRef.GetEmitter(Context, true);
+	if (Emitter && Renderer)
 	{
-		return;
+		GetRendererDataInternal(*Emitter, Renderer, OutData, Context);
 	}
+}
 
-	const FScopedTransaction Transaction(TransactionDescription);
-	System.Modify();
+void UNiagaraExternalEditUtilities::GetRendererDataInternal(FNiagaraEmitterHandle& Emitter, UNiagaraRendererProperties* Renderer, FNiagaraExt_RendererData& OutData, FNiagaraExternalEditContext& Context)
+{
+	OutData.PropertyValues = GetAllObjectProperties(Renderer);
+}
+
+void UNiagaraExternalEditUtilities::GetStackInputData(const FNiagaraExt_StackItemReference& StackInputRef, FNiagaraExt_StackInputValue& OutData, FNiagaraExternalEditContext& Context)
+{
+	OutData.InitFromStackInput(StackInputRef, Context);
+}
 
-	System.GetExposedParameters().SetParameter_InternalUseOnly(Variable, Variant, true);
+void UNiagaraExternalEditUtilities::AddUserVariable(UNiagaraSystem* System, const FNiagaraExt_UserVariable& Variable, FNiagaraExternalEditContext& Context)
+{
+	if(!Context.CheckSystem(System))
+	{
+		return;
+	}
 
-	if (UNiagaraScriptVariable* ScriptVariable = FNiagaraEditorUtilities::UserParameters::GetScriptVariableForUserParameter(Variable, System))
+	if (!Variable.DefaultValue.IsValid())
 	{
-		ScriptVariable->Metadata.Description = FText::FromString(VariableDesc);
+		return;
 	}
-	FPropertyChangedEvent PropertyChangedEvent(nullptr, EPropertyChangeType::ValueSet);
-	System.PostEditChangeProperty(PropertyChangedEvent);
-}
 
-void FNiagaraExternalSystemEditorUtilities::RemoveUserVariable(UNiagaraSystem& System, const FNiagaraVariableBase& Variable, FText TransactionDescription)
-{
-	const FScopedTransaction Transaction(TransactionDescription);
-	System.Modify();
+	const FScopedTransaction Transaction(LOCTEXT("ExternalAddUserVariable", "External System Add User Variable"));
+	System->Modify();
 
-	System.GetExposedParameters().RemoveParameter(Variable);
+	//Use this possibly? There is no corresponding remove and it doens't appear to do much beyond add, set default and being a rename.
+	//FNiagaraEditorUtilities::AddParameter(NewParameter, SystemViewModel.Pin()->GetSystem().GetExposedParameters(), SystemViewModel.Pin()->GetSystem(), nullptr);
 
+	FNiagaraVariableBase Var(Variable.Type, Variable.Name);
+	FNiagaraVariant DefaultValue;
+	Variable.DefaultValue.Get(DefaultValue, Context);
+	System->GetExposedParameters().SetParameter_InternalUseOnly(Var, DefaultValue, true);
+	if (UNiagaraScriptVariable* ScriptVariable = FNiagaraEditorUtilities::UserParameters::GetScriptVariableForUserParameter(Var, *System))
+	{
+		ScriptVariable->Metadata.Description = Variable.Description;
+	}
 	FPropertyChangedEvent PropertyChangedEvent(nullptr, EPropertyChangeType::ValueSet);
-	System.PostEditChangeProperty(PropertyChangedEvent);
+	System->PostEditChangeProperty(PropertyChangedEvent);
 }
 
-void FNiagaraExternalSystemEditorUtilities::GetUserVariableInfo(UNiagaraSystem& System, TArray<FNiagaraExternalSystemEditorUtilities::FUserVariableInfo>& OutVariableInfo)
+void UNiagaraExternalEditUtilities::RemoveUserVariable(UNiagaraSystem* System, const FNiagaraExt_Variable& Variable, FNiagaraExternalEditContext& Context)
 {
-	const FNiagaraUserRedirectionParameterStore& UserParamStore = System.GetExposedParameters();
-	TArray<FNiagaraVariable> UserVariables;
-	UserParamStore.GetUserParameters(UserVariables);
+	if (!Context.CheckSystem(System)) return;
 
-	OutVariableInfo.Reset(UserVariables.Num());
-	for (const FNiagaraVariable& Var : UserVariables)
-	{
-		FNiagaraExternalSystemEditorUtilities::FUserVariableInfo& NewVarInfo = OutVariableInfo.AddDefaulted_GetRef();
-		NewVarInfo.Variable = Var;
-		const FNiagaraTypeDefinition& Type = Var.GetType();
+	const FScopedTransaction Transaction(LOCTEXT("ExternalRemoveUserVariable", "External System Remove User Variable"));
+	System->Modify();
 
-		if (const FNiagaraVariableWithOffset* StoreVarWithOffset = UserParamStore.FindParameterVariable(Var))
-		{
-			if (Var.IsUObject())
-			{
-				NewVarInfo.Value.SetUObject(UserParamStore.GetUObject(StoreVarWithOffset->Offset));
-			}
-			else if (Var.IsDataInterface())
-			{
-				NewVarInfo.Value.SetDataInterface(UserParamStore.GetDataInterface(StoreVarWithOffset->Offset));
-			}
-			else
-			{
-				check(Var.GetSizeInBytes() == StoreVarWithOffset->GetSizeInBytes());
-				int32 VarSize = Var.GetSizeInBytes();
-				NewVarInfo.Value.SetBytes(UserParamStore.GetParameterData(StoreVarWithOffset->Offset, VarSize), VarSize);
-			}
-		}
+	FNiagaraTypeDefinition TypeDef;
+	TypeDef = Variable.Type;
 
-		if (UNiagaraScriptVariable* ScriptVariable = FNiagaraEditorUtilities::UserParameters::GetScriptVariableForUserParameter(Var, System))
+	bool bRemoved = System->GetExposedParameters().RemoveParameter(FNiagaraVariableBase(TypeDef, Variable.Name));
+
+	if (bRemoved)
+	{
+		FPropertyChangedEvent PropertyChangedEvent(nullptr, EPropertyChangeType::ValueSet);
+		System->PostEditChangeProperty(PropertyChangedEvent);
+	}
+	else
+	{
+		// Build list of available user variables for better error message
+		TArray<FString> AvailableVariables;
+		for (const FNiagaraVariableWithOffset& Param : System->GetExposedParameters().ReadParameterVariables())
 		{
-			NewVarInfo.Description = ScriptVariable->Metadata.Description.ToString();
+			AvailableVariables.Add(Param.GetName().ToString());
 		}
+		Context.Error(FText::Format(
+			NSLOCTEXT("NiagaraExternalEdit", "UserVariable_Remove_NotFound", "User variable '{0}' not found in system '{1}'. Available variables: [{2}]"),
+			FText::FromName(Variable.Name),
+			FText::FromString(System->GetName()),
+			FText::FromString(FString::Join(AvailableVariables, TEXT(", ")))));
 	}
 }
 
-void FNiagaraExternalSystemEditorUtilities::GetSystemInfo(UNiagaraSystem& System, FSystemInfo& OutSystemInfo)
+void UNiagaraExternalEditUtilities::AddEmitter(UNiagaraEmitter* TemplateEmitter, FName EmitterName, FNiagaraExt_EmitterTopology& OutTopology, FNiagaraExternalEditContext& Context)
 {
-	OutSystemInfo.Emitters.Reset(System.GetEmitterHandles().Num());
+	UNiagaraSystem* System = Context.GetSystem();
+	TSharedPtr<FNiagaraSystemViewModel> SystemViewModel = Context.GetSystemViewModel();
+	if (System == nullptr || SystemViewModel == nullptr)
+	{
+		return;
+	}
 
-	for (const FNiagaraEmitterHandle& EmitterHandle : System.GetEmitterHandles())
+	if (TemplateEmitter == nullptr)
 	{
-		FText ErrorText;
-		FEmitterInfo& NewInfo = OutSystemInfo.Emitters.AddDefaulted_GetRef();
-		GetEmitterInfo(EmitterHandle, NewInfo);
+		Context.Error(FText::Format(
+			NSLOCTEXT("NiagaraExternalEdit", "Emitter_Add_TemplateNull", "Cannot add emitter '{0}' to system '{1}': template emitter is null."),
+			FText::FromName(EmitterName),
+			FText::FromString(System->GetName())));
+		return;
 	}
-}
 
-FName FNiagaraExternalSystemEditorUtilities::AddEmitterToSystem(UNiagaraSystem& System, UNiagaraEmitter& InEmitterToAdd, FText TransactionDescription)
-{
-	const FScopedTransaction Transaction(TransactionDescription);
-	System.Modify();
+	const FScopedTransaction Transaction(LOCTEXT("ExternalAddEmitter", "External System Add Emitter"));
+	System->Modify();
+
+	FGuid NewEmitterId;
+	if (TemplateEmitter)
+	{
+		NewEmitterId = FNiagaraEditorUtilities::AddEmitterToSystem(*System, *TemplateEmitter, TemplateEmitter->GetExposedVersion().VersionGuid, true);
+	}
+	//TODO: Something janky with adding an empty/default emitter. Need to investigate. For now forceing a template above.
+// 	else
+// 	{
+// 		UNiagaraEmitter* DefaultEmitter = NewObject<UNiagaraEmitter>(System, UNiagaraEmitter::StaticClass(), EmitterName, RF_Transactional);
+// 		UNiagaraEmitterFactoryNew::InitializeEmitter(DefaultEmitter, true);
+// 		System->AddEmitterHandle(*DefaultEmitter, EmitterName, DefaultEmitter->GetExposedVersion().VersionGuid);
+// 	}
 
-	FGuid NewEmitterId = FNiagaraEditorUtilities::AddEmitterToSystem(System, InEmitterToAdd, InEmitterToAdd.GetExposedVersion().VersionGuid, true);
-	if (FNiagaraEmitterHandle* EmitterHandle = GetEmitterHandle(System, NewEmitterId))
+	SystemViewModel->RefreshAll();
+
+	FNiagaraEmitterHandle* EmitterHandle = System->GetEmitterHandles().FindByPredicate([&NewEmitterId](FNiagaraEmitterHandle& Handle) { return Handle.GetId() == NewEmitterId; });
+	if (EmitterHandle)
 	{
+		EmitterHandle->SetName(EmitterName, *System);
+
+		FNiagaraExt_StackItemReference EmitterRef(System);
+		EmitterRef.SetEmitter(EmitterHandle);
+		GetEmitterTopology(EmitterRef, OutTopology, Context);
+
 		//Try to place emitter overview node in a sensible location.
 		TArray<UNiagaraOverviewNode*> OverviewNodes;
-		UNiagaraSystemEditorData* SystemEditorData = CastChecked<UNiagaraSystemEditorData>(System.GetEditorData(), ECastCheckedType::NullChecked);
+		UNiagaraSystemEditorData* SystemEditorData = CastChecked<UNiagaraSystemEditorData>(System->GetEditorData(), ECastCheckedType::NullChecked);
 		SystemEditorData->GetSystemOverviewGraph()->GetNodesOfClass<UNiagaraOverviewNode>(OverviewNodes);
 
 		UNiagaraOverviewNode* RightMostNode = nullptr;
@@ -159,164 +1552,571 @@ FName FNiagaraExternalSystemEditorUtilities::AddEmitterToSystem(UNiagaraSystem&
 			NewNode->NodePosY = RightMostPos.Y;
 		}
 
-		System.OnSystemPostEditChange().Broadcast(&System);
-		return EmitterHandle->GetName();
+		System->OnSystemPostEditChange().Broadcast(System);
+	}
+	else
+	{
+		Context.Error(FText::Format(
+			NSLOCTEXT("NiagaraExternalEdit", "Emitter_Add_Failed", "Failed to add emitter '{0}' from template '{1}' to system '{2}'. The emitter handle was not created."),
+			FText::FromName(EmitterName),
+			FText::FromString(TemplateEmitter ? TemplateEmitter->GetName() : TEXT("No Template Emitter")),
+			FText::FromString(System->GetName())));
 	}
-
-	return NAME_None;
 }
 
-bool FNiagaraExternalSystemEditorUtilities::RemoveEmitterFromSystem(UNiagaraSystem& System, FName EmitterToRemove, FText TransactionDescription)
+void UNiagaraExternalEditUtilities::RemoveEmitter(const FNiagaraExt_StackItemReference& EmitterToRemove, FNiagaraExternalEditContext& Context)
 {
-	if (FNiagaraEmitterHandle* EmitterHandle = GetEmitterHandle(System, EmitterToRemove))
+	UNiagaraSystem* System = Context.GetSystem();
+	if (!Context.CheckSystem(System)) return;
+
+	if (FNiagaraEmitterHandle* EmitterHandle = EmitterToRemove.GetEmitter(Context))
 	{
 		TSet<FGuid> EmitterIdsToRemove;
 		EmitterIdsToRemove.Add(EmitterHandle->GetId());
-		FNiagaraEditorUtilities::KillSystemInstances(System);
+		FNiagaraEditorUtilities::KillSystemInstances(*System);
 
-		const FScopedTransaction Transaction(TransactionDescription);
-		System.Modify();
+		const FScopedTransaction Transaction(LOCTEXT("ExternalRemoveEmitter", "External System Remove Emitter"));
+		System->Modify();
 
-		System.RemoveEmitterHandlesById(EmitterIdsToRemove);
+		System->RemoveEmitterHandlesById(EmitterIdsToRemove);
 
 		FNiagaraScriptMergeManager::Get()->ClearMergeAdapterCache();
-		FNiagaraStackGraphUtilities::RebuildEmitterNodes(System);
-		UNiagaraSystemEditorData* SystemEditorData = CastChecked<UNiagaraSystemEditorData>(System.GetEditorData(), ECastCheckedType::NullChecked);
-		SystemEditorData->SynchronizeOverviewGraphWithSystem(System);
+		FNiagaraStackGraphUtilities::RebuildEmitterNodes(*System);
+		UNiagaraSystemEditorData* SystemEditorData = CastChecked<UNiagaraSystemEditorData>(System->GetEditorData(), ECastCheckedType::NullChecked);
+		SystemEditorData->SynchronizeOverviewGraphWithSystem(*System);
+
+		System->OnSystemPostEditChange().Broadcast(System);
+		return;
+	}
+
+	TArray<FString> AvailableEmitters;
+	for (const FNiagaraEmitterHandle& Handle : System->GetEmitterHandles())
+	{
+		AvailableEmitters.Add(Handle.GetName().ToString());
+	}
+
+	Context.Error(FText::Format(
+		NSLOCTEXT("NiagaraExternalEdit", "Emitter_Remove_NotFound", "Cannot remove emitter '{0}' from system '{1}': emitter not found. Available emitters: [{2}]"),
+		FText::FromName(EmitterToRemove.EmitterName),
+		FText::FromString(System->GetName()),
+		FText::FromString(FString::Join(AvailableEmitters, TEXT(", ")))));
+}
+
+void UNiagaraExternalEditUtilities::AddRenderer(const FNiagaraExt_StackItemReference& NewRendererLocation, const TSubclassOf<UNiagaraRendererProperties> RendererClass, FNiagaraExternalEditContext& Context)
+{
+	// Validate renderer class
+	if (!RendererClass.Get())
+	{
+		Context.Error(NSLOCTEXT("NiagaraExternalEdit", "Renderer_Add_InvalidClass", "Cannot add renderer: renderer class is invalid or null."));
+		return;
+	}
+
+	// Get the emitter handle (error already emitted by GetEmitter if not found)
+	FNiagaraEmitterHandle* EmitterHandle = NewRendererLocation.GetEmitter(Context);
+	if (!EmitterHandle)
+	{
+		return;
+	}
+
+	// Validate emitter instance
+	UNiagaraEmitter* Emitter = EmitterHandle->GetInstance().Emitter;
+	if (!Emitter)
+	{
+		Context.Error(FText::Format(
+			NSLOCTEXT("NiagaraExternalEdit", "Renderer_Add_InvalidEmitter", "Cannot add renderer to emitter '{0}': emitter instance is invalid."),
+			FText::FromName(EmitterHandle->GetName())));
+		return;
+	}
+
+	// Validate emitter data
+	FVersionedNiagaraEmitterData* EmitterData = EmitterHandle->GetEmitterData();
+	if (!EmitterData)
+	{
+		Context.Error(FText::Format(
+			NSLOCTEXT("NiagaraExternalEdit", "Renderer_Add_InvalidEmitterData", "Cannot add renderer to emitter '{0}': emitter data is invalid."),
+			FText::FromName(EmitterHandle->GetName())));
+		return;
+	}
+
+	// Create and add the new renderer
+	int32 NewRendererIndex = EmitterData->GetRenderers().Num();
+	UNiagaraRendererProperties* NewRenderer = NewObject<UNiagaraRendererProperties>(Emitter, RendererClass.Get(), "Renderer");
+	Emitter->AddRenderer(NewRenderer, EmitterData->Version.VersionGuid);
+
+	// Move to specified position if requested
+	if (NewRendererLocation.RendererIndex != INDEX_NONE)
+	{
+		NewRendererIndex = NewRendererLocation.RendererIndex;
+		Emitter->MoveRenderer(NewRenderer, NewRendererIndex, EmitterData->Version.VersionGuid);
+	}
+}
+
+void UNiagaraExternalEditUtilities::RemoveRenderer(const FNiagaraExt_StackItemReference& RendererToRemove, FNiagaraExternalEditContext& Context)
+{
+	// Get the emitter handle (error already emitted by GetEmitter if not found)
+	FNiagaraEmitterHandle* EmitterHandle = RendererToRemove.GetEmitter(Context);
+	if (!EmitterHandle)
+	{
+		return;
+	}
+
+	// Validate emitter instance
+	UNiagaraEmitter* Emitter = EmitterHandle->GetInstance().Emitter;
+	if (!Emitter)
+	{
+		Context.Error(FText::Format(
+			NSLOCTEXT("NiagaraExternalEdit", "Renderer_Remove_InvalidEmitter", "Cannot remove renderer from emitter '{0}': emitter instance is invalid."),
+			FText::FromName(EmitterHandle->GetName())));
+		return;
+	}
 
-		System.OnSystemPostEditChange().Broadcast(&System);
+	// Get the renderer (error already emitted by GetRenderer if not found)
+	UNiagaraRendererProperties* Renderer = RendererToRemove.GetRenderer(Context);
+	if (!Renderer)
+	{
+		return;
+	}
 
-		return true;
+	// Validate emitter data
+	FVersionedNiagaraEmitterData* EmitterData = EmitterHandle->GetEmitterData();
+	if (!EmitterData)
+	{
+		Context.Error(FText::Format(
+			NSLOCTEXT("NiagaraExternalEdit", "Renderer_Remove_InvalidEmitterData", "Cannot remove renderer from emitter '{0}': emitter data is invalid."),
+			FText::FromName(EmitterHandle->GetName())));
+		return;
 	}
-	return false;
+
+	// Remove the renderer
+	Emitter->RemoveRenderer(Renderer, EmitterData->Version.VersionGuid);
 }
 
-void FNiagaraExternalSystemEditorUtilities::GetEmitterInfo(const FNiagaraEmitterHandle& EmitterHandle, FEmitterInfo& OutEmitterInfo)
+void UNiagaraExternalEditUtilities::AddModule(const FNiagaraExt_StackItemReference& NewModuleLocationRef, const UNiagaraScript* ModuleAsset, FNiagaraExt_ModuleTopology& OutTopology, FNiagaraExternalEditContext& Context)
 {
-	OutEmitterInfo.Name = EmitterHandle.GetName();
-	OutEmitterInfo.Id = EmitterHandle.GetId();
-	OutEmitterInfo.EmitterMode = EmitterHandle.GetEmitterMode();
-	if (EmitterHandle.GetEmitterMode() == ENiagaraEmitterMode::Stateless)
+	if (!Context.CheckScriptAsset(ModuleAsset)) return;
+
+	UNiagaraSystem* System = NewModuleLocationRef.GetSystem(Context);
+	UNiagaraStackScriptItemGroup* ScriptItem = NewModuleLocationRef.GetScript(Context);
+
+	if (System == nullptr || ScriptItem == nullptr)
+	{
+		return;
+	}
+
+	UNiagaraStackModuleItem* PreviousModule = NewModuleLocationRef.GetModule(Context);
+	UNiagaraNodeOutput* TargetOutputNode = ScriptItem->GetScriptOutputNode();
+	int32 TargetIndex = INDEX_NONE;
+	if (PreviousModule)
+	{
+		//TargetOutputNode = PreviousModule->GetOutputNode();
+		TargetIndex = PreviousModule->GetModuleIndex() + 1;
+	}
+
+	if (TargetOutputNode)
 	{
-		if (const UNiagaraStatelessEmitter* StatelessEmitter = EmitterHandle.GetStatelessEmitter())
+		const FScopedTransaction Transaction(LOCTEXT("ExternalAddModule", "External System Add Module"));
+		System->Modify();
+
+		FNiagaraStackGraphUtilities::FAddScriptModuleToStackArgs Args(ModuleAsset, *TargetOutputNode);
+		Args.TargetIndex = TargetIndex;
+		Args.VersionGuid = ModuleAsset->GetLatestScriptData()->Version.VersionGuid;//TODO: allow old versions?
+		if (UNiagaraNodeFunctionCall* Func = FNiagaraStackGraphUtilities::AddScriptModuleToStack(Args))
 		{
-			if (const FNiagaraStatelessEmitterDataPtr EmitterData = StatelessEmitter ? StatelessEmitter->GetEmitterData() : nullptr)
+			//Refresh the query.
+			ScriptItem->RefreshChildren();
+			TNiagaraStackQueryResult<UNiagaraStackModuleItem> AddModuleResult = FNiagaraStackScriptItemGroupQuery(*ScriptItem).FindModuleItem(Func->GetFunctionName()).GetResult();
+			if (AddModuleResult.IsValid())
 			{
-				OutEmitterInfo.bLocalSpace = true;
-				OutEmitterInfo.bDeterministic = EmitterData->bDeterministic;
-				OutEmitterInfo.SimTarget = ENiagaraSimTarget::GPUComputeSim;
+				//GetStackModuleInfo(AddModuleResult.StackEntry, OutInfo, OutErrorMessage);
+				FPropertyChangedEvent PropertyChangedEvent(nullptr, EPropertyChangeType::ValueSet);
+				System->PostEditChangeProperty(PropertyChangedEvent);
+
+				FNiagaraExt_StackItemReference NewModuleReference = NewModuleLocationRef;
+				NewModuleReference.SetModule(AddModuleResult.StackEntry);
+				GetModuleTopology(NewModuleReference, OutTopology, Context);
 			}
-		}		
+		}
+		else
+		{
+			Context.Error(FText::Format(
+				NSLOCTEXT("NiagaraExternalEdit", "Module_Add_Failed", "Failed to add module '{0}' to stack reference: {1}. The module asset may be corrupt or invalid."),
+				FText::FromString(ModuleAsset->GetName()),
+				FText::FromString(GetStackItemPath(NewModuleLocationRef))));
+		}
 	}
 	else
 	{
-		const FVersionedNiagaraEmitterData* VersionedData = EmitterHandle.GetEmitterData();
-		if (VersionedData)
-		{
-			OutEmitterInfo.bLocalSpace = VersionedData->bLocalSpace;
-			OutEmitterInfo.bDeterministic = VersionedData->bDeterminism;
-			OutEmitterInfo.SimTarget = VersionedData->SimTarget;
-		}
+		Context.Error(FText::Format(
+			NSLOCTEXT("NiagaraExternalEdit", "Module_Add_NoOutputNode", "Failed to add module '{0}': script output node not found in stack reference: {1}"),
+			FText::FromString(ModuleAsset->GetName()),
+			FText::FromString(GetStackItemPath(NewModuleLocationRef))));
 	}
 }
 
-bool FNiagaraExternalSystemEditorUtilities::GetEmitterInfo(UNiagaraSystem& System, FName EmitterName, FEmitterInfo& OutEmitterInfo, FText& OutErrorMessage)
+void UNiagaraExternalEditUtilities::RemoveModule(const FNiagaraExt_StackItemReference& ModuleToRemove, FNiagaraExternalEditContext& Context)
 {
-	if (FNiagaraEmitterHandle* EmitterHandle = GetEmitterHandle(System, EmitterName))
+	UNiagaraSystem* System = ModuleToRemove.GetSystem(Context);
+	if (!Context.CheckSystem(System)) return;
+
+	FNiagaraEmitterHandle* EmitterHandle = ModuleToRemove.GetEmitter(Context, false);
+	UNiagaraStackModuleItem* ModuleItemToRemove = ModuleToRemove.GetModule(Context);
+
+	if (System == nullptr || ModuleItemToRemove == nullptr)
 	{
-		GetEmitterInfo(*EmitterHandle, OutEmitterInfo);
-		return true;
+		return;
 	}
-	OutErrorMessage = LOCTEXT("InvalidEmitter", "An emitter with this name does not exist on this System.");
-	return false;
+
+	const FScopedTransaction Transaction(LOCTEXT("ExternalRemoveModule", "External System Remove Module"));
+	System->Modify();
+	FGuid EmitterId = EmitterHandle ? EmitterHandle->GetId() : FGuid();
+	TArray<TWeakObjectPtr<UNiagaraNodeInput>> RemovedInputNodes;
+	FNiagaraStackGraphUtilities::RemoveModuleFromStack(*System, EmitterId, ModuleItemToRemove->GetModuleNode(), RemovedInputNodes);
+	ModuleToRemove.GetScript(Context)->RefreshChildren();
+	FPropertyChangedEvent PropertyChangedEvent(nullptr, EPropertyChangeType::ValueSet);
+	System->PostEditChangeProperty(PropertyChangedEvent);
 }
 
-bool FNiagaraExternalSystemEditorUtilities::SetEmitterName(UNiagaraSystem& System, FName EmitterName, FName Name, FText TransactionDescription, FText& OutErrorMessage)
+void UNiagaraExternalEditUtilities::SetModuleEnabled(const FNiagaraExt_StackItemReference& ModuleRef, bool bEnabled, FNiagaraExternalEditContext& Context)
 {
-	if (FNiagaraEmitterHandle* EmitterHandle = GetEmitterHandle(System, EmitterName))
+	UNiagaraSystem* System = ModuleRef.GetSystem(Context);
+	UNiagaraStackModuleItem* ModuleItem = ModuleRef.GetModule(Context);
+	if (System == nullptr || ModuleItem == nullptr)
 	{
-		const FScopedTransaction Transaction(TransactionDescription);
-		System.Modify();
-		EmitterHandle->SetName(Name, System);
-		System.OnSystemPostEditChange().Broadcast(&System);
-		return true;
+		return;
 	}
-	OutErrorMessage = LOCTEXT("InvalidEmitter", "An emitter with this name does not exist on this System.");
-	return false;
+
+	const FScopedTransaction Transaction(LOCTEXT("ExternalSetModuleEnabled", "External System Set Module Enabled"));
+	System->Modify();
+	ModuleItem->SetIsEnabled(bEnabled);
+}
+
+void UNiagaraExternalEditUtilities::SetSystemData(UNiagaraSystem* System, const FNiagaraExt_SystemData& InData, FNiagaraExternalEditContext& Context)
+{
+	if (!Context.CheckSystem(System)) return;
+	SetAllObjectProperties(System, InData.PropertyValues, Context.Errors);
 }
 
-bool FNiagaraExternalSystemEditorUtilities::SetEmitterLocalSpace(UNiagaraSystem& System, FName EmitterName, bool bLocalSpace, FText TransactionDescription, FText& OutErrorMessage)
+void UNiagaraExternalEditUtilities::SetEmitterData(const FNiagaraExt_StackItemReference& EmitterRef, const FNiagaraExt_EmitterData& InData, FNiagaraExternalEditContext& Context)
 {
-	if (FNiagaraEmitterHandle* EmitterHandle = GetEmitterHandle(System, EmitterName))
+	if (FNiagaraEmitterHandle* Emitter = EmitterRef.GetEmitter(Context))
 	{
-		const FScopedTransaction Transaction(TransactionDescription);
-		System.Modify();
-		if (EmitterHandle->GetEmitterMode() == ENiagaraEmitterMode::Stateless)
+		UNiagaraSystem* System = EmitterRef.GetSystem(Context);
+		if (!System) return;
+
+		// Find the emitter index
+		int32 EmitterIndex = System->GetEmitterHandles().IndexOfByPredicate([&](const FNiagaraEmitterHandle& Handle) {
+			return Handle.GetId() == Emitter->GetId();
+		});
+
+		if (EmitterIndex != INDEX_NONE)
 		{
-			OutErrorMessage = LOCTEXT("SetStatelessLocalSpace", "Stateless Emitters are always local space.");
+			SetEmitterHandleAndDataProperties(System, EmitterIndex, InData.PropertyValues, Context.Errors);
 		}
 		else
 		{
-			FVersionedNiagaraEmitterData* VersionedData = EmitterHandle->GetEmitterData();
-			if (VersionedData)
+			// Build list of available emitters for better error message
+			TArray<FString> AvailableEmitters;
+			for (const FNiagaraEmitterHandle& Handle : System->GetEmitterHandles())
+			{
+				AvailableEmitters.Add(Handle.GetName().ToString());
+			}
+			Context.Error(FText::Format(
+				NSLOCTEXT("NiagaraExternalEdit", "Emitter_Data_Set_NotFound", "Cannot set emitter data: emitter not found in system '{0}'. Stack reference: {1}. Available emitters: [{2}]"),
+				FText::FromString(System->GetName()),
+				FText::FromString(GetStackItemPath(EmitterRef)),
+				FText::FromString(FString::Join(AvailableEmitters, TEXT(", ")))));
+		}
+	}
+}
+
+void UNiagaraExternalEditUtilities::SetRendererData(const FNiagaraExt_StackItemReference& RendererRef, const FNiagaraExt_RendererData& InData, FNiagaraExternalEditContext& Context)
+{
+	if (UNiagaraRendererProperties* RendererProps = RendererRef.GetRenderer(Context))
+	{
+		SetAllObjectProperties(RendererProps, InData.PropertyValues, Context.Errors);
+	}
+}
+
+void UNiagaraExternalEditUtilities::SetStackInputData(const FNiagaraExt_StackItemReference& StackInputRef, const FNiagaraExt_StackInputValue& InData, FNiagaraExternalEditContext& Context)
+{
+	InData.InitStackInputFromValue(StackInputRef, Context);
+}
+
+//////////////////////////////////////////////////////////////////////////
+
+void FNiagaraExt_VariableValue::Set(const FNiagaraTypeDefinition& TypeDef, const FNiagaraVariant& Variant)
+{
+	if (Variant.IsValid(TypeDef))
+	{
+		if (UEnum* Enum = TypeDef.GetEnum())
+		{
+			FNiagaraExt_VariableValue_Enum& EnumVal = InitializeAs<FNiagaraExt_VariableValue_Enum>();
+			int32 EnumValue = Variant.GetBytesValue<int32>(TypeDef);
+			EnumVal.Enum = Enum;
+			if (Enum)
 			{
-				VersionedData->bLocalSpace = bLocalSpace;
+				EnumVal.EnumName = Enum->GetNameByValue(EnumValue);
 			}
 		}
-		System.OnSystemPostEditChange().Broadcast(&System);
-		return true;
+		else if (TypeDef.IsDataInterface())
+		{
+			FNiagaraExt_VariableValue_DataInterface& DIValue = InitializeAs<FNiagaraExt_VariableValue_DataInterface>();
+			DIValue.DataInterfaceClass = TypeDef.GetClass();
+			DIValue.DataInterface = Variant.GetMode() == ENiagaraVariantMode::DataInterface ? Cast<UNiagaraDataInterface>(Variant.GetDataInterface()) : nullptr;
+		}
+		else if (TypeDef.IsUObject())
+		{
+			FNiagaraExt_VariableValue_Object& ObjValue = InitializeAs<FNiagaraExt_VariableValue_Object>();
+			ObjValue.ObjectClass = TypeDef.GetClass();
+			ObjValue.Object = Variant.GetMode() == ENiagaraVariantMode::Object ? Variant.GetUObject() : nullptr;
+		}
+		else if (UScriptStruct* InnerScriptStruct = TypeDef.GetScriptStruct())
+		{
+			InitializeAs(InnerScriptStruct, Variant.GetMode() == ENiagaraVariantMode::Bytes ? Variant.GetBytes() : nullptr);
+		}
 	}
-	OutErrorMessage = LOCTEXT("InvalidEmitter", "An emitter with this name does not exist on this System.");
-	return false;
 }
 
-bool FNiagaraExternalSystemEditorUtilities::SetEmitterDeterministic(UNiagaraSystem& System, FName EmitterName, bool bDeterministic, FText TransactionDescription, FText& OutErrorMessage)
+void FNiagaraExt_VariableValue::Get(FNiagaraVariant& Variant, FNiagaraExternalEditContext& Context)const
 {
-	if (FNiagaraEmitterHandle* EmitterHandle = GetEmitterHandle(System, EmitterName))
+	if(!IsValid())
 	{
-		const FScopedTransaction Transaction(TransactionDescription);
-		System.Modify();
-		if (EmitterHandle->GetEmitterMode() == ENiagaraEmitterMode::Stateless)
+		Context.Error(NSLOCTEXT("NiagaraExternalEdit", "VariableValue_Get_Invalid", "Cannot get variable value: data is invalid or uninitialized."));
+		return;
+	}
+
+	if (const FNiagaraExt_VariableValue_Enum* EnumValue = GetPtr<FNiagaraExt_VariableValue_Enum>())
+	{
+		if (UEnum* Enum = EnumValue->Enum)
 		{
-			OutErrorMessage = LOCTEXT("SetStatelessDeterministic", "Modiftying this value for stateless emitters is not currently supported.");
+			int32 NewEnumValue = Enum->GetValueByName(EnumValue->EnumName);
+			Variant.SetBytes((uint8*)&NewEnumValue, sizeof(int32));
 		}
-		else
+	}
+	else if (const FNiagaraExt_VariableValue_DataInterface* DIValue = GetPtr<FNiagaraExt_VariableValue_DataInterface>())
+	{
+		Variant.SetDataInterface(DIValue->DataInterface);
+	}
+	else if (const FNiagaraExt_VariableValue_Object* ObjValue = GetPtr<FNiagaraExt_VariableValue_Object>())
+	{
+		Variant.SetUObject(ObjValue->Object);
+	}
+	else if (IsValid())
+	{
+		Variant.SetBytes(GetMemory(), GetScriptStruct()->GetStructureSize());
+	}
+}
+
+//////////////////////////////////////////////////////////////////////////
+
+void FNiagaraExt_StackInputValue::InitFromStackInput(const FNiagaraExt_StackItemReference& StackInputRef, FNiagaraExternalEditContext& Context)
+{
+	UNiagaraStackFunctionInput* Input = StackInputRef.GetInput(Context);
+	UNiagaraSystem* System = StackInputRef.GetSystem(Context);
+
+	if(System == nullptr || Input == nullptr)
+	{
+		return;
+	}
+
+	const FNiagaraTypeDefinition& InputType = Input->GetInputType();
+
+	switch (Input->GetValueMode())
+	{
+	case UNiagaraStackFunctionInput::EValueMode::Local:
+	{
+		TSharedPtr<const FStructOnScope> LocalData = Input->GetLocalValueStruct();
+		const UScriptStruct* InScriptStruct = LocalData.IsValid() ? Cast<const UScriptStruct>(LocalData->GetStruct()) : nullptr;
+		if (LocalData.IsValid() && InScriptStruct && LocalData->IsValid())
 		{
-			FVersionedNiagaraEmitterData* VersionedData = EmitterHandle->GetEmitterData();
-			if (VersionedData)
+			if (UEnum* Enum = InputType.GetEnum())
 			{
-				VersionedData->bDeterminism = bDeterministic;
+				FNiagaraExt_StackInputData_Enum& EnumVal = InitializeAs<FNiagaraExt_StackInputData_Enum>();
+				EnumVal.Enum = Enum;
+				if (Enum && ensure(InScriptStruct == FNiagaraInt32::StaticStruct()))
+				{
+					int32 EnumValue = *(int32*)LocalData->GetStructMemory();
+					EnumVal.EnumName = Enum->GetNameByValue(EnumValue);
+				}
 			}
+			else
+			{
+				InitializeAs(InScriptStruct, LocalData->GetStructMemory());
+			}
+		}
+		else
+		{
+			Context.Error(FText::Format(
+				NSLOCTEXT("NiagaraExternalEdit", "StackInput_Get_LocalDataInvalid", "Local data for input '{0}' is invalid or struct type could not be determined. Stack reference: {1}"),
+				FText::FromName(Input->GetInputParameterHandle().GetName()),
+				FText::FromString(GetStackItemPath(StackInputRef))));
+		}
+	}
+	break;
+	case UNiagaraStackFunctionInput::EValueMode::Linked:
+	{
+		FNiagaraExt_StackInputData_Linked& Value = InitializeAs<FNiagaraExt_StackInputData_Linked>();
+		Value.LinkedVariable.Name = Input->GetLinkedParameterValue().GetName();
+		Value.LinkedVariable.Type = Input->GetLinkedParameterValue().GetType();
+	}
+	break;
+	case UNiagaraStackFunctionInput::EValueMode::Expression:
+	{
+		FNiagaraExt_StackInputData_HlslExpression& Value = InitializeAs<FNiagaraExt_StackInputData_HlslExpression>();
+		Value.HlslExpression = Input->GetCustomExpressionText().ToString();
+	}
+	break;
+	case UNiagaraStackFunctionInput::EValueMode::Data:
+	{
+		if (UNiagaraDataInterface* DI = Input->GetDataValueObject())
+		{
+			FNiagaraExt_StackInputData_DataInterface& Value = InitializeAs<FNiagaraExt_StackInputData_DataInterface>();
+			Value.PropertyValues = GetAllObjectProperties(DI);
+			Context.UsedDataInterfaces.Add(DI->GetClass());
 		}
-		System.OnSystemPostEditChange().Broadcast(&System);
-		return true;
 	}
-	OutErrorMessage = LOCTEXT("InvalidEmitter", "An emitter with this name does not exist on this System.");
-	return false;
+	break;
+	case UNiagaraStackFunctionInput::EValueMode::Dynamic:
+	{
+		FNiagaraExt_StackInputData_DynamicInput& Value = InitializeAs<FNiagaraExt_StackInputData_DynamicInput>();
+		Value.DynamicInputAsset = Input->GetDynamicInputNode()->FunctionScript;
+		Context.UsedDynamicInputs.Add(Value.DynamicInputAsset);
+	}
+	break;
+	default:
+		InitializeAs<FNiagaraExt_StackInputData_Unsupported>();
+	}
 }
-bool FNiagaraExternalSystemEditorUtilities::SetEmitterSimTarget(UNiagaraSystem& System, FName EmitterName, ENiagaraSimTarget SimTarget, FText TransactionDescription, FText& OutErrorMessage)
+
+void FNiagaraExt_StackInputValue::InitStackInputFromValue(const FNiagaraExt_StackItemReference& StackInputRef, FNiagaraExternalEditContext& Context)const
 {
-	if (FNiagaraEmitterHandle* EmitterHandle = GetEmitterHandle(System, EmitterName))
+	UNiagaraStackFunctionInput* Input = StackInputRef.GetInput(Context);
+	UNiagaraSystem* System = StackInputRef.GetSystem(Context);
+
+	if(System == nullptr || Input == nullptr)
+	{
+		return;
+	}
+	
+	const FNiagaraTypeDefinition& InputType = Input->GetInputType();
+	UScriptStruct* InputTypeStruct = InputType.GetScriptStruct();
+	
+	if (const FNiagaraExt_StackInputData_Enum* EnumValue = GetPtr<FNiagaraExt_StackInputData_Enum>())
+	{
+		if(UEnum* Enum = EnumValue->Enum)
+		{			
+			int32 NewEnumValue = Enum->GetValueByName(EnumValue->EnumName);
+			TSharedPtr<FStructOnScope> StructOnScope = MakeShared<FStructOnScope>(FNiagaraInt32::StaticStruct());
+			FMemory::Memcpy(StructOnScope->GetStructMemory(), &NewEnumValue, sizeof(FNiagaraInt32));
+			Input->SetLocalValue(StructOnScope.ToSharedRef());
+		}
+	}
+	else if (const FNiagaraExt_StackInputData_Linked* LinkedValue = GetPtr<FNiagaraExt_StackInputData_Linked>())
+	{
+		Input->SetLinkedParameterValue(FNiagaraVariableBase(LinkedValue->LinkedVariable.Type, LinkedValue->LinkedVariable.Name));
+	}
+	else if (const FNiagaraExt_StackInputData_HlslExpression* HlslValue = GetPtr<FNiagaraExt_StackInputData_HlslExpression>())
+	{
+		if (Input->SupportsCustomExpressions())
+		{
+			Input->SetCustomExpression(HlslValue->HlslExpression);
+		}
+		else
+		{
+			Context.Error(FText::Format(
+				NSLOCTEXT("NiagaraExternalEdit", "StackInput_Set_CustomExpressionNotSupported", "Input '{0}' does not support custom HLSL expressions. Stack reference: {1}"),
+				FText::FromName(Input->GetInputParameterHandle().GetName()),
+				FText::FromString(GetStackItemPath(StackInputRef))));
+		}
+	}
+	else if (const FNiagaraExt_StackInputData_DataInterface* DIValue = GetPtr<FNiagaraExt_StackInputData_DataInterface>())
 	{
-		const FScopedTransaction Transaction(TransactionDescription);
-		System.Modify();
-		if (EmitterHandle->GetEmitterMode() == ENiagaraEmitterMode::Stateless)
+		UClass* DIClass = Input->GetInputType().GetClass();
+		if(DIClass && DIClass->IsChildOf<UNiagaraDataInterface>())
 		{
-			OutErrorMessage = LOCTEXT("SetStatelessSimTarget", "Stateless Emitters determine their own Sim Target.");
+			Input->SetDataInterfaceValueExternal(DIClass, [&](UNiagaraDataInterface* DI)
+				{
+					if (DI)
+					{
+						if (!SetAllObjectProperties(DI, DIValue->PropertyValues, Context.Errors))
+						{
+							Context.Error(FText::Format(
+								NSLOCTEXT("NiagaraExternalEdit", "StackInput_Set_DataInterface_PropertiesFailed", "Failed to set data interface properties for input '{0}'. Stack reference: {1}"),
+								FText::FromName(Input->GetInputParameterHandle().GetName()),
+								FText::FromString(GetStackItemPath(StackInputRef))));
+						}
+					}
+					else
+					{
+						Context.Error(FText::Format(
+							NSLOCTEXT("NiagaraExternalEdit", "StackInput_Set_DataInterface_CreateFailed", "Failed to create data interface for input '{0}'. Stack reference: {1}"),
+							FText::FromName(Input->GetInputParameterHandle().GetName()),
+							FText::FromString(GetStackItemPath(StackInputRef))));
+					}
+				});
 		}
 		else
 		{
-			FVersionedNiagaraEmitterData* VersionedData = EmitterHandle->GetEmitterData();
-			if (VersionedData)
+			Context.Error(FText::Format(
+				NSLOCTEXT("NiagaraExternalEdit", "StackInput_Set_NotDataInterface", "Input '{0}' is not a data interface type (actual type: {1}). Stack reference: {2}"),
+				FText::FromName(Input->GetInputParameterHandle().GetName()),
+				FText::FromString(InputType.GetName()),
+				FText::FromString(GetStackItemPath(StackInputRef))));
+		}
+	}
+	else if (const FNiagaraExt_StackInputData_DynamicInput* DynamicValue = GetPtr<FNiagaraExt_StackInputData_DynamicInput>())
+	{
+		Input->SetDynamicInput(DynamicValue->DynamicInputAsset);
+	}
+	else if(IsValid())
+	{
+		FNiagaraTypeDefinition IncommingTypeDef(const_cast<UScriptStruct*>(ScriptStruct.Get()), FNiagaraTypeDefinition::EAllowUnfriendlyStruct::Allow);//Ugh. Todo. Propagate constness into the NiagaraTypeDefinition Struct.
+		if (InputType != IncommingTypeDef)
+		{
+			//Allow int32 inputs for enums.
+			bool bAllowDifferentType = InputType.IsEnum() && IncommingTypeDef == FNiagaraTypeDefinition::GetIntDef();
+			if (!bAllowDifferentType)
 			{
-				VersionedData->SimTarget = SimTarget;
+				Context.Error(FText::Format(
+					NSLOCTEXT("NiagaraExternalEdit", "StackInput_Set_TypeMismatch", "Type mismatch for input '{0}': expected '{1}' but got '{2}'. Stack reference: {3}"),
+					FText::FromName(Input->GetInputParameterHandle().GetName()),
+					FText::FromString(InputType.GetName()),
+					FText::FromString(IncommingTypeDef.GetName()),
+					FText::FromString(GetStackItemPath(StackInputRef))));
+				return;
 			}
 		}
-		System.OnSystemPostEditChange().Broadcast(&System);
-		return true;
+
+		TSharedPtr<FStructOnScope> StructOnScope = MakeShared<FStructOnScope>(ScriptStruct);
+		FMemory::Memcpy(StructOnScope->GetStructMemory(), GetMemory(), ScriptStruct->GetStructureSize());
+		Input->SetLocalValue(StructOnScope.ToSharedRef());
+	}
+	else
+	{
+		Context.Error(FText::Format(
+			NSLOCTEXT("NiagaraExternalEdit", "StackInput_Set_InvalidData", "Cannot set input '{0}': input data is invalid or uninitialized. Stack reference: {1}"),
+			FText::FromName(Input->GetInputParameterHandle().GetName()),
+			FText::FromString(GetStackItemPath(StackInputRef))));
 	}
-	OutErrorMessage = LOCTEXT("InvalidEmitter", "An emitter with this name does not exist on this System.");
-	return false;
 }
 
+
+
+//////////////////////////////////////////////////////////////////////////
+
+FNiagaraExt_DynamicInputTopology::FNiagaraExt_DynamicInputTopology()
+{
+	InitializeAs<FNiagaraExt_StackInputTopology>();
+}
+
+const FNiagaraExt_StackInputTopology& FNiagaraExt_DynamicInputTopology::Get()const
+{
+	return FInstancedStruct::Get<FNiagaraExt_StackInputTopology>();
+}
+
+FNiagaraExt_StackInputTopology& FNiagaraExt_DynamicInputTopology::GetMutable()
+{
+	return FInstancedStruct::GetMutable<FNiagaraExt_StackInputTopology>();
+}
 #undef LOCTEXT_NAMESPACE
+
diff --git a/Engine/Plugins/FX/Niagara/Source/NiagaraEditor/Private/ViewModels/NiagaraSystemViewModel.cpp b/Engine/Plugins/FX/Niagara/Source/NiagaraEditor/Private/ViewModels/NiagaraSystemViewModel.cpp
--- a/Engine/Plugins/FX/Niagara/Source/NiagaraEditor/Private/ViewModels/NiagaraSystemViewModel.cpp
+++ b/Engine/Plugins/FX/Niagara/Source/NiagaraEditor/Private/ViewModels/NiagaraSystemViewModel.cpp
@@ -512,8 +512,11 @@ ENiagaraSystemViewModelEditMode FNiagaraSystemViewModel::GetEditMode() const
 
 FName FNiagaraSystemViewModel::GetWorkflowMode() const
 {
-	check(OnGetWorkflowModeDelegate.IsBound());
-	return OnGetWorkflowModeDelegate.Execute();
+	if(OnGetWorkflowModeDelegate.IsBound())
+	{
+		return OnGetWorkflowModeDelegate.Execute();
+	}
+	return NAME_None;	
 }
 
 FOnGetWorkflowMode& FNiagaraSystemViewModel::OnGetWorkflowMode()
@@ -523,8 +526,10 @@ FOnGetWorkflowMode& FNiagaraSystemViewModel::OnGetWorkflowMode()
 
 void FNiagaraSystemViewModel::SetWorkflowMode(FName WorkflowMode)
 {
-	check(OnChangeWorkflowModeDelegate.IsBound());
-	OnChangeWorkflowModeDelegate.Execute(WorkflowMode);
+	if(OnChangeWorkflowModeDelegate.IsBound())
+	{
+		OnChangeWorkflowModeDelegate.Execute(WorkflowMode);
+	}
 }
 
 FNiagaraSystemViewModel::FOnChangeWorkflowMode& FNiagaraSystemViewModel::OnChangeWorkflowMode()
diff --git a/Engine/Plugins/FX/Niagara/Source/NiagaraEditor/Private/ViewModels/Stack/NiagaraStackFunctionInput.cpp b/Engine/Plugins/FX/Niagara/Source/NiagaraEditor/Private/ViewModels/Stack/NiagaraStackFunctionInput.cpp
--- a/Engine/Plugins/FX/Niagara/Source/NiagaraEditor/Private/ViewModels/Stack/NiagaraStackFunctionInput.cpp
+++ b/Engine/Plugins/FX/Niagara/Source/NiagaraEditor/Private/ViewModels/Stack/NiagaraStackFunctionInput.cpp
@@ -2375,6 +2375,40 @@ void UNiagaraStackFunctionInput::SetLocalValue(TSharedRef<FStructOnScope> InLoca
 	}
 }
 
+void UNiagaraStackFunctionInput::SetDataInterfaceValueExternal(TSubclassOf<UNiagaraDataInterface> InDataInterfaceType, TFunction<void(UNiagaraDataInterface*)> ExternalSet)
+{
+	if(ensureMsgf(InDataInterfaceType && InDataInterfaceType->IsChildOf(InputType.GetClass()), TEXT("Can not set an input to an unrelated type.")))
+	{
+		TGuardValue<bool> UpdateGuard(bUpdatingLocalValueDirectly, true);
+		FScopedTransaction ScopedTransaction(LOCTEXT("UpdateInputLocalDataValue", "Update input local data interface value"));
+		bool bGraphWillNeedRelayout = false;
+		UEdGraphPin* OverridePin = GetOverridePin();
+
+		if (OverridePin != nullptr && OverridePin->LinkedTo.Num() > 0)
+		{
+			// If there is an override pin and it's linked we'll need to remove all of the linked nodes to set a local value.
+			RemoveNodesForOverridePin(*OverridePin);
+			bGraphWillNeedRelayout = true;
+		}
+
+		if (OverridePin == nullptr)
+		{
+			OverridePin = &GetOrCreateOverridePin();
+		}
+
+		UNiagaraDataInterface* NewDataInterface = nullptr;
+		FNiagaraStackGraphUtilities::SetDataInterfaceValueForFunctionInput(*OverridePin, InDataInterfaceType, AliasedInputParameterHandle.GetParameterHandleString().ToString(), NewDataInterface);
+
+		RefreshValues();
+
+		if (NewDataInterface && ExternalSet)
+		{
+			ExternalSet(NewDataInterface);
+			RefreshValues();
+		}
+	}
+}
+
 void UNiagaraStackFunctionInput::SetDataInterfaceValue(TSubclassOf<UNiagaraDataInterface> InDataInterfaceType)
 {
 	checkf(InDataInterfaceType->IsChildOf(InputType.GetClass()), TEXT("Can not set an input to an unrelated type."));
@@ -2404,6 +2438,7 @@ void UNiagaraStackFunctionInput::SetDataInterfaceValue(TSubclassOf<UNiagaraDataI
 
 	UNiagaraDataInterface* InputDataInterface = nullptr;
 	FNiagaraStackGraphUtilities::SetDataInterfaceValueForFunctionInput(*OverridePin, InDataInterfaceType, AliasedInputParameterHandle.GetParameterHandleString().ToString(), InputDataInterface);
+	RefreshValues();
 }
 
 bool UNiagaraStackFunctionInput::CanReset() const
@@ -2918,6 +2953,16 @@ void UNiagaraStackFunctionInput::SetEditConditionEnabled(bool bIsEnabled)
 	}
 }
 
+TOptional<FNiagaraVariable> UNiagaraStackFunctionInput::GetVisibleConditionVariable() const
+{
+	if(GetHasVisibleCondition())
+	{
+		return FNiagaraVariable(VisibleCondition.GetConditionInputType(), VisibleCondition.GetConditionInputName());
+	}
+
+	return TOptional<FNiagaraVariable>();
+}
+
 bool UNiagaraStackFunctionInput::GetHasVisibleCondition() const
 {
 	return VisibleCondition.IsValid();
diff --git a/Engine/Plugins/FX/Niagara/Source/NiagaraEditor/Private/ViewModels/Stack/NiagaraStackGraphUtilities.cpp b/Engine/Plugins/FX/Niagara/Source/NiagaraEditor/Private/ViewModels/Stack/NiagaraStackGraphUtilities.cpp
--- a/Engine/Plugins/FX/Niagara/Source/NiagaraEditor/Private/ViewModels/Stack/NiagaraStackGraphUtilities.cpp
+++ b/Engine/Plugins/FX/Niagara/Source/NiagaraEditor/Private/ViewModels/Stack/NiagaraStackGraphUtilities.cpp
@@ -1033,7 +1033,7 @@ FString FNiagaraStackGraphUtilities::StackKeys::GenerateStackRendererEditorDataK
 	return FString::Printf(TEXT("Renderer-%s"), *Renderer.GetName());
 }
 
-void ExtractInputPinsFromHistory(FNiagaraParameterMapHistory& History, UEdGraph* FunctionGraph, FNiagaraStackGraphUtilities::ENiagaraGetStackFunctionInputPinsOptions Options, TArray<const UEdGraphPin*>& OutPins)
+void FNiagaraStackGraphUtilities::ExtractInputPinsFromHistory(FNiagaraParameterMapHistory& History, UEdGraph* FunctionGraph, FNiagaraStackGraphUtilities::ENiagaraGetStackFunctionInputPinsOptions Options, TArray<const UEdGraphPin*>& OutPins)
 {
 	const int32 ModuleNamespaceLength = FNiagaraConstants::ModuleNamespaceString.Len();
 
@@ -2320,42 +2320,48 @@ void FNiagaraStackGraphUtilities::SetCustomExpressionForFunctionInput(UEdGraphPi
 {
 	checkf(OverridePin.LinkedTo.Num() == 0, TEXT("Can't set a data value when the override pin already has a value."));
 
-	UNiagaraNodeParameterMapSet* OverrideNode = CastChecked<UNiagaraNodeParameterMapSet>(OverridePin.GetOwningNode());
-	UEdGraph* Graph = OverrideNode->GetGraph();
-	const UEdGraphSchema_Niagara* Schema = CastChecked<UEdGraphSchema_Niagara>(OverrideNode->GetSchema());
+	if(UNiagaraNodeParameterMapSet* OverrideNode = Cast<UNiagaraNodeParameterMapSet>(OverridePin.GetOwningNode()))
+	{
+		UEdGraph* Graph = OverrideNode->GetGraph();
+		const UEdGraphSchema_Niagara* Schema = CastChecked<UEdGraphSchema_Niagara>(OverrideNode->GetSchema());
 
-	Graph->Modify();
-	FGraphNodeCreator<UNiagaraNodeCustomHlsl> FunctionCallNodeCreator(*Graph);
-	UNiagaraNodeCustomHlsl* FunctionCallNode = FunctionCallNodeCreator.CreateNode();
-	FunctionCallNode->InitAsCustomHlslDynamicInput(Schema->PinToTypeDefinition(&OverridePin));
-	FunctionCallNodeCreator.Finalize();
-	FunctionCallNode->SetEnabledState(OverrideNode->GetDesiredEnabledState(), OverrideNode->HasUserSetTheEnabledState());
+		Graph->Modify();
+		FGraphNodeCreator<UNiagaraNodeCustomHlsl> FunctionCallNodeCreator(*Graph);
+		UNiagaraNodeCustomHlsl* FunctionCallNode = FunctionCallNodeCreator.CreateNode();
+		FunctionCallNode->InitAsCustomHlslDynamicInput(Schema->PinToTypeDefinition(&OverridePin));
+		FunctionCallNodeCreator.Finalize();
+		FunctionCallNode->SetEnabledState(OverrideNode->GetDesiredEnabledState(), OverrideNode->HasUserSetTheEnabledState());
 
-	UEdGraphPin* FunctionCallInputPin = FNiagaraStackGraphUtilities::GetParameterMapInputPin(*FunctionCallNode);
-	FPinCollectorArray FunctionCallOutputPins;
-	FunctionCallNode->GetOutputPins(FunctionCallOutputPins);
+		UEdGraphPin* FunctionCallInputPin = FNiagaraStackGraphUtilities::GetParameterMapInputPin(*FunctionCallNode);
+		FPinCollectorArray FunctionCallOutputPins;
+		FunctionCallNode->GetOutputPins(FunctionCallOutputPins);
 
-	const UEdGraphSchema_Niagara* NiagaraSchema = GetDefault<UEdGraphSchema_Niagara>();
+		const UEdGraphSchema_Niagara* NiagaraSchema = GetDefault<UEdGraphSchema_Niagara>();
 
-	FNiagaraTypeDefinition InputType = NiagaraSchema->PinToTypeDefinition(&OverridePin);
-	checkf(FunctionCallInputPin != nullptr, TEXT("Dynamic Input function call did not have a parameter map input pin."));
-	checkf(FunctionCallOutputPins.Num() == 2 && NiagaraSchema->PinToTypeDefinition(FunctionCallOutputPins[0]) == InputType, TEXT("Invalid Stack Graph - Dynamic Input function did not have the correct typed output pin"));
+		FNiagaraTypeDefinition InputType = NiagaraSchema->PinToTypeDefinition(&OverridePin);
+		checkf(FunctionCallInputPin != nullptr, TEXT("Dynamic Input function call did not have a parameter map input pin."));
+		checkf(FunctionCallOutputPins.Num() == 2 && NiagaraSchema->PinToTypeDefinition(FunctionCallOutputPins[0]) == InputType, TEXT("Invalid Stack Graph - Dynamic Input function did not have the correct typed output pin"));
 
-	UEdGraphPin* OverrideNodeInputPin = FNiagaraStackGraphUtilities::GetParameterMapInputPin(*OverrideNode);
-	UEdGraphPin* PreviousStackNodeOutputPin = OverrideNodeInputPin->LinkedTo[0];
-	checkf(PreviousStackNodeOutputPin != nullptr, TEXT("Invalid Stack Graph - No previous stack node."));
+		UEdGraphPin* OverrideNodeInputPin = FNiagaraStackGraphUtilities::GetParameterMapInputPin(*OverrideNode);
+		UEdGraphPin* PreviousStackNodeOutputPin = OverrideNodeInputPin->LinkedTo[0];
+		checkf(PreviousStackNodeOutputPin != nullptr, TEXT("Invalid Stack Graph - No previous stack node."));
 
-	MakeLinkTo(FunctionCallInputPin, PreviousStackNodeOutputPin);
-	MakeLinkTo(FunctionCallOutputPins[0], &OverridePin);
+		MakeLinkTo(FunctionCallInputPin, PreviousStackNodeOutputPin);
+		MakeLinkTo(FunctionCallOutputPins[0], &OverridePin);
 
-	OutDynamicInputFunctionCall = FunctionCallNode;
+		OutDynamicInputFunctionCall = FunctionCallNode;
 
-	if (NewNodePersistentId.IsValid())
+		if (NewNodePersistentId.IsValid())
+		{
+			FunctionCallNode->NodeGuid = NewNodePersistentId;
+		}
+
+		FunctionCallNode->SetCustomHlsl(CustomExpression);
+	}
+	else
 	{
-		FunctionCallNode->NodeGuid = NewNodePersistentId;
+		OutDynamicInputFunctionCall = nullptr;
 	}
-
-	FunctionCallNode->SetCustomHlsl(CustomExpression);
 }
 
 bool FNiagaraStackGraphUtilities::RemoveModuleFromStack(UNiagaraSystem& OwningSystem, FGuid OwningEmitterId, UNiagaraNodeFunctionCall& ModuleNode)
diff --git a/Engine/Plugins/FX/Niagara/Source/NiagaraEditor/Public/NiagaraExternalSystemEditorUtilities.h b/Engine/Plugins/FX/Niagara/Source/NiagaraEditor/Public/NiagaraExternalSystemEditorUtilities.h
--- a/Engine/Plugins/FX/Niagara/Source/NiagaraEditor/Public/NiagaraExternalSystemEditorUtilities.h
+++ b/Engine/Plugins/FX/Niagara/Source/NiagaraEditor/Public/NiagaraExternalSystemEditorUtilities.h
@@ -5,59 +5,762 @@
 #include "NiagaraVariant.h"
 #include "NiagaraEditorUtilities.h"
 
+#include "JsonObjectConverter.h"
+#include "JsonSchema/JsonSchemaPropertyFilter.h"
+
+#include "NiagaraParameterMapHistoryFwd.h"
+#include "NiagaraVariableMetaData.h"
+
+#include "Kismet/BlueprintFunctionLibrary.h"
+
+#include "NiagaraExternalSystemEditorUtilities.generated.h"
+
+class UNiagaraStackModuleItem;
+class UNiagaraStackFunctionInput;
+
 #define UE_API NIAGARAEDITOR_API
 
 //Tools for external code to modify Niagara Systems.
-//Provides minimal public API required for views into and modification of Niagara internals.
-class FNiagaraExternalSystemEditorUtilities
+//Provides minimal public API required for info about and modification of Niagara internals.
+// 
+//EXPERIMENTAL! DO NOT USE EXCEPT FOR TESTING! EVERYTHING HERE IS SUBJECT TO CHANGE WITHOUT NOTICE!
+//EXPERIMENTAL! DO NOT USE EXCEPT FOR TESTING! EVERYTHING HERE IS SUBJECT TO CHANGE WITHOUT NOTICE!
+//EXPERIMENTAL! DO NOT USE EXCEPT FOR TESTING! EVERYTHING HERE IS SUBJECT TO CHANGE WITHOUT NOTICE!
+
+
+/** Base class for instanced values used in Niagara External System Editor. */
+USTRUCT(BlueprintInternalUseOnly)
+struct FNiagaraExt_InstancedValue : public FInstancedStruct
 {
-public:
+	GENERATED_BODY()	
+};
 
-	struct FUserVariableInfo
-	{
-		FNiagaraVariableBase Variable;
-		FNiagaraVariant Value;
-		FString Description;
-	};
+/**
+An instanced struct type that holds the value of a Niagara Variable.
+Can contain specialized types (Object, Enum, DataInterface wrappers) or standard value types.
+*/
+USTRUCT(BlueprintInternalUseOnly, meta=(
+	BaseStruct="/Script/NiagaraEditor.NiagaraExt_VariableValueBase",
+	AdditionalTypes="/Script/Niagara.NiagaraFloat,/Script/Niagara.NiagaraBool,/Script/Niagara.NiagaraInt32,/Script/Niagara.NiagaraPosition,/Script/Niagara.NiagaraID,/Script/CoreUObject.Vector,/Script/CoreUObject.Vector2D,/Script/CoreUObject.Vector4,/Script/CoreUObject.Quat,/Script/CoreUObject.LinearColor,/Script/CoreUObject.Rotator"
+))
+struct FNiagaraExt_VariableValue : public FNiagaraExt_InstancedValue
+{
+	GENERATED_BODY()
+
+	UE_API void Get(FNiagaraVariant& OutVariant, FNiagaraExternalEditContext& Context)const;
+	UE_API void Set(const FNiagaraTypeDefinition& TypeDef, const FNiagaraVariant& Variant);
+};
+
+/** Base struct for variable value types. Excluded from schema - only derived types are valid. */
+USTRUCT(BlueprintInternalUseOnly, meta=(EDAIncludeStructsInSchema, ExcludeBaseStruct))
+struct FNiagaraExt_VariableValueBase
+{
+	GENERATED_BODY()
+};
+
+/**
+Wrapper for UEnum values. USE THIS when the Niagara variable type is an Enum.
+Provides type-safe enum reference and string value selection.
+*/
+USTRUCT(BlueprintInternalUseOnly, meta=(
+	ToolTip="Wrapper for UEnum references. Use this struct when passing enum values to Niagara variables."
+))
+struct FNiagaraExt_VariableValue_Enum : public FNiagaraExt_VariableValueBase
+{
+	GENERATED_BODY()
+
+	/** The UEnum type definition */
+	UPROPERTY(EditAnywhere, Category="Variable", meta=(ToolTip="Reference to the UEnum type"))
+	TObjectPtr<UEnum> Enum;
+
+	/** 
+	The selected enum name as a string (e.g., \"MyEnum::Name1\"). 
+	For some enum names this can be an auto generated and un-descriptive name such as EnumValue0.
+	The display name is often much more useful and descriptive. Access this via the schema for the UEnum.
+	*/
+	UPROPERTY(EditAnywhere, Category="Variable", meta=(ToolTip="The enum value name"))
+	FName EnumName;
+};
+
+/**
+Wrapper for UNiagaraDataInterface values. USE THIS when the Niagara variable type is a DataInterface.
+Provides class selection and instance reference.
+*/
+USTRUCT(BlueprintInternalUseOnly, meta=(
+	ToolTip="Wrapper for DataInterface references. Use this struct when passing DataInterface instances to Niagara variables."
+))
+struct FNiagaraExt_VariableValue_DataInterface : public FNiagaraExt_VariableValueBase
+{
+	GENERATED_BODY()
+
+	/** The DataInterface class type */
+	UPROPERTY(EditAnywhere, Category="Variable", meta=(
+		ToolTip="The class of the DataInterface"
+	))
+	TSubclassOf<UNiagaraDataInterface> DataInterfaceClass;
+
+	/** The DataInterface instance */
+	UPROPERTY(EditAnywhere, Category="Variable", meta=(
+		ToolTip="Reference to the DataInterface object"
+	))
+	TObjectPtr<UNiagaraDataInterface> DataInterface;
+};
+
+/**
+Wrapper for generic UObject values. USE THIS when the Niagara variable type is a UObject reference (textures, materials, etc.).
+Provides class constraint and object reference. Do NOT use simple types for object references.
+*/
+USTRUCT(BlueprintInternalUseOnly, meta=(
+	ToolTip="Wrapper for UObject references. Use this struct when passing UObject instances (textures, materials, etc.) to Niagara variables."
+))
+struct FNiagaraExt_VariableValue_Object : public FNiagaraExt_VariableValueBase
+{
+	GENERATED_BODY()
+
+	/** The expected UObject class type (e.g., UTexture2D, UMaterialInterface) */
+	UPROPERTY(EditAnywhere, Category="Variable", meta=(
+		ToolTip="The class type constraint for the object reference"
+	))
+	TSubclassOf<UObject> ObjectClass;
+
+	/** The UObject instance reference */
+	UPROPERTY(EditAnywhere, Category="Variable", meta=(
+		ToolTip="Reference to the UObject instance"
+	))
+	TObjectPtr<UObject> Object;
+};
+
+USTRUCT(BlueprintInternalUseOnly)
+struct FNiagaraExt_Variable
+{
+	GENERATED_BODY()
+
+	UPROPERTY(EditAnywhere, Category="Variable")
+	FName Name;
+	
+ 	UPROPERTY()
+ 	FNiagaraTypeDefinition Type;
+};
+
+USTRUCT(BlueprintInternalUseOnly)
+struct FNiagaraExt_VariableInst : public FNiagaraExt_Variable
+{
+	GENERATED_BODY()
+	
+	UPROPERTY()
+	FNiagaraExt_VariableValue Value;
+};
+
+USTRUCT(BlueprintInternalUseOnly)
+struct FNiagaraExt_UserVariable : public FNiagaraExt_Variable
+{
+	GENERATED_BODY()
+		
+	UPROPERTY(EditAnywhere, Category="Variable")
+	FNiagaraExt_VariableValue DefaultValue;
+	
+	UPROPERTY(EditAnywhere, Category="Variable")
+	FText Description;
+};
+
+USTRUCT(BlueprintInternalUseOnly)
+struct FNiagaraExt_UserVariables 
+{
+	GENERATED_BODY()
+	
+	UPROPERTY(EditAnywhere, Category="Variables")
+	TArray<FNiagaraExt_UserVariable> UserVariables;
+};
+
+
+/* We expose the internals of Niagara in 3 layers.
+* 
+* 1. Property Schema: Top level properties that can be modified for Systems,Emitters,Modules,Renderers,DataInterfaces. Constant and may need calling only once.
+* 2. Topology: The currently topology of the Niagara System, Emitters, Modules and their Dynamic input chains. Changes infrequently.
+* 3. Values: The current values of properties and module inputs. Changes frequently. 
+*/
+
+/** Schema describing a Niagara System asset's top level properties. */
+USTRUCT(BlueprintInternalUseOnly)
+struct FNiagaraExt_SystemSchema
+{
+	GENERATED_BODY()
+	
+	/** Json Schema describing editable properties of a Niagara System. */	
+	UPROPERTY(EditAnywhere, Category="Schema")
+	FString PropertySchema;
+};
+
+/** Schema describing a Niagara Emitter asset's top level properties. */
+USTRUCT(BlueprintInternalUseOnly)
+struct FNiagaraExt_EmitterSchema
+{
+	GENERATED_BODY()
+
+	/** Json Schema describing editable properties of a Niagara Emitter. */
+	UPROPERTY(EditAnywhere, Category="Schema")
+	FString PropertySchema;
+};
+
+/** Schema describing a Niagara Renderer. */
+USTRUCT(BlueprintInternalUseOnly)
+struct FNiagaraExt_RendererSchema
+{
+	GENERATED_BODY()
+
+	UPROPERTY(EditAnywhere, Category="Schema")
+	TSubclassOf<UNiagaraRendererProperties> RendererClass;
+
+	/** Json Schema describing editable properties of a the Niagara Renderer class. */
+	UPROPERTY(EditAnywhere, Category="Schema")
+	FString PropertySchema;
+};
+
+/** Schema describing a Niagara Data Interface. */
+USTRUCT(BlueprintInternalUseOnly)
+struct FNiagaraExt_DataInterfaceSchema
+{
+	GENERATED_BODY()
+
+	UPROPERTY(EditAnywhere, Category="Schema")
+	TSubclassOf<UNiagaraDataInterface> DataInterfaceClass;
+	
+	/** Json Schema describing editable properties of a the data interface class. */
+	UPROPERTY(EditAnywhere, Category="Schema")
+	FString PropertySchema;
+};
+
+/** Schema describing a single Niagara Module input. */
+USTRUCT(BlueprintInternalUseOnly)
+struct FNiagaraExt_StackInputSchema
+{
+	GENERATED_BODY()
+
+	UPROPERTY(EditAnywhere, Category="Schema")
+	FName Name;
+	
+ 	UPROPERTY()
+ 	FNiagaraTypeDefinition Type;
+
+	UPROPERTY(EditAnywhere, Category="Schema")
+	FText Category;
+
+	/** True if this input can be set with a hlsl expression. */
+	UPROPERTY(EditAnywhere, Category="Schema")
+	bool bSupportsExpressions = false;
+	
+	UPROPERTY(EditAnywhere, Category="Schema")
+	FNiagaraVariableMetaData MetaData;
+};
+
+/** Schema describing a Niagara Module asset and its inputs. */
+USTRUCT(BlueprintInternalUseOnly)
+struct FNiagaraExt_ModuleSchema
+{
+	GENERATED_BODY()
+
+	UPROPERTY(EditAnywhere, Category="Schema")
+	TObjectPtr<UNiagaraScript> Asset;
+
+	UPROPERTY(EditAnywhere, Category="Schema")
+	TArray<FNiagaraExt_StackInputSchema> Inputs;
+
+	UPROPERTY(EditAnywhere, Category="Schema")
+	TArray<FNiagaraExt_Variable> Outputs;
+};
+
+/** Schema describing a Niagara Dynamic Input asset. */
+USTRUCT(BlueprintInternalUseOnly)
+struct FNiagaraExt_DynamicInputSchema : public FNiagaraExt_ModuleSchema
+{
+	GENERATED_BODY()
+};
+
+//////////////////////////////////////////////////////////////////////////
+// Data Layer - Defines the actual values of object properties and stack module inputs. 
+
+USTRUCT(BlueprintInternalUseOnly)
+struct FNiagaraExt_SystemData
+{
+	GENERATED_BODY()
+
+	/** Json values corresponding to the schema for a Niagara System. */
+	UPROPERTY(EditAnywhere, Category="Data")
+	FString PropertyValues;
+};
+
+USTRUCT(BlueprintInternalUseOnly)
+struct FNiagaraExt_EmitterData
+{
+	GENERATED_BODY()
+
+	/** Json values corresponding to the schema for a Niagara Emitter. */
+	UPROPERTY(EditAnywhere, Category="Data")
+	FString PropertyValues;
+};
+
+/** Top level properties for a Niagara Module. */
+USTRUCT(BlueprintInternalUseOnly)
+struct FNiagaraExt_ModuleData
+{
+	GENERATED_BODY()
+
+	UPROPERTY(EditAnywhere, Category="Data")
+	FName ModuleName;
+
+	UPROPERTY(EditAnywhere, Category="Data")
+	bool Enabled = true;
+};
+
+USTRUCT(BlueprintInternalUseOnly)
+struct FNiagaraExt_RendererData
+{
+	GENERATED_BODY()
+
+	/** Json values corresponding to the schema for a Niagara Renderer. */
+	UPROPERTY(EditAnywhere, Category="Data")
+	FString PropertyValues;
+};
+
+/** An instanced struct type that holds the value of a Niagara Stack Input. */
+USTRUCT(BlueprintInternalUseOnly, 
+meta=(
+BaseStruct="/Script/NiagaraEditor.FNiagaraExt_StackInputData",
+AdditionalTypes="/Script/Niagara.NiagaraFloat,/Script/Niagara.NiagaraBool,/Script/Niagara.NiagaraInt32,/Script/Niagara.NiagaraPosition,/Script/Niagara.NiagaraID,/Script/CoreUObject.Vector3f,/Script/CoreUObject.Vector2f,/Script/CoreUObject.Vector4f,/Script/CoreUObject.Quat4f,/Script/CoreUObject.LinearColor,/Script/CoreUObject.Rotator"
+))
+struct FNiagaraExt_StackInputValue : public FNiagaraExt_InstancedValue
+{
+	GENERATED_BODY()
+	
+	void InitStackInputFromValue(const FNiagaraExt_StackItemReference& StackInputRef, FNiagaraExternalEditContext& Context)const;
+	void InitFromStackInput(const FNiagaraExt_StackItemReference& StackInputRef, FNiagaraExternalEditContext& Context);
+};
+
+/** Base struct for stack input data types. Excluded from schema - only derived types are valid. */
+USTRUCT(BlueprintInternalUseOnly, meta=(EDAIncludeStructsInSchema, ExcludeBaseStruct))
+struct FNiagaraExt_StackInputData
+{
+	GENERATED_BODY()
+};
+
+/** A special case of Local input that stores an Enum value as a string rather than it's raw int. */
+USTRUCT(BlueprintInternalUseOnly)
+struct FNiagaraExt_StackInputData_Enum : public FNiagaraExt_StackInputData
+{
+	GENERATED_BODY()
+
+	/** The UEnum type definition */
+	UPROPERTY(EditAnywhere, Category="Input", meta=(ToolTip="Reference to the UEnum type"))
+	TObjectPtr<UEnum> Enum;
+
+	/** 
+	The selected enum name as a string (e.g., \"MyEnum::Name1\"). 
+	For some enum names this can be an auto generated and un-descriptive name such as EnumValue0.
+	The display name is often much more useful and descriptive. Access this via the schema for the UEnum.
+	*/
+	UPROPERTY(EditAnywhere, Category="Input", meta=(ToolTip="The enum value name"))
+	FName EnumName;
+};
+
+USTRUCT(BlueprintInternalUseOnly)
+struct FNiagaraExt_StackInputData_Linked : public FNiagaraExt_StackInputData
+{
+	GENERATED_BODY()
+
+	UPROPERTY(EditAnywhere, Category="Input")
+	FNiagaraExt_Variable LinkedVariable;
+};
+
+USTRUCT(BlueprintInternalUseOnly)
+struct FNiagaraExt_StackInputData_HlslExpression : public FNiagaraExt_StackInputData
+{
+	GENERATED_BODY()
+
+	/** A single rvalue hlsl expression defining the value of this input. MUST be a single rvalue with no variable definitions or return statements etc. For example a color could use the expression "float4(1.0, 0.0, 0.0, 1.0)" */
+	UPROPERTY(EditAnywhere, Category="Input")
+	FString HlslExpression;
+};
+
+USTRUCT(BlueprintInternalUseOnly)
+struct FNiagaraExt_StackInputData_DataInterface : public FNiagaraExt_StackInputData
+{
+	GENERATED_BODY()
+	
+	/** Json values corresponding to the schema for this data interface. */
+	UPROPERTY(EditAnywhere, Category="Input")
+	FString PropertyValues;
+};
+
+/** 
+Stack input for a dynamic input.
+*/
+USTRUCT(BlueprintInternalUseOnly)
+struct FNiagaraExt_StackInputData_DynamicInput : public FNiagaraExt_StackInputData
+{
+	GENERATED_BODY()
+	
+	/** The dynamic input asset defining this dynamic input. */
+	UPROPERTY(EditAnywhere, Category="Input")
+	TObjectPtr<UNiagaraScript> DynamicInputAsset;
+};
+
+/** An unsupported input type */
+USTRUCT(BlueprintInternalUseOnly)
+struct FNiagaraExt_StackInputData_Unsupported : public FNiagaraExt_StackInputData
+{
+	GENERATED_BODY()
+};
+
+
+//////////////////////////////////////////////////////////////////////////
+//Topology Layer - Defines how Niagara Systems, Emitters, Modules and their inputs are arranged.
+
+/** Instanced Struct wrapper around Stack Input Topology allowing us to compose them hierarchically for dynamic input chains. */
+USTRUCT(BlueprintInternalUseOnly, meta=(BaseStruct="/Script/NiagaraEditor.FNiagaraExt_StackInputTopology"))
+struct FNiagaraExt_DynamicInputTopology : public FNiagaraExt_InstancedValue
+{
+	GENERATED_BODY()
+
+	FNiagaraExt_DynamicInputTopology();
+	const FNiagaraExt_StackInputTopology& Get()const;
+	FNiagaraExt_StackInputTopology& GetMutable();
+};
+
+USTRUCT(BlueprintInternalUseOnly)
+struct FNiagaraExt_StackInputTopology : public FNiagaraExt_Variable
+{
+	GENERATED_BODY()
+	
+	/** True if this input is visible. Inputs can be hidden based on the value of a linked Visible Condition. See the Module Schema for details. */
+	UPROPERTY(EditAnywhere, Category="Input")
+	bool bIsVisible = true;
+
+	/** True if this input is editable. Inputs can be made read only based on the value of a linked Edit Condition. See the Module Schema for details. */
+	UPROPERTY(EditAnywhere, Category="Input")
+	bool bIsEditable = true;
+
+	UPROPERTY(EditAnywhere, Category="Topology")
+	FNiagaraExt_StackInputValue Data;
+	
+	/** Topology of sub inputs if this input is a Dynamic Input.*/
+	UPROPERTY(EditAnywhere, Category="Topology")
+	TArray<FNiagaraExt_DynamicInputTopology> DynamicInputs;
+};
+
+USTRUCT(BlueprintInternalUseOnly)
+struct FNiagaraExt_ModuleTopology
+{
+	GENERATED_BODY()
+	
+	UPROPERTY(EditAnywhere, Category="Topology")
+	FName ModuleName;
+
+	UPROPERTY(EditAnywhere, Category="Topology")
+	bool Enabled = true;
+
+	UPROPERTY(EditAnywhere, Category="Topology")
+	TObjectPtr<UNiagaraScript> ModuleScript;
+	
+	UPROPERTY(EditAnywhere, Category="Topology")
+	TArray<FNiagaraExt_StackInputTopology> Inputs;
+};
+
+USTRUCT(BlueprintInternalUseOnly)
+struct FNiagaraExt_ScriptStackTopology
+{
+	GENERATED_BODY()
+
+	UPROPERTY(EditAnywhere, Category="Topology")
+	FName ScriptName;
+
+	UPROPERTY(EditAnywhere, Category="Topology")
+	TArray<FNiagaraExt_ModuleTopology> Modules;
+};
+
+USTRUCT(BlueprintInternalUseOnly)
+struct FNiagaraExt_RendererRef
+{
+	GENERATED_BODY()
+
+	UPROPERTY(EditAnywhere, Category="Topology")
+	int32 Index = INDEX_NONE;
+
+	UPROPERTY(EditAnywhere, Category="Topology")
+	TSubclassOf<UNiagaraRendererProperties> RendererClass;
+
+	UPROPERTY(EditAnywhere, Category="Data")
+	FNiagaraExt_RendererData Data;
+};
+
+USTRUCT(BlueprintInternalUseOnly)
+struct FNiagaraExt_EmitterTopology
+{
+	GENERATED_BODY()
+	
+	UPROPERTY(EditAnywhere, Category="Topology")
+	FName EmitterName;
+
+	UPROPERTY(EditAnywhere, Category="Topology")
+	FNiagaraExt_ScriptStackTopology EmitterSpawnScript;
+	
+	UPROPERTY(EditAnywhere, Category="Topology")
+	FNiagaraExt_ScriptStackTopology EmitterUpdateScript;
+	
+	UPROPERTY(EditAnywhere, Category="Topology")
+	FNiagaraExt_ScriptStackTopology ParticleSpawnScript;
+	
+	UPROPERTY(EditAnywhere, Category="Topology")
+	FNiagaraExt_ScriptStackTopology ParticleUpdateScript;
+
+	UPROPERTY(EditAnywhere, Category="Topology")
+	TArray<FNiagaraExt_RendererRef> Renderers;
+	
+	UPROPERTY(EditAnywhere, Category="Data")
+	FNiagaraExt_EmitterData Data;
+};
+
+USTRUCT(BlueprintInternalUseOnly)
+struct FNiagaraExt_SystemTopology
+{
+	GENERATED_BODY()
+
+	UPROPERTY(EditAnywhere, Category="Topology")
+	FNiagaraExt_ScriptStackTopology SystemSpawnScript;
+	
+	UPROPERTY(EditAnywhere, Category="Topology")
+	FNiagaraExt_ScriptStackTopology SystemUpdateScript;
+
+	UPROPERTY(EditAnywhere, Category="Topology")
+	TArray<FNiagaraExt_EmitterTopology> Emitters;
+
+	UPROPERTY(EditAnywhere, Category="Topology")
+	TArray<FNiagaraExt_UserVariable> UserVariables;
 	
-	//User Variables
-	NIAGARAEDITOR_API static void AddUserVariable(UNiagaraSystem& System, const FNiagaraVariableBase& Variable, const FNiagaraVariant& Variant, FString VariableDesc, FText TransactionDescription);
-	NIAGARAEDITOR_API static void RemoveUserVariable(UNiagaraSystem& System, const FNiagaraVariableBase& Variable, FText TransactionDescription);
-	NIAGARAEDITOR_API static void GetUserVariableInfo(UNiagaraSystem& System, TArray<FUserVariableInfo>& OutVariableInfo);
+	UPROPERTY(EditAnywhere, Category="Data")
+	FNiagaraExt_SystemData Data;
+};
+
+//////////////////////////////////////////////////////////////////////////
 
-	//System
+struct FNiagaraEmitterHandle;
+class UNiagaraStackScriptItemGroup;
+class UNiagaraStackModuleItem;
+class UNiagaraStackFunctionInput;
+class UNiagaraRendererProperties;
 
-	struct FEmitterInfo
+//A reference to an item inside a Niagara System such as an emitter, script, module, input or renderer.
+USTRUCT(BlueprintInternalUseOnly)
+struct FNiagaraExt_StackItemReference
+{
+	GENERATED_BODY()
+
+	FNiagaraExt_StackItemReference() = default;
+	FNiagaraExt_StackItemReference(UNiagaraSystem* InSystem, FName InEmitterName=NAME_None, FName InScriptName=NAME_None, FName InModuleName=NAME_None)
+		:System(InSystem), EmitterName(InEmitterName), ScriptName(InScriptName), ModuleName(InModuleName)
 	{
-		FGuid Id;
-		FName Name;
-		bool bLocalSpace = false;
-		bool bDeterministic = false;
-		ENiagaraSimTarget SimTarget = ENiagaraSimTarget::CPUSim;		
-		ENiagaraEmitterMode EmitterMode = ENiagaraEmitterMode::Standard;
-	};
-
-	struct FSystemInfo
+	}
+
+	UPROPERTY(EditAnywhere, Category="Stack Item")
+	TObjectPtr<UNiagaraSystem> System;
+
+	UPROPERTY(EditAnywhere, Category="Stack Item")
+	FName EmitterName;
+	
+	UPROPERTY(EditAnywhere, Category="Stack Item")
+	FName ScriptName;
+	
+	/** Name of the module to reference. Requires a valid ScriptName. */
+	UPROPERTY(EditAnywhere, Category="Stack Item")
+	FName ModuleName;
+	
+	/** Index of a renderer to reference. Requires a valid EmitterName. */
+	UPROPERTY(EditAnywhere, Category="Stack Item")
+	int32 RendererIndex = INDEX_NONE;
+
+	/** Input Name stack. Requires a valid ScriptName and ModuleName. For a top level module input this will be a single FName. For a dynamic input this could be a chain of several nested dynamic inputs. */
+	UPROPERTY(EditAnywhere, Category="Stack Item")
+	TArray<FName> InputNameStack;
+	
+	UE_API UNiagaraSystem* GetSystem(FNiagaraExternalEditContext& Context)const;
+	FNiagaraEmitterHandle* GetEmitter(FNiagaraExternalEditContext& Context, bool bRequired=true)const;
+	UNiagaraStackScriptItemGroup* GetScript(FNiagaraExternalEditContext& Context, bool bRequired=true)const;
+	UNiagaraStackModuleItem* GetModule(FNiagaraExternalEditContext& Context, bool bRequired=true)const;
+	UNiagaraStackFunctionInput* GetInput(FNiagaraExternalEditContext& Context, bool bRequired=true)const;
+	UNiagaraRendererProperties* GetRenderer(FNiagaraExternalEditContext& Context, bool bRequired=true)const;
+
+	void SetEmitter(FNiagaraEmitterHandle* Emitter);
+	void SetEmitter(FName InEmitterName);
+	
+	void SetScript(UNiagaraStackScriptItemGroup* Script);
+	void SetScript(FName InScriptName);
+	
+	void SetModule(UNiagaraStackModuleItem* Module);
+	void SetModule(FName InModuleName);
+
+	void SetInput(UNiagaraStackFunctionInput* Input);
+	void SetInput(FName Name);
+	
+	void PushInput(UNiagaraStackFunctionInput* Input);
+	void PushInput(FName Name);
+	
+	void SetRenderer(int32 InRendererIndex);
+	
+private:
+
+	// IMPORTANT: These caches are only valid for a single operation and should never be stored long-term.
+	// Emitter is not cached to avoid dangling pointers into the emitter handles array.
+	// UObject pointers use weak references to detect stale objects.
+	mutable TWeakObjectPtr<UNiagaraStackScriptItemGroup> CachedScript = nullptr;
+	mutable TWeakObjectPtr<UNiagaraStackModuleItem> CachedModule = nullptr;
+	mutable TWeakObjectPtr<UNiagaraStackFunctionInput> CachedInput = nullptr;
+};
+
+USTRUCT(BlueprintInternalUseOnly)
+struct FNiagaraExternalEditContext
+{
+	GENERATED_BODY()
+
+	UPROPERTY(EditAnywhere, Category="Context")
+	TArray<FText> Errors;
+
+	FNiagaraExternalEditContext() = default;
+	UE_API FNiagaraExternalEditContext(UNiagaraSystem* InSystem);
+	UE_API FNiagaraExternalEditContext(const FNiagaraExt_StackItemReference& ItemRef); 
+		UNiagaraSystem* GetSystem();
+	TSharedPtr<FNiagaraSystemViewModel> GetSystemViewModel();
+
+	TSet<TSubclassOf<UNiagaraRendererProperties>> UsedRenderers;
+	TSet<TSubclassOf<UNiagaraDataInterface>> UsedDataInterfaces;
+	TSet<UNiagaraScript*> UsedModules;
+	TSet<UNiagaraScript*> UsedDynamicInputs;
+
+	bool HasErrors()const { return Errors.Num() > 0; }
+
+	void Error(FText InError)
 	{
-		TArray<FEmitterInfo> Emitters;
-	};
-	NIAGARAEDITOR_API static void GetSystemInfo(UNiagaraSystem& System, FSystemInfo& OutSystemInfo);
+		Errors.Emplace(InError);
+	}
+
+	bool CheckSystem(const UNiagaraSystem* System);
+	bool CheckScriptAsset(const UNiagaraScript* ScriptAsset);
+	
+private:	
 
-	//Emitters
-	NIAGARAEDITOR_API static FName AddEmitterToSystem(UNiagaraSystem& System, UNiagaraEmitter& EmitterToAdd, FText TransactionDescription);
-	NIAGARAEDITOR_API static bool RemoveEmitterFromSystem(UNiagaraSystem& System, FName EmitterUniqueName, FText TransactionDescription);
+	TSharedPtr<FNiagaraSystemViewModel> SystemViewModel;
+};
 
-	NIAGARAEDITOR_API static bool GetEmitterInfo(UNiagaraSystem& System, FName EmitterUniqueName, FEmitterInfo& OutEmitterInfo, FText& OutErrorMessage);
-	NIAGARAEDITOR_API static bool SetEmitterName(UNiagaraSystem& System, FName EmitterUniqueName, FName Name, FText TransactionDescription, FText& OutErrorMessage);
-	NIAGARAEDITOR_API static bool SetEmitterLocalSpace(UNiagaraSystem& System, FName EmitterUniqueName, bool bLocalSpace, FText TransactionDescription, FText& OutErrorMessage);
-	NIAGARAEDITOR_API static bool SetEmitterDeterministic(UNiagaraSystem& System, FName EmitterUniqueName, bool bDeterministic, FText TransactionDescription, FText& OutErrorMessage);
-	NIAGARAEDITOR_API static bool SetEmitterSimTarget(UNiagaraSystem& System, FName EmitterUniqueName, ENiagaraSimTarget SimTarget, FText TransactionDescription, FText& OutErrorMessage);
+/** A complete set of information about a Niagara System. */
+USTRUCT(BlueprintInternalUseOnly)
+struct FNiagaraExt_SystemInfo
+{
+	GENERATED_BODY();
 
-	private:
-		
-	static void GetEmitterInfo(const FNiagaraEmitterHandle& EmitterHandle, FEmitterInfo& OutEmitterInfo);
-	static FNiagaraEmitterHandle* GetEmitterHandle(UNiagaraSystem& System, FName EmitterUniqueName);
-	static FNiagaraEmitterHandle* GetEmitterHandle(UNiagaraSystem& System, const FGuid& EmitterId);
+	/** All renderers used by this System. */
+	UPROPERTY(EditAnywhere, Category="Schema")
+	TArray<TSubclassOf<UNiagaraRendererProperties>> UsedRenderers;
+	
+	/** All data interfaces used by this System. */
+	UPROPERTY(EditAnywhere, Category="Schema")
+	TArray<TSubclassOf<UNiagaraDataInterface>> UsedDataInterfaces;
+	
+	/** All module scripts used by this System. */
+	UPROPERTY(EditAnywhere, Category="Schema")
+	TArray<TObjectPtr<UNiagaraScript>> UsedModules;
+	
+	/** All dynamic input scripts used by this System. */
+	UPROPERTY(EditAnywhere, Category="Schema")
+	TArray<TObjectPtr<UNiagaraScript>> UsedDynamicInputs;
+	
+	/** Complete System Topology with data detailing the layout and values of all emitters, scripts, modules, dynamic inputs and renderers. */
+	UPROPERTY(EditAnywhere, Category="Topology")
+	FNiagaraExt_SystemTopology SystemTopology;
+};
+
+/**
+* A C++ and Blueprint accessible library of functions allowing for the inspection and modification of Niagara Systems.
+*/
+UCLASS(MinimalAPI)
+class UNiagaraExternalEditUtilities : public UBlueprintFunctionLibrary
+{
+	GENERATED_BODY()
+
+public:
+
+	/**
+	 * Creates a new Niagara System asset.
+	 *
+	 * @param AssetName - Name for the new system asset (required)
+	 * @param AssetPath - Package path where the asset will be created (required)
+	 * @param TemplateSystem - Optional template system to copy. If null, an empty system will be created.
+	 * @param Context - Edit context for error reporting
+	 * @return The newly created system, or nullptr if creation failed
+	 */
+	UE_API static UNiagaraSystem* CreateNiagaraSystem(const FString& AssetName, const FString& AssetPath, UNiagaraSystem* TemplateSystem, FNiagaraExternalEditContext& Context);
+	UE_API static void GetAvailableDynamicInputs(const FNiagaraTypeDefinition& Type, TArray<UNiagaraScript*>& OutDynamicInputScripts, FNiagaraExternalEditContext& Context);
+	
+	
+	/** Gathers complete system info including all topology, data and schemas of used classes and assets. */
+	UE_API static void GetSystemInfoComplete(UNiagaraSystem* System, FNiagaraExt_SystemInfo& OutInfo, FNiagaraExternalEditContext& Context);
+
+	//Schemas Layer - Defines constant data about editable properties and inputs based on specific classes or assets rather than a specific System being edited.
+	UE_API static void GetSystemSchema(FNiagaraExt_SystemSchema& OutSchema);
+	UE_API static void GetEmitterSchema(FNiagaraExt_EmitterSchema& OutSchema);
+	UE_API static void GetRendererSchema(TSubclassOf<UNiagaraRendererProperties> RendererClass, FNiagaraExt_RendererSchema& OutSchema);
+	UE_API static void GetDataInterfaceSchema(TSubclassOf<UNiagaraDataInterface> DataInterfaceClass, FNiagaraExt_DataInterfaceSchema& OutSchema);
+ 	UE_API static void GetStackInputSchema(const FNiagaraExt_StackItemReference& InputReference, FNiagaraExt_StackInputSchema& OutSchema, FNiagaraExternalEditContext& Context);
+ 	UE_API static void GetModuleSchema(const FNiagaraExt_StackItemReference& ModuleReference, FNiagaraExt_ModuleSchema& OutSchema, FNiagaraExternalEditContext& Context);
+	UE_API static void GetDynamicInputSchema(const FNiagaraExt_StackItemReference& DynamicInputReference, FNiagaraExt_DynamicInputSchema& OutSchema, FNiagaraExternalEditContext& Context);
+
+	//Ideally we'll be able to inspect these without the need for an owning system but that isn't currently possible.
+// 	static void GetStackInputSchema(const TObjectPtr<UNiagaraScript> ModuleAsset, FNiagaraExt_Variable& Input, FNiagaraExt_StackInputSchema& OutSchema);
+ 	UE_API static void GetModuleSchema(const UNiagaraScript* ModuleAsset, FNiagaraExt_ModuleSchema& OutSchema, FNiagaraExternalEditContext& Context);
+ 	UE_API static void GetDynamicInputSchema(const UNiagaraScript* ModuleAsset, FNiagaraExt_DynamicInputSchema& OutSchema, FNiagaraExternalEditContext& Context);
+
+	//Schema Layer END
+
+	//Topology Layer - Defines the current topology of niagara Systems, Emitters and module inputs.
+	UE_API static void GetSystemTopology(UNiagaraSystem* System, FNiagaraExt_SystemTopology& OutTopology, FNiagaraExternalEditContext& Context);	
+	UE_API static void GetScriptStackTopology(const FNiagaraExt_StackItemReference& ScriptRef, FNiagaraExt_ScriptStackTopology& OutTopology, FNiagaraExternalEditContext& Context);
+	UE_API static void GetEmitterTopology(const FNiagaraExt_StackItemReference& EmitterRef, FNiagaraExt_EmitterTopology& OutTopology, FNiagaraExternalEditContext& Context);
+	UE_API static void GetModuleTopology(const FNiagaraExt_StackItemReference& ModuleRef, FNiagaraExt_ModuleTopology& OutTopology, FNiagaraExternalEditContext& Context);
+	UE_API static void GetStackInputTopology(const FNiagaraExt_StackItemReference& StackInputRef, FNiagaraExt_StackInputTopology& OutTopology, FNiagaraExternalEditContext& Context);
+
+	//Topology Layer END
+
+public:
+	//Data - Access actual data values of Systems, Emitters and module inputs.	
+	UE_API static void GetUserVariables(UNiagaraSystem* System, FNiagaraExt_UserVariables& OutVariables, FNiagaraExternalEditContext& Context);
+	UE_API static void GetSystemData(UNiagaraSystem* System, FNiagaraExt_SystemData& OutData, FNiagaraExternalEditContext& Context);
+	UE_API static void GetEmitterData(const FNiagaraExt_StackItemReference& EmitterRef, FNiagaraExt_EmitterData& OutData, FNiagaraExternalEditContext& Context);
+	UE_API static void GetRendererData(const FNiagaraExt_StackItemReference& RendererRef, FNiagaraExt_RendererData& OutData, FNiagaraExternalEditContext& Context);
+	UE_API static void GetStackInputData(const FNiagaraExt_StackItemReference& StackInputRef, FNiagaraExt_StackInputValue& OutData, FNiagaraExternalEditContext& Context);
+
+
+	//Edit Operations
+	UE_API static void AddUserVariable(UNiagaraSystem* System, const FNiagaraExt_UserVariable& Variable, FNiagaraExternalEditContext& Context);
+	UE_API static void RemoveUserVariable(UNiagaraSystem* System, const FNiagaraExt_Variable& Variable, FNiagaraExternalEditContext& Context);
+	
+	UE_API static void AddEmitter(UNiagaraEmitter* TemplateEmitter, FName EmitterName, FNiagaraExt_EmitterTopology& OutTopology, FNiagaraExternalEditContext& Context);
+ 	UE_API static void RemoveEmitter(const FNiagaraExt_StackItemReference& EmitterToRemoveRef, FNiagaraExternalEditContext& Context);
+	 
+	UE_API static void AddRenderer(const FNiagaraExt_StackItemReference& NewRendererLocation, const TSubclassOf<UNiagaraRendererProperties> RendererClass, FNiagaraExternalEditContext& Context);
+	UE_API static void RemoveRenderer(const FNiagaraExt_StackItemReference& RendererToRemove, FNiagaraExternalEditContext& Context);
+
+	UE_API static void AddModule(const FNiagaraExt_StackItemReference& NewModuleLocationRef, const UNiagaraScript* ModuleAsset, FNiagaraExt_ModuleTopology& OutTopology, FNiagaraExternalEditContext& Context);
+	UE_API static void RemoveModule(const FNiagaraExt_StackItemReference& ModuleToRemove, FNiagaraExternalEditContext& Context);
+	UE_API static void SetModuleEnabled(const FNiagaraExt_StackItemReference& ModuleRef, bool bEnabled, FNiagaraExternalEditContext& Context);
+
+	UE_API static void SetSystemData(UNiagaraSystem* System, const FNiagaraExt_SystemData& InData, FNiagaraExternalEditContext& Context);
+	UE_API static void SetEmitterData(const FNiagaraExt_StackItemReference& EmitterItemRef, const FNiagaraExt_EmitterData& InData, FNiagaraExternalEditContext& Context);
+	UE_API static void SetRendererData(const FNiagaraExt_StackItemReference& RendererItemRef, const FNiagaraExt_RendererData& InData, FNiagaraExternalEditContext& Context);
+	UE_API static void SetStackInputData(const FNiagaraExt_StackItemReference& StackItemRef, const FNiagaraExt_StackInputValue& InData, FNiagaraExternalEditContext& Context);
+
+private:
+	static void GetEmitterDataInternal(UNiagaraSystem* System, FNiagaraEmitterHandle& Emitter, FNiagaraExt_EmitterData& OutData, FNiagaraExternalEditContext& Context);
+	static void GetRendererDataInternal(FNiagaraEmitterHandle& Emitter, UNiagaraRendererProperties* Renderer, FNiagaraExt_RendererData& OutData, FNiagaraExternalEditContext& Context);
+
+	// Builds parameter map histories for a script allowing some inspection of script schemas without needing an owning system.
+	static bool BuildParameterMapHistoriesFromScript(const UNiagaraScript* Script, TArray<FNiagaraParameterMapHistory>& OutHistories, FNiagaraExternalEditContext& Context);
 };
 
 #undef UE_API
diff --git a/Engine/Plugins/FX/Niagara/Source/NiagaraEditor/Public/ViewModels/Stack/NiagaraStackFunctionInput.h b/Engine/Plugins/FX/Niagara/Source/NiagaraEditor/Public/ViewModels/Stack/NiagaraStackFunctionInput.h
--- a/Engine/Plugins/FX/Niagara/Source/NiagaraEditor/Public/ViewModels/Stack/NiagaraStackFunctionInput.h
+++ b/Engine/Plugins/FX/Niagara/Source/NiagaraEditor/Public/ViewModels/Stack/NiagaraStackFunctionInput.h
@@ -239,6 +239,9 @@ public:
 
 	/** Sets this input's local value. */
 	NIAGARAEDITOR_API void SetLocalValue(TSharedRef<FStructOnScope> InLocalValue);
+	
+	/** Sets this input's data interface value by creating a new default Data Interface and allowing it to be modified by the caller. */
+	NIAGARAEDITOR_API void SetDataInterfaceValueExternal(TSubclassOf<UNiagaraDataInterface> InDataInterfaceType, TFunction<void(UNiagaraDataInterface*)> ExternalSetFunction);
 
 	/** Sets this input's data interface value by creating a new default Data Interface. */
 	NIAGARAEDITOR_API void SetDataInterfaceValue(TSubclassOf<UNiagaraDataInterface> InDataInterfaceType);
@@ -304,6 +307,9 @@ public:
 
 	/** Gets whether or not this input has an associated visible condition input. */
 	NIAGARAEDITOR_API bool GetHasVisibleCondition() const;
+	
+	/** Gets the variable that serves as a visible condition for this input. */
+	NIAGARAEDITOR_API TOptional<FNiagaraVariable> GetVisibleConditionVariable() const;
 
 	/** Gets the enabled value of the visible condition input associated with this input. */
 	NIAGARAEDITOR_API bool GetVisibleConditionEnabled() const;
diff --git a/Engine/Plugins/FX/Niagara/Source/NiagaraEditor/Public/ViewModels/Stack/NiagaraStackGraphUtilities.h b/Engine/Plugins/FX/Niagara/Source/NiagaraEditor/Public/ViewModels/Stack/NiagaraStackGraphUtilities.h
--- a/Engine/Plugins/FX/Niagara/Source/NiagaraEditor/Public/ViewModels/Stack/NiagaraStackGraphUtilities.h
+++ b/Engine/Plugins/FX/Niagara/Source/NiagaraEditor/Public/ViewModels/Stack/NiagaraStackGraphUtilities.h
@@ -143,6 +143,8 @@ namespace FNiagaraStackGraphUtilities
 		ENiagaraGetStackFunctionInputPinsOptions Options,
 		bool bIgnoreDisabled,
 		bool bFilterForCompilation);
+		
+	void ExtractInputPinsFromHistory(FNiagaraParameterMapHistory& History, UEdGraph* FunctionGraph, FNiagaraStackGraphUtilities::ENiagaraGetStackFunctionInputPinsOptions Options, TArray<const UEdGraphPin*>& OutPins);
 
 	TArray<UNiagaraNodeOutput*> GetAllOutputNodes(TSharedRef<FNiagaraEmitterViewModel> EmitterViewModel);
 	
