- Remove include and use forward decl instead

#rnx

#Synced-CL 51985790

diff --git a/Engine/Plugins/FX/Niagara/Source/Niagara/Private/NiagaraParameterStore.cpp b/Engine/Plugins/FX/Niagara/Source/Niagara/Private/NiagaraParameterStore.cpp
--- a/Engine/Plugins/FX/Niagara/Source/Niagara/Private/NiagaraParameterStore.cpp
+++ b/Engine/Plugins/FX/Niagara/Source/Niagara/Private/NiagaraParameterStore.cpp
@@ -5,6 +5,7 @@
 #include "NiagaraCustomVersion.h"
 #include "NiagaraDataInterface.h"
 #include "NiagaraStats.h"
+#include "NiagaraVariant.h"
 
 #include UE_INLINE_GENERATED_CPP_BY_NAME(NiagaraParameterStore)
 
diff --git a/Engine/Plugins/FX/Niagara/Source/Niagara/Public/NiagaraParameterStore.h b/Engine/Plugins/FX/Niagara/Source/Niagara/Public/NiagaraParameterStore.h
--- a/Engine/Plugins/FX/Niagara/Source/Niagara/Public/NiagaraParameterStore.h
+++ b/Engine/Plugins/FX/Niagara/Source/Niagara/Public/NiagaraParameterStore.h
@@ -3,9 +3,10 @@
 #pragma once
 
 #include "NiagaraTypes.h"
-#include "NiagaraVariant.h"
 #include "NiagaraParameterStore.generated.h"
 
+struct FNiagaraVariant;
+
 // When not cooked, sort by actual name to ensure deterministic cooked data
 #define NIAGARA_VARIABLE_LEXICAL_SORTING WITH_EDITORONLY_DATA
 
