- Fix crash in scalability view due to using a raw pointer into a map structure

#tests Enter scalability view, copy and paste the same emitter multiple times (4 times) click on a module, observer crash

#Synced-CL 51985790

diff --git a/Engine/Plugins/FX/Niagara/Source/NiagaraEditor/Private/Widgets/Scalability/SNiagaraSystemResolvedScalabilitySettings.cpp b/Engine/Plugins/FX/Niagara/Source/NiagaraEditor/Private/Widgets/Scalability/SNiagaraSystemResolvedScalabilitySettings.cpp
--- a/Engine/Plugins/FX/Niagara/Source/NiagaraEditor/Private/Widgets/Scalability/SNiagaraSystemResolvedScalabilitySettings.cpp
+++ b/Engine/Plugins/FX/Niagara/Source/NiagaraEditor/Private/Widgets/Scalability/SNiagaraSystemResolvedScalabilitySettings.cpp
@@ -519,7 +519,10 @@ TSharedRef<SWidget> SNiagaraSystemResolvedScalabilitySettings::GenerateResolvedS
 	FRequiredInstanceInformation InstanceInfo;
 	InstanceInfo.DefaultPropertyRowGenerator = DefaultPropertyGenerator;
 	InstanceInfo.ResolvedPropertyRowGenerator = ResolvedPropertyGenerator;
-	InstanceInfo.ScalabilityValues = ScalabilityValues;
+	for (const TSharedRef<FScalabilityRowData>& Value : ScalabilityValues)
+	{
+		InstanceInfo.ScalabilityValues->Add(Value);
+	}
 	InstanceInformation.Add(Object, InstanceInfo);
 	const FSlateBrush* AssetIcon = FSlateIconFinder::FindIconBrushForClass(Object->GetClass());
 	
@@ -548,7 +551,7 @@ TSharedRef<SWidget> SNiagaraSystemResolvedScalabilitySettings::GenerateResolvedS
 		.BorderImage(FAppStyle::GetBrush(TEXT("Graph.TitleBackground")))
 		[
 			SNew(SListView<TSharedRef<FScalabilityRowData>>)
-			.ListItemsSource(&InstanceInformation[Object].ScalabilityValues)
+			.ListItemsSource(InstanceInformation[Object].ScalabilityValues)
 			.OnGenerateRow(this, &SNiagaraSystemResolvedScalabilitySettings::GenerateScalabilityValueRow)
 			.HeaderRow
 			(
diff --git a/Engine/Plugins/FX/Niagara/Source/NiagaraEditor/Private/Widgets/Scalability/SNiagaraSystemResolvedScalabilitySettings.h b/Engine/Plugins/FX/Niagara/Source/NiagaraEditor/Private/Widgets/Scalability/SNiagaraSystemResolvedScalabilitySettings.h
--- a/Engine/Plugins/FX/Niagara/Source/NiagaraEditor/Private/Widgets/Scalability/SNiagaraSystemResolvedScalabilitySettings.h
+++ b/Engine/Plugins/FX/Niagara/Source/NiagaraEditor/Private/Widgets/Scalability/SNiagaraSystemResolvedScalabilitySettings.h
@@ -48,7 +48,7 @@ class SNiagaraSystemResolvedScalabilitySettings : public SCompoundWidget
 	{
 		TSharedPtr<IPropertyRowGenerator> DefaultPropertyRowGenerator;
 		TSharedPtr<IPropertyRowGenerator> ResolvedPropertyRowGenerator;
-		TArray<TSharedRef<FScalabilityRowData>> ScalabilityValues;
+		TSharedRef<UE::Slate::Containers::TObservableArray<TSharedRef<FScalabilityRowData>>> ScalabilityValues = MakeShared<UE::Slate::Containers::TObservableArray<TSharedRef<FScalabilityRowData>>>();
 	};
 public:
 	SLATE_BEGIN_ARGS(SNiagaraSystemResolvedScalabilitySettings)
