Niagara: Export the asset builder symbols to resolve test build link issues.

#rnx
#Synced-CL 52015945

diff --git a/Engine/Plugins/FX/Niagara/Source/NiagaraEditor/Internal/AssetBuildUtilities/NiagaraAssetBuilder.h b/Engine/Plugins/FX/Niagara/Source/NiagaraEditor/Internal/AssetBuildUtilities/NiagaraAssetBuilder.h
--- a/Engine/Plugins/FX/Niagara/Source/NiagaraEditor/Internal/AssetBuildUtilities/NiagaraAssetBuilder.h
+++ b/Engine/Plugins/FX/Niagara/Source/NiagaraEditor/Internal/AssetBuildUtilities/NiagaraAssetBuilder.h
@@ -23,7 +23,7 @@ public:
 	 *
 	 * @return reference to this
 	 */
-	FNiagaraEmitterAssetBuilder& WithEmitterToCopy(UNiagaraEmitter* Emitter);
+	NIAGARAEDITOR_API FNiagaraEmitterAssetBuilder& WithEmitterToCopy(UNiagaraEmitter* Emitter);
 
 	/**
 	 * Build the NiagaraEmitter.
@@ -34,7 +34,7 @@ public:
 	 * 
 	 * @note Copied over from `UNiagaraEmitterFactoryNew` to mimic asset creation as close as possible
 	 */
-	UNiagaraEmitter* BuildEmitter(FString& OutError);
+	NIAGARAEDITOR_API UNiagaraEmitter* BuildEmitter(FString& OutError);
 
 private:
 	/** Package path of the new NiagaraEmitter asset. */
@@ -61,7 +61,7 @@ public:
 	 *
 	 * @return reference to this
 	 */
-	FNiagaraSystemAssetBuilder& WithSystemToCopy(UNiagaraSystem* System);
+	NIAGARAEDITOR_API FNiagaraSystemAssetBuilder& WithSystemToCopy(UNiagaraSystem* System);
 
 	/**
 	 * Add an Emitter to the new System.
@@ -70,7 +70,7 @@ public:
 	 *
 	 * @return reference to this
 	 */
-	FNiagaraSystemAssetBuilder& WithEmitter(FVersionedNiagaraEmitter&& Emitter);
+	NIAGARAEDITOR_API FNiagaraSystemAssetBuilder& WithEmitter(FVersionedNiagaraEmitter&& Emitter);
 
 	/**
 	 * Build the NiagaraSystem.
@@ -81,7 +81,7 @@ public:
 	 *
 	 * @note Copied over from `UNiagaraSystemFactoryNew` to mimic asset creation as close as possible
 	 */
-	UNiagaraSystem* BuildSystem(FString& OutError);
+	NIAGARAEDITOR_API UNiagaraSystem* BuildSystem(FString& OutError);
 
 private:
 	/** Package path of the new NiagaraSystem asset. */
