- Ensure parameters are already bound before initializing data interfaces

#Synced-CL 52062699

diff --git a/Engine/Plugins/FX/Niagara/Source/Niagara/Private/NiagaraSystemInstance.cpp b/Engine/Plugins/FX/Niagara/Source/Niagara/Private/NiagaraSystemInstance.cpp
--- a/Engine/Plugins/FX/Niagara/Source/Niagara/Private/NiagaraSystemInstance.cpp
+++ b/Engine/Plugins/FX/Niagara/Source/Niagara/Private/NiagaraSystemInstance.cpp
@@ -950,6 +950,7 @@ void FNiagaraSystemInstance::ResetInternal(bool bResetSimulations)
 	check(World);
 	if (OverrideParameters && World->WorldType == EWorldType::Editor)
 	{
+		BindParameters();
 		InitDataInterfaces();
 	}
 #endif
