- Add bindable velocity scale to cone / point velocity
- Ensure linear velocity is used in other cases to make it more clear

#Synced-CL 52062699

diff --git a/Engine/Plugins/FX/Niagara/Source/Niagara/Internal/Stateless/Modules/NiagaraStatelessModuleCommon.h b/Engine/Plugins/FX/Niagara/Source/Niagara/Internal/Stateless/Modules/NiagaraStatelessModuleCommon.h
--- a/Engine/Plugins/FX/Niagara/Source/Niagara/Internal/Stateless/Modules/NiagaraStatelessModuleCommon.h
+++ b/Engine/Plugins/FX/Niagara/Source/Niagara/Internal/Stateless/Modules/NiagaraStatelessModuleCommon.h
@@ -130,8 +130,8 @@ namespace NiagaraStateless
 		FNiagaraStatelessRangeFloat		MassRange = FNiagaraStatelessRangeFloat(1.0f);
 		FNiagaraStatelessRangeFloat		DragRange = FNiagaraStatelessRangeFloat(0.0f);
 
-		ENiagaraCoordinateSpace			VelocityCoordinateSpace = ENiagaraCoordinateSpace::Local;
-		FNiagaraStatelessRangeVector3	VelocityRange = FNiagaraStatelessRangeVector3(FVector3f::ZeroVector);
+		ENiagaraCoordinateSpace			LinearVelocityCoordinateSpace = ENiagaraCoordinateSpace::Local;
+		FNiagaraStatelessRangeVector3	LinearVelocityRange = FNiagaraStatelessRangeVector3(FVector3f::ZeroVector);
 		FNiagaraStatelessRangeFloat		LinearVelocityScale = FNiagaraStatelessRangeFloat(1.0f);
 
 		ENiagaraCoordinateSpace			WindCoordinateSpace = ENiagaraCoordinateSpace::Local;
@@ -148,6 +148,7 @@ namespace NiagaraStateless
 		FNiagaraStatelessRangeRotator	ConeRotator = FNiagaraStatelessRangeRotator(FRotator3f::ZeroRotator);
 		FNiagaraStatelessRangeVector3	ConeDirection = FNiagaraStatelessRangeVector3(FVector3f::XAxisVector);
 		FNiagaraStatelessRangeFloat		ConeVelocityRange = FNiagaraStatelessRangeFloat(0.0f);
+		FNiagaraStatelessRangeFloat		ConeVelocityScale = FNiagaraStatelessRangeFloat(1.0f);
 		float							ConeOuterAngle = 0.0f;
 		float							ConeInnerAngle = 0.0f;
 		float							ConeVelocityFalloff = 0.0f;
@@ -155,6 +156,7 @@ namespace NiagaraStateless
 		bool							bPointVelocity = false;
 		ENiagaraCoordinateSpace			PointCoordinateSpace = ENiagaraCoordinateSpace::Local;
 		FNiagaraStatelessRangeFloat		PointVelocityRange = FNiagaraStatelessRangeFloat(0.0f);
+		FNiagaraStatelessRangeFloat		PointVelocityScale = FNiagaraStatelessRangeFloat(1.0f);
 		float							PointVelocityMax = 0.0f;
 		FVector3f						PointOrigin = FVector3f::ZeroVector;
 
diff --git a/Engine/Plugins/FX/Niagara/Source/Niagara/Internal/Stateless/Modules/NiagaraStatelessModule_AddVelocity.h b/Engine/Plugins/FX/Niagara/Source/Niagara/Internal/Stateless/Modules/NiagaraStatelessModule_AddVelocity.h
--- a/Engine/Plugins/FX/Niagara/Source/Niagara/Internal/Stateless/Modules/NiagaraStatelessModule_AddVelocity.h
+++ b/Engine/Plugins/FX/Niagara/Source/Niagara/Internal/Stateless/Modules/NiagaraStatelessModule_AddVelocity.h
@@ -34,13 +34,15 @@ public:
 	UPROPERTY(EditAnywhere, Category = "Parameters", meta=(SegmentedDisplay))
 	ENSM_VelocityType VelocityType = ENSM_VelocityType::Linear;
 
-	UPROPERTY(EditAnywhere, Category = "Parameters", meta = (DisplayName = "Velocity", DisableUniformDistribution, EditConditionHides, EditCondition = "VelocityType == ENSM_VelocityType::Linear"))
+	UPROPERTY(EditAnywhere, Category = "Parameters", meta = (DisplayName = "Linear Velocity", DisableUniformDistribution, EditConditionHides, EditCondition = "VelocityType == ENSM_VelocityType::Linear"))
 	FNiagaraDistributionRangeVector3 LinearVelocityDistribution = FNiagaraDistributionRangeVector3(FVector3f(0.0f, 0.0f, 100.0f));
-	UPROPERTY(EditAnywhere, Category = "Parameters", meta = (DisplayName = "Velocity Scale", DisableRangeDistribution, EditConditionHides, EditCondition = "VelocityType == ENSM_VelocityType::Linear"))
+	UPROPERTY(EditAnywhere, Category = "Parameters", meta = (DisplayName = "Linear Velocity Scale", DisableRangeDistribution, EditConditionHides, EditCondition = "VelocityType == ENSM_VelocityType::Linear"))
 	FNiagaraDistributionRangeFloat LinearVelocityScale = FNiagaraDistributionRangeFloat(1.0f);
 
 	UPROPERTY(EditAnywhere, Category = "Parameters", meta = (DisplayName = "Cone Velocity", EditConditionHides, EditCondition = "VelocityType == ENSM_VelocityType::InCone"))
 	FNiagaraDistributionRangeFloat ConeVelocityDistribution = FNiagaraDistributionRangeFloat(500.0f);
+	UPROPERTY(EditAnywhere, Category = "Parameters", meta = (DisplayName = "Cone Velocity Scale", DisableRangeDistribution, EditConditionHides, EditCondition = "VelocityType == ENSM_VelocityType::InCone"))
+	FNiagaraDistributionRangeFloat ConeVelocityScale = FNiagaraDistributionRangeFloat(1.0f);
 	UPROPERTY(EditAnywhere, Category = "Parameters", meta = (SegmentedDisplay, EditConditionHides, EditCondition = "VelocityType == ENSM_VelocityType::InCone"))
 	ENSM_ConeRotationType ConeRotationType = ENSM_ConeRotationType::Rotation;
 	UPROPERTY(EditAnywhere, Category = "Parameters", meta = (EditConditionHides, EditCondition = "VelocityType == ENSM_VelocityType::InCone && ConeRotationType == ENSM_ConeRotationType::Rotation", DisableRangeDistribution, DisableUniformDistribution))
@@ -55,6 +57,8 @@ public:
 
 	UPROPERTY(EditAnywhere, Category = "Parameters", meta = (DisplayName = "Point Velocity", EditConditionHides, EditCondition = "VelocityType == ENSM_VelocityType::FromPoint"))
 	FNiagaraDistributionRangeFloat PointVelocityDistribution = FNiagaraDistributionRangeFloat(100.0f);
+	UPROPERTY(EditAnywhere, Category = "Parameters", meta = (DisplayName = "Point Velocity Scale", DisableRangeDistribution, EditConditionHides, EditCondition = "VelocityType == ENSM_VelocityType::FromPoint"))
+	FNiagaraDistributionRangeFloat PointVelocityScale = FNiagaraDistributionRangeFloat(1.0f);
 	UPROPERTY(EditAnywhere, Category = "Parameters", meta = (EditConditionHides, EditCondition = "VelocityType == ENSM_VelocityType::FromPoint"))
 	FVector3f PointOrigin = FVector3f::ZeroVector;
 
diff --git a/Engine/Plugins/FX/Niagara/Source/Niagara/Private/Stateless/Modules/NiagaraStatelessModule_AddVelocity.cpp b/Engine/Plugins/FX/Niagara/Source/Niagara/Private/Stateless/Modules/NiagaraStatelessModule_AddVelocity.cpp
--- a/Engine/Plugins/FX/Niagara/Source/Niagara/Private/Stateless/Modules/NiagaraStatelessModule_AddVelocity.cpp
+++ b/Engine/Plugins/FX/Niagara/Source/Niagara/Private/Stateless/Modules/NiagaraStatelessModule_AddVelocity.cpp
@@ -16,8 +16,8 @@ void UNiagaraStatelessModule_AddVelocity::BuildEmitterData(const FNiagaraStatele
 	if (VelocityType == ENSM_VelocityType::Linear)
 	{
 		const FNiagaraStatelessRangeVector3 VelocityRange = BuildContext.ConvertDistributionToRange(LinearVelocityDistribution, FVector3f::ZeroVector);
-		PhysicsBuildData.VelocityCoordinateSpace = CoordinateSpace;
-		PhysicsBuildData.VelocityRange = BuildContext.ConvertDistributionToRange(LinearVelocityDistribution, FVector3f::ZeroVector);
+		PhysicsBuildData.LinearVelocityCoordinateSpace = CoordinateSpace;
+		PhysicsBuildData.LinearVelocityRange = BuildContext.ConvertDistributionToRange(LinearVelocityDistribution, FVector3f::ZeroVector);
 		PhysicsBuildData.LinearVelocityScale = BuildContext.ConvertDistributionToRange(LinearVelocityScale, 1.0f);
 	}
 	else if (VelocityType == ENSM_VelocityType::FromPoint)
@@ -27,6 +27,7 @@ void UNiagaraStatelessModule_AddVelocity::BuildEmitterData(const FNiagaraStatele
 		PhysicsBuildData.bPointVelocity = true;
 		PhysicsBuildData.PointCoordinateSpace = CoordinateSpace;
 		PhysicsBuildData.PointVelocityRange = BuildContext.ConvertDistributionToRange(PointVelocityDistribution, 0.0f);
+		PhysicsBuildData.PointVelocityScale = BuildContext.ConvertDistributionToRange(PointVelocityScale, 1.0f);
 		PhysicsBuildData.PointOrigin = PointOrigin;
 	}
 	else if (VelocityType == ENSM_VelocityType::InCone)
@@ -39,6 +40,7 @@ void UNiagaraStatelessModule_AddVelocity::BuildEmitterData(const FNiagaraStatele
 		PhysicsBuildData.ConeRotator = ConeRotationType == ENSM_ConeRotationType::Rotation ? BuildContext.ConvertDistributionToRange(ConeRotation, FRotator3f::ZeroRotator) : FNiagaraStatelessRangeRotator(FRotator3f::ZeroRotator);
 		PhysicsBuildData.ConeDirection = ConeRotationType == ENSM_ConeRotationType::Direction ? BuildContext.ConvertDistributionToRange(ConeDirection, FVector3f::ZAxisVector) : FNiagaraStatelessRangeVector3(FVector3f::ZAxisVector);
 		PhysicsBuildData.ConeVelocityRange = BuildContext.ConvertDistributionToRange(ConeVelocityDistribution, 0.0f);
+		PhysicsBuildData.ConeVelocityScale = BuildContext.ConvertDistributionToRange(ConeVelocityScale, 1.0f);
 		PhysicsBuildData.ConeOuterAngle = ConeAngle;
 		PhysicsBuildData.ConeInnerAngle = InnerCone;
 		PhysicsBuildData.ConeVelocityFalloff = bSpeedFalloffFromConeAxisEnabled ? FMath::Clamp(SpeedFalloffFromConeAxis, 0.0f, 1.0f) : 0.0f;
@@ -52,73 +54,81 @@ void UNiagaraStatelessModule_AddVelocity::DrawDebug(const FNiagaraStatelessDrawD
 
 	switch (VelocityType)
 	{
-		case ENSM_VelocityType::Linear:
-		{
-			const FNiagaraStatelessRangeVector3 VelocityRange = LinearVelocityDistribution.CalculateRange(FVector3f::ZeroVector);
-			const FNiagaraStatelessRangeFloat LinearVelocityScaleRange = LinearVelocityScale.CalculateRange();
-			const FVector MinDir = FVector(VelocityRange.Min * LinearVelocityScale.Min);
-			const FVector MaxDir = FVector(VelocityRange.Max * LinearVelocityScale.Max);
-			DrawDebugContext.DrawArrow(FVector::ZeroVector, MinDir);
+	case ENSM_VelocityType::Linear:
+	{
+		const FNiagaraStatelessRangeVector3 LinearVelocityRange = LinearVelocityDistribution.CalculateRange(FVector3f::ZeroVector);
+		const FNiagaraStatelessRangeFloat LinearVelocityScaleRange = LinearVelocityScale.CalculateRange();
+		const FVector MinDir = FVector(LinearVelocityRange.Min * LinearVelocityScaleRange.Min * LinearVelocityRange.Min);
+		const FVector MaxDir = FVector(LinearVelocityRange.Max * LinearVelocityScaleRange.Max * LinearVelocityRange.Max);
+		DrawDebugContext.DrawArrow(FVector::ZeroVector, MinDir);
 
-			if (!FMath::IsNearlyEqual(MinDir.Length(), MaxDir.Length()))
-			{
-				DrawDebugContext.DrawArrow(WorldOrigin, MaxDir);
-			}
-			break;
+		if (!FMath::IsNearlyEqual(MinDir.Length(), MaxDir.Length()))
+		{
+			DrawDebugContext.DrawArrow(WorldOrigin, MaxDir);
 		}
+		break;
+	}
 
-		case ENSM_VelocityType::InCone:
+	case ENSM_VelocityType::InCone:
+	{
+		FQuat4f LocalConeQuat;
+		if (ConeRotationType == ENSM_ConeRotationType::Rotation)
 		{
-			FQuat4f LocalConeQuat;
-			if (ConeRotationType == ENSM_ConeRotationType::Rotation)
+			if (ConeRotation.IsBinding())
 			{
-				if (ConeRotation.IsBinding())
-				{
-					return;
-				}
-
-				LocalConeQuat = ConeRotation.CalculateRange().Min.Quaternion();
+				return;
 			}
-			else
+
+			LocalConeQuat = ConeRotation.CalculateRange().Min.Quaternion();
+		}
+		else
+		{
+			if (ConeDirection.IsBinding())
 			{
-				if (ConeDirection.IsBinding())
-				{
-					return;
-				}
-				LocalConeQuat = NiagaraStateless::DirectionToQuat(ConeDirection.CalculateRange().Min);
+				return;
 			}
+			LocalConeQuat = NiagaraStateless::DirectionToQuat(ConeDirection.CalculateRange().Min);
+		}
 
-			const FQuat ConeQuat = DrawDebugContext.TransformRotation(LocalConeQuat);
-			const float ConeHAngle = ConeAngle / 2.0f;
+		const FQuat ConeQuat = DrawDebugContext.TransformRotation(LocalConeQuat);
+		const float ConeHAngle = ConeAngle / 2.0f;
 
-			TOptional<float> InnerConeHAngle;
-			if (InnerCone > 0.0f && !FMath::IsNearlyEqual(ConeAngle, InnerCone))
-			{
-				InnerConeHAngle = InnerCone / 2.0f;
-			}
+		TOptional<float> InnerConeHAngle;
+		if (InnerCone > 0.0f && !FMath::IsNearlyEqual(ConeAngle, InnerCone))
+		{
+			InnerConeHAngle = InnerCone / 2.0f;
+		}
 
-			const FNiagaraStatelessRangeFloat ConeVelocityRange = ConeVelocityDistribution.CalculateRange(0.0f);
-			DrawDebugContext.DrawCone(WorldOrigin, ConeQuat, ConeHAngle, ConeVelocityRange.Min);
-			if ( InnerConeHAngle.IsSet() )
-			{
-				DrawDebugContext.DrawCone(WorldOrigin, ConeQuat, InnerConeHAngle.GetValue(), ConeVelocityRange.Min);
-			}
+		const FNiagaraStatelessRangeFloat ConeVelocityScaleRange = ConeVelocityScale.CalculateRange();
+		FNiagaraStatelessRangeFloat ConeVelocityRange = ConeVelocityDistribution.CalculateRange(0.0f);
+		ConeVelocityRange.Min *= ConeVelocityScaleRange.Min;
+		ConeVelocityRange.Max *= ConeVelocityScaleRange.Max;
+
+		DrawDebugContext.DrawCone(WorldOrigin, ConeQuat, ConeHAngle, ConeVelocityRange.Min);
+		if (InnerConeHAngle.IsSet())
+		{
+			DrawDebugContext.DrawCone(WorldOrigin, ConeQuat, InnerConeHAngle.GetValue(), ConeVelocityRange.Min);
+		}
 
-			if (!FMath::IsNearlyEqual(ConeVelocityRange.Min, ConeVelocityRange.Max))
+		if (!FMath::IsNearlyEqual(ConeVelocityRange.Min, ConeVelocityRange.Max))
+		{
+			DrawDebugContext.DrawCone(WorldOrigin, ConeQuat, ConeHAngle, ConeVelocityRange.Max);
+			if (InnerConeHAngle.IsSet())
 			{
-				DrawDebugContext.DrawCone(WorldOrigin, ConeQuat, ConeHAngle, ConeVelocityRange.Max);
-				if (InnerConeHAngle.IsSet())
-				{
-					DrawDebugContext.DrawCone(WorldOrigin, ConeQuat, InnerConeHAngle.GetValue(), ConeVelocityRange.Max);
-				}
+				DrawDebugContext.DrawCone(WorldOrigin, ConeQuat, InnerConeHAngle.GetValue(), ConeVelocityRange.Max);
 			}
-			break;
 		}
+		break;
+	}
 
-		case ENSM_VelocityType::FromPoint:
-		{
-			const FNiagaraStatelessRangeFloat PointVelocityRange = PointVelocityDistribution.CalculateRange(0.0f);
-			if (!FMath::IsNearlyEqual(PointVelocityRange.Min, 0.0f))
+	case ENSM_VelocityType::FromPoint:
+	{
+		const FNiagaraStatelessRangeFloat PointVelocityScaleRange = PointVelocityScale.CalculateRange();
+		FNiagaraStatelessRangeFloat PointVelocityRange = PointVelocityDistribution.CalculateRange(0.0f);
+		PointVelocityRange.Min *= PointVelocityScaleRange.Min;
+		PointVelocityRange.Max *= PointVelocityScaleRange.Max;
+
+		if (!FMath::IsNearlyEqual(PointVelocityRange.Min, 0.0f))
 			{
 				DrawDebugContext.DrawSphere(DrawDebugContext.TransformPosition(PointOrigin), PointVelocityRange.Min);
 			}
diff --git a/Engine/Plugins/FX/Niagara/Source/Niagara/Private/Stateless/Modules/NiagaraStatelessModule_SolveVelocitiesAndForces.cpp b/Engine/Plugins/FX/Niagara/Source/Niagara/Private/Stateless/Modules/NiagaraStatelessModule_SolveVelocitiesAndForces.cpp
--- a/Engine/Plugins/FX/Niagara/Source/Niagara/Private/Stateless/Modules/NiagaraStatelessModule_SolveVelocitiesAndForces.cpp
+++ b/Engine/Plugins/FX/Niagara/Source/Niagara/Private/Stateless/Modules/NiagaraStatelessModule_SolveVelocitiesAndForces.cpp
@@ -227,11 +227,15 @@ void UNiagaraStatelessModule_SolveVelocitiesAndForces::SetShaderParameters(const
 	SetShaderParameterContext.ConvertRangeToScaleBias(PhysicsData.MassRange, Parameters->SolveVelocitiesAndForces_MassScale, Parameters->SolveVelocitiesAndForces_MassBias);
 	SetShaderParameterContext.ConvertRangeToScaleBias(PhysicsData.DragRange, Parameters->SolveVelocitiesAndForces_DragScale, Parameters->SolveVelocitiesAndForces_DragBias);
 
-	SetShaderParameterContext.TransformVectorNoScaleRangeToScaleBias(PhysicsData.VelocityRange, PhysicsData.VelocityCoordinateSpace, Parameters->SolveVelocitiesAndForces_VelocityScale, Parameters->SolveVelocitiesAndForces_VelocityBias);
-	const float VelocityScale = SetShaderParameterContext.ConvertRangeToValue(PhysicsData.LinearVelocityScale);
-	Parameters->SolveVelocitiesAndForces_VelocityScale *= VelocityScale;
-	Parameters->SolveVelocitiesAndForces_VelocityBias  *= VelocityScale;
+	// Linear Velocity
+	{
+		SetShaderParameterContext.TransformVectorNoScaleRangeToScaleBias(PhysicsData.LinearVelocityRange, PhysicsData.LinearVelocityCoordinateSpace, Parameters->SolveVelocitiesAndForces_VelocityScale, Parameters->SolveVelocitiesAndForces_VelocityBias);
+		const float LinearVelocityScale = SetShaderParameterContext.ConvertRangeToValue(PhysicsData.LinearVelocityScale);
+		Parameters->SolveVelocitiesAndForces_VelocityScale *= LinearVelocityScale;
+		Parameters->SolveVelocitiesAndForces_VelocityBias *= LinearVelocityScale;
+	}
 
+	// Wind / Acceleration
 	SetShaderParameterContext.TransformVectorNoScaleRangeToScaleBias(PhysicsData.WindRange, PhysicsData.WindCoordinateSpace, Parameters->SolveVelocitiesAndForces_WindScale, Parameters->SolveVelocitiesAndForces_WindBias);
 	{
 		FVector3f AccelerationScale, AccelerationBias;
@@ -244,6 +248,7 @@ void UNiagaraStatelessModule_SolveVelocitiesAndForces::SetShaderParameters(const
 		Parameters->SolveVelocitiesAndForces_GravityBias	= SpaceTransforms.TransformVectorNoScale(ENiagaraCoordinateSpace::World, GravityBias);
 	}
 
+	// Cone Velocity
 	Parameters->SolveVelocitiesAndForces_ConeVelocityEnabled = PhysicsData.bConeVelocity ? 1 : 0;
 	{
 		FQuat4f LocalConeQuat;
@@ -262,12 +267,21 @@ void UNiagaraStatelessModule_SolveVelocitiesAndForces::SetShaderParameters(const
 		Parameters->SolveVelocitiesAndForces_ConeAngleScale			= (PhysicsData.ConeOuterAngle - PhysicsData.ConeInnerAngle) * (UE_PI / 360.0f);
 		Parameters->SolveVelocitiesAndForces_ConeAngleBias			= PhysicsData.ConeInnerAngle * (UE_PI / 360.0f);
 		Parameters->SolveVelocitiesAndForces_ConeVelocityFalloff	= PhysicsData.ConeVelocityFalloff;
+
+		const float ConeVelocityScale = SetShaderParameterContext.ConvertRangeToValue(PhysicsData.ConeVelocityScale);
+		Parameters->SolveVelocitiesAndForces_ConeVelocityScale *= ConeVelocityScale;
+		Parameters->SolveVelocitiesAndForces_ConeVelocityBias  *= ConeVelocityScale;
 	}
 
+	// Point Velocity
 	Parameters->SolveVelocitiesAndForces_PointVelocityEnabled = PhysicsData.bPointVelocity ? 1 : 0;
 	{
 		SetShaderParameterContext.ConvertRangeToScaleBias(PhysicsData.PointVelocityRange, Parameters->SolveVelocitiesAndForces_PointVelocityScale, Parameters->SolveVelocitiesAndForces_PointVelocityBias);
 		Parameters->SolveVelocitiesAndForces_PointOrigin = SpaceTransforms.TransformPosition(PhysicsData.PointCoordinateSpace, PhysicsData.PointOrigin);
+
+		const float PointVelocityScale = SetShaderParameterContext.ConvertRangeToValue(PhysicsData.PointVelocityScale);
+		Parameters->SolveVelocitiesAndForces_PointVelocityScale *= PointVelocityScale;
+		Parameters->SolveVelocitiesAndForces_PointVelocityBias  *= PointVelocityScale;
 	}
 
 	{
