A few small niagara assistant fixes.

Fix incorrect BaseStruct meta tag on some types which was breaking the automated schema generation and reducing assistant understanding.
Fixed issue with stack input traversal causing many stack inputs to be always hidden.
Including display text in enum name for greater understanding of enum values.
Adding a NiagaraWorkerSkill asset to give additional guidance to the EDA for Niagara tasks.
Stripped module namespace off the module input names generated from the asset scripts. Helps improve understanding of these schemas vs what is produced from the stack.
Remove now redundant prompt helpers.
Improved tooltip on FNiagaraBool to directly tell assistant the special case int32 values required for True and False.

#virtualized
#Synced-CL 52000591

diff --git a/Engine/Plugins/FX/Niagara/Source/Niagara/Public/NiagaraTypes.h b/Engine/Plugins/FX/Niagara/Source/Niagara/Public/NiagaraTypes.h
--- a/Engine/Plugins/FX/Niagara/Source/Niagara/Public/NiagaraTypes.h
+++ b/Engine/Plugins/FX/Niagara/Source/Niagara/Public/NiagaraTypes.h
@@ -72,7 +72,7 @@ struct FNiagaraBool
 	inline operator bool() const { return GetValue(); }
 
 private:
-	UPROPERTY(EditAnywhere, Category = Parameters)// Must be either FNiagaraBool::True or FNiagaraBool::False.
+	UPROPERTY(EditAnywhere, Category = Parameters)// Must be either FNiagaraBool::True(-1) or FNiagaraBool::False(0).
 	int32 Value = False;
 };
 
diff --git a/Engine/Plugins/FX/Niagara/Source/NiagaraEditor/Private/NiagaraExternalSystemEditorUtilities.cpp b/Engine/Plugins/FX/Niagara/Source/NiagaraEditor/Private/NiagaraExternalSystemEditorUtilities.cpp
--- a/Engine/Plugins/FX/Niagara/Source/NiagaraEditor/Private/NiagaraExternalSystemEditorUtilities.cpp
+++ b/Engine/Plugins/FX/Niagara/Source/NiagaraEditor/Private/NiagaraExternalSystemEditorUtilities.cpp
@@ -353,7 +353,10 @@ void ForEachFunctionInput(UNiagaraStackEntry* Entry, TFunction<bool(UNiagaraStac
 			{
 				if (UNiagaraStackFunctionInput* StackFunctionInput = Cast<UNiagaraStackFunctionInput>(ChildItem))
 				{
-					return Func(StackFunctionInput);
+					if(Func(StackFunctionInput))//If we return false here we should keep iterating to sub items.
+					{
+						return true;
+					}
 				}
 
 				//Look in child items for non-inputs. This may be too permissive. 
@@ -718,8 +721,8 @@ UNiagaraStackFunctionInput* FNiagaraExt_StackItemReference::GetInput(FNiagaraExt
 							{
 								//End of input chain and we found our match.
 								CachedInput = Input;
+								return true;
 							}
-							return true;
 						}
 						return false;
 					};
@@ -1064,7 +1067,12 @@ void UNiagaraExternalEditUtilities::GetModuleSchema(const UNiagaraScript* Module
 					{
 						FNiagaraVariableBase Var = UEdGraphSchema_Niagara::PinToNiagaraVariable(InputPin);
 						FNiagaraExt_StackInputSchema& NewInputSchema = OutSchema.Inputs.AddDefaulted_GetRef();
-						NewInputSchema.Name = Var.GetName();
+						
+						//Strip the module namespace so we just get the raw name of the input.
+						FString NameStr = Var.GetName().ToString();
+						NameStr.RemoveFromStart(PARAM_MAP_MODULE_STR);
+						
+						NewInputSchema.Name = *NameStr;
 						NewInputSchema.Type = Var.GetType();
 
 						NewInputSchema.bSupportsExpressions = bSupportsExpressions;
@@ -1856,6 +1864,7 @@ void FNiagaraExt_VariableValue::Set(const FNiagaraTypeDefinition& TypeDef, const
 			if (Enum)
 			{
 				EnumVal.EnumName = Enum->GetNameByValue(EnumValue);
+				EnumVal.DisplayName = Enum->GetDisplayNameTextByValue(EnumValue);
 			}
 		}
 		else if (TypeDef.IsDataInterface())
@@ -1937,6 +1946,7 @@ void FNiagaraExt_StackInputValue::InitFromStackInput(const FNiagaraExt_StackItem
 				{
 					int32 EnumValue = *(int32*)LocalData->GetStructMemory();
 					EnumVal.EnumName = Enum->GetNameByValue(EnumValue);
+					EnumVal.DisplayName = Enum->GetDisplayNameTextByValue(EnumValue);
 				}
 			}
 			else
diff --git a/Engine/Plugins/FX/Niagara/Source/NiagaraEditor/Public/NiagaraExternalSystemEditorUtilities.h b/Engine/Plugins/FX/Niagara/Source/NiagaraEditor/Public/NiagaraExternalSystemEditorUtilities.h
--- a/Engine/Plugins/FX/Niagara/Source/NiagaraEditor/Public/NiagaraExternalSystemEditorUtilities.h
+++ b/Engine/Plugins/FX/Niagara/Source/NiagaraEditor/Public/NiagaraExternalSystemEditorUtilities.h
@@ -69,17 +69,17 @@ struct FNiagaraExt_VariableValue_Enum : public FNiagaraExt_VariableValueBase
 {
 	GENERATED_BODY()
 
-	/** The UEnum type definition */
-	UPROPERTY(EditAnywhere, Category="Variable", meta=(ToolTip="Reference to the UEnum type"))
+	/** Reference to the UEnum type. */
+	UPROPERTY(EditAnywhere, Category="Variable")
 	TObjectPtr<UEnum> Enum;
 
-	/** 
-	The selected enum name as a string (e.g., \"MyEnum::Name1\"). 
-	For some enum names this can be an auto generated and un-descriptive name such as EnumValue0.
-	The display name is often much more useful and descriptive. Access this via the schema for the UEnum.
-	*/
+	/** The selected enum name (e.g., \"MyEnum::Name1\"). */
 	UPROPERTY(EditAnywhere, Category="Variable", meta=(ToolTip="The enum value name"))
 	FName EnumName;
+	
+	/** Display value for EnumName. Often more descriptive and useful that the internal EnunName.  For informational purposes only, the actual set value of the enum is EnumName. */
+	UPROPERTY(EditAnywhere, Category="Variable")
+	FText DisplayName;
 };
 
 /**
@@ -325,7 +325,7 @@ struct FNiagaraExt_RendererData
 /** An instanced struct type that holds the value of a Niagara Stack Input. */
 USTRUCT(BlueprintInternalUseOnly, 
 meta=(
-BaseStruct="/Script/NiagaraEditor.FNiagaraExt_StackInputData",
+BaseStruct="/Script/NiagaraEditor.NiagaraExt_StackInputData",
 AdditionalTypes="/Script/Niagara.NiagaraFloat,/Script/Niagara.NiagaraBool,/Script/Niagara.NiagaraInt32,/Script/Niagara.NiagaraPosition,/Script/Niagara.NiagaraID,/Script/CoreUObject.Vector3f,/Script/CoreUObject.Vector2f,/Script/CoreUObject.Vector4f,/Script/CoreUObject.Quat4f,/Script/CoreUObject.LinearColor,/Script/CoreUObject.Rotator"
 ))
 struct FNiagaraExt_StackInputValue : public FNiagaraExt_InstancedValue
@@ -349,17 +349,17 @@ struct FNiagaraExt_StackInputData_Enum : public FNiagaraExt_StackInputData
 {
 	GENERATED_BODY()
 
-	/** The UEnum type definition */
-	UPROPERTY(EditAnywhere, Category="Input", meta=(ToolTip="Reference to the UEnum type"))
+	/** Reference to the UEnum type. */
+	UPROPERTY(EditAnywhere, Category="Variable")
 	TObjectPtr<UEnum> Enum;
 
-	/** 
-	The selected enum name as a string (e.g., \"MyEnum::Name1\"). 
-	For some enum names this can be an auto generated and un-descriptive name such as EnumValue0.
-	The display name is often much more useful and descriptive. Access this via the schema for the UEnum.
-	*/
-	UPROPERTY(EditAnywhere, Category="Input", meta=(ToolTip="The enum value name"))
+	/** The selected enum name (e.g., \"MyEnum::Name1\"). */
+	UPROPERTY(EditAnywhere, Category="Variable", meta=(ToolTip="The enum value name"))
 	FName EnumName;
+	
+	/** Display value for EnumName. Often more descriptive and useful that the internal EnunName. For informational purposes only, the actual set value of the enum is EnumName. */
+	UPROPERTY(EditAnywhere, Category="Variable")
+	FText DisplayName;
 };
 
 USTRUCT(BlueprintInternalUseOnly)
@@ -416,7 +416,7 @@ struct FNiagaraExt_StackInputData_Unsupported : public FNiagaraExt_StackInputDat
 //Topology Layer - Defines how Niagara Systems, Emitters, Modules and their inputs are arranged.
 
 /** Instanced Struct wrapper around Stack Input Topology allowing us to compose them hierarchically for dynamic input chains. */
-USTRUCT(BlueprintInternalUseOnly, meta=(BaseStruct="/Script/NiagaraEditor.FNiagaraExt_StackInputTopology"))
+USTRUCT(BlueprintInternalUseOnly, meta=(BaseStruct="/Script/NiagaraEditor.NiagaraExt_StackInputTopology"))
 struct FNiagaraExt_DynamicInputTopology : public FNiagaraExt_InstancedValue
 {
 	GENERATED_BODY()
