- Fix for hitting check where null is already handled

#Synced-CL 52062699

diff --git a/Engine/Plugins/FX/Niagara/Source/NiagaraEditor/Private/ViewModels/Stack/NiagaraStackGraphUtilities.cpp b/Engine/Plugins/FX/Niagara/Source/NiagaraEditor/Private/ViewModels/Stack/NiagaraStackGraphUtilities.cpp
--- a/Engine/Plugins/FX/Niagara/Source/NiagaraEditor/Private/ViewModels/Stack/NiagaraStackGraphUtilities.cpp
+++ b/Engine/Plugins/FX/Niagara/Source/NiagaraEditor/Private/ViewModels/Stack/NiagaraStackGraphUtilities.cpp
@@ -454,7 +454,7 @@ OutputNodeType* GetEmitterOutputNodeForStackNodeInternal(InputNodeType& StackNod
 		{
 			for (UEdGraphPin* LinkedPin : OutputPin->LinkedTo)
 			{
-				InputNodeType* LinkedNiagaraNode = LinkedPin ? Cast<UNiagaraNode>(LinkedPin->GetOwningNode()) : nullptr;
+				InputNodeType* LinkedNiagaraNode = LinkedPin ? Cast<UNiagaraNode>(LinkedPin->GetOwningNodeUnchecked()) : nullptr;
 				if (LinkedNiagaraNode != nullptr && NodesSeen.Contains(LinkedNiagaraNode) == false)
 				{
 					NodesSeen.Add(LinkedNiagaraNode);
