- Fix using wrong value in the estimation count with GPU ribbons

#rnx

#Synced-CL 52062699

diff --git a/Engine/Plugins/FX/Niagara/Source/Niagara/Private/NiagaraRendererRibbons.cpp b/Engine/Plugins/FX/Niagara/Source/Niagara/Private/NiagaraRendererRibbons.cpp
--- a/Engine/Plugins/FX/Niagara/Source/Niagara/Private/NiagaraRendererRibbons.cpp
+++ b/Engine/Plugins/FX/Niagara/Source/Niagara/Private/NiagaraRendererRibbons.cpp
@@ -1058,7 +1058,7 @@ FNiagaraDynamicDataBase* FNiagaraRendererRibbons::GenerateDynamicData(const FNia
 			DynamicData->MaxAllocatedCountEstimate = 0;
 			if (FVersionedNiagaraEmitterData* EmitterData = Emitter->GetVersionedEmitter().GetEmitterData())
 			{
-				DynamicData->MaxAllocatedCountEstimate = FMath::Min<uint32>(EmitterData->GetMaxParticleCountEstimate(), DynamicData->MaxAllocatedCountEstimate);
+				DynamicData->MaxAllocatedCountEstimate = FMath::Min<uint32>(EmitterData->GetMaxParticleCountEstimate(), DynamicData->MaxAllocationCount);
 			}
 			
 			if (!DynamicData->bUseGPUInit)
