Add option to limit niagara's maximum PSO priority to  high priority instead of highest. Disabled by default to allow highest.

diff --git a/Engine/Plugins/FX/Niagara/Source/Niagara/Private/NiagaraScalabilitySettings.cpp b/Engine/Plugins/FX/Niagara/Source/Niagara/Private/NiagaraScalabilitySettings.cpp
--- a/Engine/Plugins/FX/Niagara/Source/Niagara/Private/NiagaraScalabilitySettings.cpp
+++ b/Engine/Plugins/FX/Niagara/Source/Niagara/Private/NiagaraScalabilitySettings.cpp
@@ -10,6 +10,19 @@ UE_DEFINE_GAMEPLAY_TAG_COMMENT(TAG_Niagara_ScalabilityProfile_High, "Niagara.Sca
 UE_DEFINE_GAMEPLAY_TAG_COMMENT(TAG_Niagara_ScalabilityProfile_Medium, "Niagara.ScalabilityProfile.Medium", "A medium importance element for Niagara FX that should balance performance and fidelity in a given context.")
 UE_DEFINE_GAMEPLAY_TAG_COMMENT(TAG_Niagara_ScalabilityProfile_Low, "Niagara.ScalabilityProfile.Low", "A low importance element for Niagara FX that should favour performance over fidelity in a given context.")
 
+namespace NSCVars
+{
+	static bool GLimitPSOPrecachePriorityBoost = false;
+	static FAutoConsoleVariableRef CVarNiagaraLimitPSOPrecachePriorityBoost(
+		TEXT("fx.Niagara.PSOPrecache.LimitPriorityBoost"),
+		GLimitPSOPrecachePriorityBoost,
+		TEXT("Limit Niagara to issuing only high priority PSO precache boosts.\n")
+		TEXT("false - Niagara can issue highest priority PSO precache boosts. (default)\n")
+		TEXT("true  - Niagara will issue only high priority PSO precache boosts.")
+		,
+		ECVF_Default
+	);
+}
 
 #if WITH_EDITOR
 
@@ -64,7 +77,7 @@ const FNiagaraReadinessSettingsProfile* UNiagaraSystemReadinessSettings::Find(FG
 
 EPSOPrecachePriority FNiagaraSystemPSOPrecacheSettings::GetPriorityBoost()const
 {
-	if (PreCachePriorityBoost == ENiagaraSystemPSOPrecachingPriorityBoost::Highest) return EPSOPrecachePriority::Highest;
+	if (PreCachePriorityBoost == ENiagaraSystemPSOPrecachingPriorityBoost::Highest) return  NSCVars::GLimitPSOPrecachePriorityBoost ? EPSOPrecachePriority::High : EPSOPrecachePriority::Highest;
 	else if (PreCachePriorityBoost == ENiagaraSystemPSOPrecachingPriorityBoost::High) return EPSOPrecachePriority::High;
 	else return EPSOPrecachePriority::Medium;
 }
\ No newline at end of file
