[Niagara] Convert a check to a checkf to help debug an editor test failure.

diff --git a/Engine/Plugins/FX/Niagara/Source/Niagara/Private/NiagaraTypeRegistry.cpp b/Engine/Plugins/FX/Niagara/Source/Niagara/Private/NiagaraTypeRegistry.cpp
--- a/Engine/Plugins/FX/Niagara/Source/Niagara/Private/NiagaraTypeRegistry.cpp
+++ b/Engine/Plugins/FX/Niagara/Source/Niagara/Private/NiagaraTypeRegistry.cpp
@@ -327,7 +327,9 @@ int32 FNiagaraTypeRegistry::RegisterTypeInternal(const FNiagaraTypeDefinition& N
 				if (ensure(RegisteredTypes.IsValidIndex(NewType.RegisteredTypeDefIndex)))
 				{
 					bNeedsRegistration = !RegisteredTypes[NewType.RegisteredTypeDefIndex]->IsValid();
-					check(bNeedsRegistration || (NewType == *RegisteredTypes[NewType.RegisteredTypeDefIndex]));
+					checkf(bNeedsRegistration || (NewType == *RegisteredTypes[NewType.RegisteredTypeDefIndex]),
+						TEXT("bNeedsRegistration: %d, NewType: %s"), bNeedsRegistration, *NewType.GetNameText().ToString() 
+					);
 				}
 			}
 
