[Niagara] Further expand a checkf when a new type doesn't match the old type with the same index.

diff --git a/Engine/Plugins/FX/Niagara/Source/Niagara/Private/NiagaraTypeRegistry.cpp b/Engine/Plugins/FX/Niagara/Source/Niagara/Private/NiagaraTypeRegistry.cpp
--- a/Engine/Plugins/FX/Niagara/Source/Niagara/Private/NiagaraTypeRegistry.cpp
+++ b/Engine/Plugins/FX/Niagara/Source/Niagara/Private/NiagaraTypeRegistry.cpp
@@ -326,9 +326,11 @@ int32 FNiagaraTypeRegistry::RegisterTypeInternal(const FNiagaraTypeDefinition& N
 				UE::TReadScopeLock ReadLock(RegisteredTypesLock);
 				if (ensure(RegisteredTypes.IsValidIndex(NewType.RegisteredTypeDefIndex)))
 				{
-					bNeedsRegistration = !RegisteredTypes[NewType.RegisteredTypeDefIndex]->IsValid();
-					checkf(bNeedsRegistration || (NewType == *RegisteredTypes[NewType.RegisteredTypeDefIndex]),
-						TEXT("bNeedsRegistration: %d, NewType: %s"), bNeedsRegistration, *NewType.GetNameText().ToString() 
+					const FNiagaraTypeDefinition& OldType = *RegisteredTypes[NewType.RegisteredTypeDefIndex]; 
+					bNeedsRegistration = !OldType.IsValid();
+					checkf(bNeedsRegistration || (NewType == OldType),e
+						TEXT("bNeedsRegistration: %d, NewType: %s, OldType: %s"), bNeedsRegistration, *NewType.GetNameText().ToString(),
+						*OldType.GetNameText().ToString() 
 					);
 				}
 			}
