[Niagara] fix typo

diff --git a/Engine/Plugins/FX/Niagara/Source/Niagara/Private/NiagaraTypeRegistry.cpp b/Engine/Plugins/FX/Niagara/Source/Niagara/Private/NiagaraTypeRegistry.cpp
--- a/Engine/Plugins/FX/Niagara/Source/Niagara/Private/NiagaraTypeRegistry.cpp
+++ b/Engine/Plugins/FX/Niagara/Source/Niagara/Private/NiagaraTypeRegistry.cpp
@@ -328,7 +328,7 @@ int32 FNiagaraTypeRegistry::RegisterTypeInternal(const FNiagaraTypeDefinition& N
 				{
 					const FNiagaraTypeDefinition& OldType = *RegisteredTypes[NewType.RegisteredTypeDefIndex]; 
 					bNeedsRegistration = !OldType.IsValid();
-					checkf(bNeedsRegistration || (NewType == OldType),e
+					checkf(bNeedsRegistration || (NewType == OldType),
 						TEXT("bNeedsRegistration: %d, NewType: %s, OldType: %s"), bNeedsRegistration, *NewType.GetNameText().ToString(),
 						*OldType.GetNameText().ToString() 
 					);
