- Make sure RendererEnabledBinding is setup

#tests ET / test project / firing range

#Synced-CL 52062699

diff --git a/Engine/Plugins/FX/Niagara/Source/Niagara/Private/NiagaraCommon.cpp b/Engine/Plugins/FX/Niagara/Source/Niagara/Private/NiagaraCommon.cpp
--- a/Engine/Plugins/FX/Niagara/Source/Niagara/Private/NiagaraCommon.cpp
+++ b/Engine/Plugins/FX/Niagara/Source/Niagara/Private/NiagaraCommon.cpp
@@ -688,6 +688,7 @@ void FNiagaraVariableAttributeBinding::Setup(const FNiagaraVariableBase& InRootV
 	ParamMapVariable.SetType(InRootVar.GetType());
 #if WITH_EDITORONLY_DATA
 	DefaultValue = InDefaultValue;
+	BoundVariable = FNiagaraVariable();
 #endif
 	SetValue(InRootVar.GetName(), FVersionedNiagaraEmitterBase(), InSourceMode);
 }
diff --git a/Engine/Plugins/FX/Niagara/Source/Niagara/Private/NiagaraComponentRendererProperties.cpp b/Engine/Plugins/FX/Niagara/Source/Niagara/Private/NiagaraComponentRendererProperties.cpp
--- a/Engine/Plugins/FX/Niagara/Source/Niagara/Private/NiagaraComponentRendererProperties.cpp
+++ b/Engine/Plugins/FX/Niagara/Source/Niagara/Private/NiagaraComponentRendererProperties.cpp
@@ -472,20 +472,22 @@ void UNiagaraComponentRendererProperties::PostDuplicate(bool bDuplicateForPIE)
 void UNiagaraComponentRendererProperties::InitCDOPropertiesAfterModuleStartup()
 {
 	UNiagaraComponentRendererProperties* CDO = CastChecked<UNiagaraComponentRendererProperties>(UNiagaraComponentRendererProperties::StaticClass()->GetDefaultObject());
+	CDO->InitRendererBaseBindings();
 	CDO->EnabledBinding = FNiagaraConstants::GetAttributeDefaultBinding(SYS_PARAM_PARTICLES_COMPONENTS_ENABLED);
 	CDO->RendererVisibilityTagBinding = FNiagaraConstants::GetAttributeDefaultBinding(SYS_PARAM_PARTICLES_VISIBILITY_TAG);
 
 	for (TWeakObjectPtr<UNiagaraComponentRendererProperties>& WeakComponentRendererProperties : ComponentRendererPropertiesToDeferredInit)
 	{
-		if (WeakComponentRendererProperties.Get())
+		if (UNiagaraComponentRendererProperties* RendererProperties = WeakComponentRendererProperties.Get())
 		{
-			if (!WeakComponentRendererProperties->EnabledBinding.IsValid())
+			RendererProperties->InitRendererBaseBindings();
+			if (!RendererProperties->EnabledBinding.IsValid())
 			{
-				WeakComponentRendererProperties->EnabledBinding = FNiagaraConstants::GetAttributeDefaultBinding(SYS_PARAM_PARTICLES_COMPONENTS_ENABLED);
+				RendererProperties->EnabledBinding = FNiagaraConstants::GetAttributeDefaultBinding(SYS_PARAM_PARTICLES_COMPONENTS_ENABLED);
 			}
-			if (!WeakComponentRendererProperties->RendererVisibilityTagBinding.IsValid())
+			if (!RendererProperties->RendererVisibilityTagBinding.IsValid())
 			{
-				WeakComponentRendererProperties->RendererVisibilityTagBinding = FNiagaraConstants::GetAttributeDefaultBinding(SYS_PARAM_PARTICLES_VISIBILITY_TAG);
+				RendererProperties->RendererVisibilityTagBinding = FNiagaraConstants::GetAttributeDefaultBinding(SYS_PARAM_PARTICLES_VISIBILITY_TAG);
 			}
 		}
 	}
diff --git a/Engine/Plugins/FX/Niagara/Source/Niagara/Private/NiagaraDecalRendererProperties.cpp b/Engine/Plugins/FX/Niagara/Source/Niagara/Private/NiagaraDecalRendererProperties.cpp
--- a/Engine/Plugins/FX/Niagara/Source/Niagara/Private/NiagaraDecalRendererProperties.cpp
+++ b/Engine/Plugins/FX/Niagara/Source/Niagara/Private/NiagaraDecalRendererProperties.cpp
@@ -39,6 +39,7 @@ namespace NiagaraDecalRendererPropertiesLocal
 		{
 			return;
 		}
+		Props->InitRendererBaseBindings();
 		Props->PositionBinding = FNiagaraConstants::GetAttributeDefaultBinding(SYS_PARAM_PARTICLES_POSITION);
 		Props->DecalOrientationBinding.Setup(UNiagaraDecalRendererProperties::GetDecalOrientationVariable(), UNiagaraDecalRendererProperties::GetDecalOrientationVariable());
 		Props->DecalSizeBinding.Setup(UNiagaraDecalRendererProperties::GetDecalSizeVariable(), UNiagaraDecalRendererProperties::GetDecalSizeVariable());
diff --git a/Engine/Plugins/FX/Niagara/Source/Niagara/Private/NiagaraLightRendererProperties.cpp b/Engine/Plugins/FX/Niagara/Source/Niagara/Private/NiagaraLightRendererProperties.cpp
--- a/Engine/Plugins/FX/Niagara/Source/Niagara/Private/NiagaraLightRendererProperties.cpp
+++ b/Engine/Plugins/FX/Niagara/Source/Niagara/Private/NiagaraLightRendererProperties.cpp
@@ -114,8 +114,8 @@ void UNiagaraLightRendererProperties::PostEditChangeProperty(struct FPropertyCha
 /** The bindings depend on variables that are created during the NiagaraModule startup. However, the CDO's are build prior to this being initialized, so we defer setting these values until later.*/
 void UNiagaraLightRendererProperties::InitCDOPropertiesAfterModuleStartup()
 {
-
 	UNiagaraLightRendererProperties* CDO = CastChecked<UNiagaraLightRendererProperties>(UNiagaraLightRendererProperties::StaticClass()->GetDefaultObject());
+	CDO->InitRendererBaseBindings();
 	CDO->PositionBinding = FNiagaraConstants::GetAttributeDefaultBinding(SYS_PARAM_PARTICLES_POSITION);
 	CDO->ColorBinding = FNiagaraConstants::GetAttributeDefaultBinding(SYS_PARAM_PARTICLES_COLOR);
 	CDO->RadiusBinding = FNiagaraConstants::GetAttributeDefaultBinding(SYS_PARAM_PARTICLES_LIGHT_RADIUS);
@@ -132,19 +132,20 @@ void UNiagaraLightRendererProperties::InitCDOPropertiesAfterModuleStartup()
 
 	for (TWeakObjectPtr<UNiagaraLightRendererProperties>& WeakLightRendererProperties : LightRendererPropertiesToDeferredInit)
 	{
-		if (WeakLightRendererProperties.Get())
+		if (UNiagaraLightRendererProperties* RendererProperties = WeakLightRendererProperties.Get())
 		{
-			if (!WeakLightRendererProperties->PositionBinding.IsValid())
+			if (!RendererProperties->PositionBinding.IsValid())
 			{
-				WeakLightRendererProperties->PositionBinding = FNiagaraConstants::GetAttributeDefaultBinding(SYS_PARAM_PARTICLES_POSITION);
-				WeakLightRendererProperties->ColorBinding = FNiagaraConstants::GetAttributeDefaultBinding(SYS_PARAM_PARTICLES_COLOR);
-				WeakLightRendererProperties->RadiusBinding = FNiagaraConstants::GetAttributeDefaultBinding(SYS_PARAM_PARTICLES_LIGHT_RADIUS);
-				WeakLightRendererProperties->LightExponentBinding = FNiagaraConstants::GetAttributeDefaultBinding(SYS_PARAM_PARTICLES_LIGHT_EXPONENT);
-				WeakLightRendererProperties->LightRenderingEnabledBinding = FNiagaraConstants::GetAttributeDefaultBinding(SYS_PARAM_PARTICLES_LIGHT_ENABLED);
-				WeakLightRendererProperties->VolumetricScatteringBinding = FNiagaraConstants::GetAttributeDefaultBinding(SYS_PARAM_PARTICLES_LIGHT_VOLUMETRIC_SCATTERING);
-				WeakLightRendererProperties->RendererVisibilityTagBinding = FNiagaraConstants::GetAttributeDefaultBinding(SYS_PARAM_PARTICLES_VISIBILITY_TAG);
-				WeakLightRendererProperties->SpecularScaleBinding = FNiagaraConstants::GetAttributeDefaultBinding(SYS_PARAM_PARTICLES_LIGHT_SPECULAR_SCALE);
-				WeakLightRendererProperties->DiffuseScaleBinding = FNiagaraConstants::GetAttributeDefaultBinding(SYS_PARAM_PARTICLES_LIGHT_DIFFUSE_SCALE);
+				RendererProperties->InitRendererBaseBindings();
+				RendererProperties->PositionBinding = FNiagaraConstants::GetAttributeDefaultBinding(SYS_PARAM_PARTICLES_POSITION);
+				RendererProperties->ColorBinding = FNiagaraConstants::GetAttributeDefaultBinding(SYS_PARAM_PARTICLES_COLOR);
+				RendererProperties->RadiusBinding = FNiagaraConstants::GetAttributeDefaultBinding(SYS_PARAM_PARTICLES_LIGHT_RADIUS);
+				RendererProperties->LightExponentBinding = FNiagaraConstants::GetAttributeDefaultBinding(SYS_PARAM_PARTICLES_LIGHT_EXPONENT);
+				RendererProperties->LightRenderingEnabledBinding = FNiagaraConstants::GetAttributeDefaultBinding(SYS_PARAM_PARTICLES_LIGHT_ENABLED);
+				RendererProperties->VolumetricScatteringBinding = FNiagaraConstants::GetAttributeDefaultBinding(SYS_PARAM_PARTICLES_LIGHT_VOLUMETRIC_SCATTERING);
+				RendererProperties->RendererVisibilityTagBinding = FNiagaraConstants::GetAttributeDefaultBinding(SYS_PARAM_PARTICLES_VISIBILITY_TAG);
+				RendererProperties->SpecularScaleBinding = FNiagaraConstants::GetAttributeDefaultBinding(SYS_PARAM_PARTICLES_LIGHT_SPECULAR_SCALE);
+				RendererProperties->DiffuseScaleBinding = FNiagaraConstants::GetAttributeDefaultBinding(SYS_PARAM_PARTICLES_LIGHT_DIFFUSE_SCALE);
 			}
 		}
 	}
diff --git a/Engine/Plugins/FX/Niagara/Source/Niagara/Private/NiagaraMeshRendererProperties.cpp b/Engine/Plugins/FX/Niagara/Source/Niagara/Private/NiagaraMeshRendererProperties.cpp
--- a/Engine/Plugins/FX/Niagara/Source/Niagara/Private/NiagaraMeshRendererProperties.cpp
+++ b/Engine/Plugins/FX/Niagara/Source/Niagara/Private/NiagaraMeshRendererProperties.cpp
@@ -623,13 +623,15 @@ void UNiagaraMeshRendererProperties::GetResourceSizeEx(FResourceSizeEx& Cumulati
 void UNiagaraMeshRendererProperties::InitCDOPropertiesAfterModuleStartup()
 {
 	UNiagaraMeshRendererProperties* CDO = CastChecked<UNiagaraMeshRendererProperties>(UNiagaraMeshRendererProperties::StaticClass()->GetDefaultObject());
+	CDO->InitRendererBaseBindings();
 	CDO->InitBindings();
 
 	for (TWeakObjectPtr<UNiagaraMeshRendererProperties>& WeakMeshRendererProperties : MeshRendererPropertiesToDeferredInit)
 	{
-		if (WeakMeshRendererProperties.Get())
+		if (UNiagaraMeshRendererProperties* RendererProperties = WeakMeshRendererProperties.Get())
 		{
-			WeakMeshRendererProperties->InitBindings();
+			RendererProperties->InitRendererBaseBindings();
+			RendererProperties->InitBindings();
 		}
 	}
 }
diff --git a/Engine/Plugins/FX/Niagara/Source/Niagara/Private/NiagaraRendererProperties.cpp b/Engine/Plugins/FX/Niagara/Source/Niagara/Private/NiagaraRendererProperties.cpp
--- a/Engine/Plugins/FX/Niagara/Source/Niagara/Private/NiagaraRendererProperties.cpp
+++ b/Engine/Plugins/FX/Niagara/Source/Niagara/Private/NiagaraRendererProperties.cpp
@@ -1258,8 +1258,7 @@ void UNiagaraRendererProperties::PostInitProperties()
 	{
 		SetFlags(RF_Transactional);
 
-		FNiagaraVariableBase EnabledDefaultVariable(FNiagaraTypeDefinition::GetBoolDef(), NAME_None);
-		RendererEnabledBinding.Setup(EnabledDefaultVariable, EnabledDefaultVariable, ENiagaraRendererSourceDataMode::Emitter);
+		InitRendererBaseBindings();
 	}
 #endif
 }
@@ -1424,6 +1423,15 @@ bool UNiagaraRendererProperties::IsContainedInEmitterRendererArray() const
 }
 #endif
 
+void UNiagaraRendererProperties::InitRendererBaseBindings()
+{
+	if (!RendererEnabledBinding.IsValid())
+	{
+		FNiagaraVariableBase EnabledDefaultVariable(FNiagaraTypeDefinition::GetBoolDef(), NAME_None);
+		RendererEnabledBinding.Setup(EnabledDefaultVariable, EnabledDefaultVariable, ENiagaraRendererSourceDataMode::Emitter);
+	}
+}
+
 bool UNiagaraRendererProperties::NeedsPreciseMotionVectors() const
 {
 	if (MotionVectorSetting == ENiagaraRendererMotionVectorSetting::AutoDetect)
diff --git a/Engine/Plugins/FX/Niagara/Source/Niagara/Private/NiagaraSpriteRendererProperties.cpp b/Engine/Plugins/FX/Niagara/Source/Niagara/Private/NiagaraSpriteRendererProperties.cpp
--- a/Engine/Plugins/FX/Niagara/Source/Niagara/Private/NiagaraSpriteRendererProperties.cpp
+++ b/Engine/Plugins/FX/Niagara/Source/Niagara/Private/NiagaraSpriteRendererProperties.cpp
@@ -275,13 +275,15 @@ void UNiagaraSpriteRendererProperties::GetResourceSizeEx(FResourceSizeEx& Cumula
 void UNiagaraSpriteRendererProperties::InitCDOPropertiesAfterModuleStartup()
 {
 	UNiagaraSpriteRendererProperties* CDO = CastChecked<UNiagaraSpriteRendererProperties>(UNiagaraSpriteRendererProperties::StaticClass()->GetDefaultObject());
+	CDO->InitRendererBaseBindings();
 	CDO->InitBindings();
 
 	for (TWeakObjectPtr<UNiagaraSpriteRendererProperties>& WeakSpriteRendererProperties : SpriteRendererPropertiesToDeferredInit)
 	{
-		if (WeakSpriteRendererProperties.Get())
+		if (UNiagaraSpriteRendererProperties* RendererProperties = WeakSpriteRendererProperties.Get())
 		{
-			WeakSpriteRendererProperties->InitBindings();
+			RendererProperties->InitRendererBaseBindings();
+			RendererProperties->InitBindings();
 		}
 	}
 }
diff --git a/Engine/Plugins/FX/Niagara/Source/Niagara/Public/NiagaraRendererProperties.h b/Engine/Plugins/FX/Niagara/Source/Niagara/Public/NiagaraRendererProperties.h
--- a/Engine/Plugins/FX/Niagara/Source/Niagara/Public/NiagaraRendererProperties.h
+++ b/Engine/Plugins/FX/Niagara/Source/Niagara/Public/NiagaraRendererProperties.h
@@ -464,6 +464,8 @@ public:
 	[[nodiscard]] bool IsContainedInEmitterRendererArray() const;
 #endif
 
+	NIAGARA_API void InitRendererBaseBindings();
+
 	/** Platforms on which this renderer is enabled. */
 	UPROPERTY(EditAnywhere, Category = "Scalability", meta=(DisplayInScalabilityContext))
 	FNiagaraPlatformSet Platforms;
diff --git a/Engine/Plugins/Runtime/GeometryCache/Source/GeometryCache/Private/NiagaraGeometryCacheRendererProperties.cpp b/Engine/Plugins/Runtime/GeometryCache/Source/GeometryCache/Private/NiagaraGeometryCacheRendererProperties.cpp
--- a/Engine/Plugins/Runtime/GeometryCache/Source/GeometryCache/Private/NiagaraGeometryCacheRendererProperties.cpp
+++ b/Engine/Plugins/Runtime/GeometryCache/Source/GeometryCache/Private/NiagaraGeometryCacheRendererProperties.cpp
@@ -221,6 +221,7 @@ void UNiagaraGeometryCacheRendererProperties::InitCDOPropertiesAfterModuleStartu
 	InitDefaultAttributes();
 	
 	UNiagaraGeometryCacheRendererProperties* CDO = CastChecked<UNiagaraGeometryCacheRendererProperties>(StaticClass()->GetDefaultObject());
+	CDO->InitRendererBaseBindings();
 	CDO->PositionBinding = FNiagaraConstants::GetAttributeDefaultBinding(SYS_PARAM_PARTICLES_POSITION);
 	CDO->ScaleBinding = FNiagaraConstants::GetAttributeDefaultBinding(SYS_PARAM_PARTICLES_SCALE);
 	CDO->RendererVisibilityTagBinding = FNiagaraConstants::GetAttributeDefaultBinding(SYS_PARAM_PARTICLES_VISIBILITY_TAG);
@@ -231,35 +232,36 @@ void UNiagaraGeometryCacheRendererProperties::InitCDOPropertiesAfterModuleStartu
 
 	for (TWeakObjectPtr<UNiagaraGeometryCacheRendererProperties>& WeakGeometryCacheRendererProperties : GeometryCacheRendererPropertiesToDeferredInit)
 	{
-		if (WeakGeometryCacheRendererProperties.Get())
+		if (UNiagaraGeometryCacheRendererProperties* RedererProperties = WeakGeometryCacheRendererProperties.Get())
 		{
-			if (!WeakGeometryCacheRendererProperties->PositionBinding.IsValid())
+			RedererProperties->InitRendererBaseBindings();
+			if (!RedererProperties->PositionBinding.IsValid())
 			{
-				WeakGeometryCacheRendererProperties->PositionBinding = FNiagaraConstants::GetAttributeDefaultBinding(SYS_PARAM_PARTICLES_POSITION);
+				RedererProperties->PositionBinding = FNiagaraConstants::GetAttributeDefaultBinding(SYS_PARAM_PARTICLES_POSITION);
 			}
-			if (!WeakGeometryCacheRendererProperties->ScaleBinding.IsValid())
+			if (!RedererProperties->ScaleBinding.IsValid())
 			{
-				WeakGeometryCacheRendererProperties->ScaleBinding = FNiagaraConstants::GetAttributeDefaultBinding(SYS_PARAM_PARTICLES_SCALE);
+				RedererProperties->ScaleBinding = FNiagaraConstants::GetAttributeDefaultBinding(SYS_PARAM_PARTICLES_SCALE);
 			}
-			if (!WeakGeometryCacheRendererProperties->RendererVisibilityTagBinding.IsValid())
+			if (!RedererProperties->RendererVisibilityTagBinding.IsValid())
 			{
-				WeakGeometryCacheRendererProperties->RendererVisibilityTagBinding = FNiagaraConstants::GetAttributeDefaultBinding(SYS_PARAM_PARTICLES_VISIBILITY_TAG);
+				RedererProperties->RendererVisibilityTagBinding = FNiagaraConstants::GetAttributeDefaultBinding(SYS_PARAM_PARTICLES_VISIBILITY_TAG);
 			}
-			if (!WeakGeometryCacheRendererProperties->RotationBinding.IsValid())
+			if (!RedererProperties->RotationBinding.IsValid())
 			{
-				WeakGeometryCacheRendererProperties->RotationBinding = CreateDefaultBinding(Particles_GeoCacheRotation, FVector3f(0, 0, 0));
+				RedererProperties->RotationBinding = CreateDefaultBinding(Particles_GeoCacheRotation, FVector3f(0, 0, 0));
 			}
-			if (!WeakGeometryCacheRendererProperties->ElapsedTimeBinding.IsValid())
+			if (!RedererProperties->ElapsedTimeBinding.IsValid())
 			{
-				WeakGeometryCacheRendererProperties->ElapsedTimeBinding = CreateDefaultBinding(Particles_Age, 0.0f);
+				RedererProperties->ElapsedTimeBinding = CreateDefaultBinding(Particles_Age, 0.0f);
 			}
-			if (!WeakGeometryCacheRendererProperties->EnabledBinding.IsValid())
+			if (!RedererProperties->EnabledBinding.IsValid())
 			{
-				WeakGeometryCacheRendererProperties->EnabledBinding = CreateDefaultBinding(Particles_GeoCacheIsEnabled, false);
+				RedererProperties->EnabledBinding = CreateDefaultBinding(Particles_GeoCacheIsEnabled, false);
 			}
-			if (!WeakGeometryCacheRendererProperties->ArrayIndexBinding.IsValid())
+			if (!RedererProperties->ArrayIndexBinding.IsValid())
 			{
-				WeakGeometryCacheRendererProperties->ArrayIndexBinding = CreateDefaultBinding(Particles_ArrayIndex, 0);
+				RedererProperties->ArrayIndexBinding = CreateDefaultBinding(Particles_ArrayIndex, 0);
 			}
 		}
 	}
