[Niagara]

Mark NiagaraScript to DeclareConstructClass for NiagaraHierarchyScriptParameter. NiagaraHierarchyScriptParameter can be constructed when postloading a UScript with a UNiagaraGraph leading to ConstructedDuringPostLoad logs being emitted (with LogCore verbose) indicating this relationship needs to be declared for incremental cook correctness.

#rnx
#Synced-CL 52104572

diff --git a/Engine/Plugins/FX/Niagara/Source/Niagara/Classes/NiagaraScript.h b/Engine/Plugins/FX/Niagara/Source/Niagara/Classes/NiagaraScript.h
--- a/Engine/Plugins/FX/Niagara/Source/Niagara/Classes/NiagaraScript.h
+++ b/Engine/Plugins/FX/Niagara/Source/Niagara/Classes/NiagaraScript.h
@@ -1080,7 +1080,9 @@ public:
 	virtual void PreSave(FObjectPreSaveContext ObjectSaveContext) override;
 	virtual void Serialize(FArchive& Ar)override;
 	virtual void PostLoad() override;
-	
+#if WITH_EDITORONLY_DATA
+	static void DeclareConstructClasses(TArray<FTopLevelAssetPath>& OutConstructClasses, const UClass* SpecificSubclass);
+#endif
 
 #if WITH_EDITOR
 	virtual void PostEditChangeProperty(FPropertyChangedEvent& PropertyChangedEvent) override;
diff --git a/Engine/Plugins/FX/Niagara/Source/Niagara/Private/NiagaraScript.cpp b/Engine/Plugins/FX/Niagara/Source/Niagara/Private/NiagaraScript.cpp
--- a/Engine/Plugins/FX/Niagara/Source/Niagara/Private/NiagaraScript.cpp
+++ b/Engine/Plugins/FX/Niagara/Source/Niagara/Private/NiagaraScript.cpp
@@ -2104,6 +2104,14 @@ void UNiagaraScript::PostLoad()
 	GenerateStatIDs();
 }
 
+#if WITH_EDITORONLY_DATA
+void UNiagaraScript::DeclareConstructClasses(TArray<FTopLevelAssetPath>& OutConstructClasses, const UClass* SpecificSubclass)
+{
+	Super::DeclareConstructClasses(OutConstructClasses, SpecificSubclass);
+	OutConstructClasses.Add(FTopLevelAssetPath(TEXT("/Script/NiagaraEditor.NiagaraHierarchyScriptParameter"))); // cannot use NiagaraHierarchyScriptParameter::StaticClass() here as this type is in an editor module
+}
+#endif
+
 bool UNiagaraScript::IsReadyToRun(ENiagaraSimTarget SimTarget) const
 {
 	if (SimTarget == ENiagaraSimTarget::CPUSim)
