- Fix TSAN issue

#rnx

#Synced-CL 52392748

diff --git a/Engine/Plugins/FX/Niagara/Source/NiagaraEditor/Private/NiagaraCompilationTasks.cpp b/Engine/Plugins/FX/Niagara/Source/NiagaraEditor/Private/NiagaraCompilationTasks.cpp
--- a/Engine/Plugins/FX/Niagara/Source/NiagaraEditor/Private/NiagaraCompilationTasks.cpp
+++ b/Engine/Plugins/FX/Niagara/Source/NiagaraEditor/Private/NiagaraCompilationTasks.cpp
@@ -47,7 +47,7 @@ namespace NiagaraSystemCookStats
 {
 	extern FCookStats::FDDCResourceUsageStats UsageStats;
 
-	static int32 ScriptTranslationCount = 0;
+	static std::atomic<int32> ScriptTranslationCount(0);
 	static int32 CpuScriptCompileCount = 0;
 	static int32 CpuScriptDdcCacheHitCount = 0;
 	static int32 CpuScriptDdcCacheMissCount = 0;
@@ -57,7 +57,7 @@ namespace NiagaraSystemCookStats
 
 	static FCookStatsManager::FAutoRegisterCallback RegisterNiagaraCompilationCookStats([](FCookStatsManager::AddStatFuncRef AddStat)
 	{
-		AddStat(TEXT("Niagara"), FCookStatsManager::CreateKeyValueArray(TEXT("NiagaraScriptTranslationCount"), ScriptTranslationCount));
+		AddStat(TEXT("Niagara"), FCookStatsManager::CreateKeyValueArray(TEXT("NiagaraScriptTranslationCount"), ScriptTranslationCount.load(std::memory_order_relaxed)));
 		AddStat(TEXT("Niagara"), FCookStatsManager::CreateKeyValueArray(TEXT("NiagaraCpuScriptCompileCount"), CpuScriptCompileCount));
 		AddStat(TEXT("Niagara"), FCookStatsManager::CreateKeyValueArray(TEXT("NiagaraCpuScriptDdcCacheHitCount"), CpuScriptDdcCacheHitCount));
 		AddStat(TEXT("Niagara"), FCookStatsManager::CreateKeyValueArray(TEXT("NiagaraCpuScriptDdcCacheMissCount"), CpuScriptDdcCacheMissCount));
@@ -919,7 +919,7 @@ void FNiagaraSystemCompilationTask::FCompileTaskInfo::Translate(FNiagaraSystemCo
 {
 	TRACE_CPUPROFILER_EVENT_SCOPE(NiagaraAsyncTask_Translate);
 
-	COOK_STAT(NiagaraSystemCookStats::ScriptTranslationCount++);
+	COOK_STAT(NiagaraSystemCookStats::ScriptTranslationCount.fetch_add(1, std::memory_order_relaxed));
 
 	const FNiagaraPrecompileData* PrecompileData = nullptr;
 	const FNiagaraCompilationCopyData* CompilationCopyData = nullptr;
