- Ignore time since rendered culling if we don't have a scene proxy as it will never be updated
- This fixes a bug where FX could be culled due to PSO precaching being in flight and they could never reactivate

#Synced-CL 52445843

diff --git a/Engine/Plugins/FX/Niagara/Source/Niagara/Private/NiagaraWorldManager.cpp b/Engine/Plugins/FX/Niagara/Source/Niagara/Private/NiagaraWorldManager.cpp
--- a/Engine/Plugins/FX/Niagara/Source/Niagara/Private/NiagaraWorldManager.cpp
+++ b/Engine/Plugins/FX/Niagara/Source/Niagara/Private/NiagaraWorldManager.cpp
@@ -2133,12 +2133,22 @@ void FNiagaraWorldManager::CalculateScalabilityState(UNiagaraSystem* System, con
 	OutState.bCulled = false;
 
 	DistanceCull(EffectType, ScalabilitySettings, Component, OutState);
-	
+
 	if (GEnableNiagaraVisCulling)
 	{
 		FBoxSphereBounds Bounds = Component->CalcBounds(Component->GetComponentToWorld());		
 
-		const float TimeSinceRendered = float(FMath::Max(0.0f, GetWorld()->LastRenderTime - Component->GetLastRenderTime() - World->GetDeltaSeconds() - RenderThreadWorldTimeLag - GLastRenderTimeSafetyBias));
+		// If we don't have a scene proxy we can not cull against last render time
+		// We also make sure to update the last render time so when the proxy becomes visible we don't get a potential invalid frame before we get data back from the render thread
+		float TimeSinceRendered = 0.0f;
+		if (Component->GetSceneProxy())
+		{
+			TimeSinceRendered = float(FMath::Max(0.0f, GetWorld()->LastRenderTime - Component->GetLastRenderTime() - World->GetDeltaSeconds() - RenderThreadWorldTimeLag - GLastRenderTimeSafetyBias));
+		}
+		else
+		{
+			Component->SetLastRenderTime(GetWorld()->LastRenderTime);
+		}
 
 		if(bIsPreCull)
 		{
