Add NiagaraDataInterfaceParticleNeighborQuery - New method for computing particle neighbors that has an unbound number of particles per cell.  Memory is much more efficient than NeighborGrid3D, and it has an open architecture that will allow for future optimizations and improvements.

#Synced-CL 52108856

diff --git a/Engine/Plugins/FX/Niagara/Shaders/Private/NiagaraDataInterfaceParticleNeighborQuery.ush b/Engine/Plugins/FX/Niagara/Shaders/Private/NiagaraDataInterfaceParticleNeighborQuery.ush
new file mode 100644
--- /dev/null
+++ b/Engine/Plugins/FX/Niagara/Shaders/Private/NiagaraDataInterfaceParticleNeighborQuery.ush
@@ -0,0 +1,194 @@
+// Copyright Epic Games, Inc. All Rights Reserved.
+
+/*=============================================================================
+NiagaraDataInterfaceParticleNeighborQuery.ush
+=============================================================================*/
+
+/* -----------------------------------------------------------------
+ * Parameter declarations (bound by FShaderParameters / SetShaderParameters)
+ * -----------------------------------------------------------------
+ */
+
+int				{ParameterName}_MaxCellsPerParticleValue;
+
+Buffer<int>		{ParameterName}_CellCount;
+Buffer<int>		{ParameterName}_CellOffset;
+Buffer<int>		{ParameterName}_ParticleList;
+RWBuffer<int>	{ParameterName}_OutputCellId;
+
+#ifndef NEIGHBORGRID_INVALID_CELL
+#define NEIGHBORGRID_INVALID_CELL 0xFFFFFFFF
+#endif
+
+/* -----------------------------------------------------------------
+ * Internal helpers (not exposed to users)
+ * -----------------------------------------------------------------
+ */
+
+int IndexToLinear_{ParameterName}(int3 CellCoord)
+{
+	return CellCoord.x
+		+ CellCoord.y * {ParameterName}_NumCells.x
+		+ CellCoord.z * {ParameterName}_NumCells.x * {ParameterName}_NumCells.y;
+}
+
+/* -----------------------------------------------------------------
+ * Grid property queries
+ * -----------------------------------------------------------------
+ */
+void MaxCellsPerParticle_{ParameterName}(out int OutMaxCellsPerParticle)
+{
+	OutMaxCellsPerParticle = {ParameterName}_MaxCellsPerParticleValue;
+}
+
+/* -----------------------------------------------------------------
+ * Write functions (output stage — pass 1 cell assignment)
+ * -----------------------------------------------------------------
+ */
+
+void AddParticle_{ParameterName}(bool Enabled, float3 IndexPos, int ParticleIndex, out bool Success)
+{
+	Success = false;
+	if (!Enabled) return;
+
+	int3 CellCoord = int3(round(IndexPos));
+
+	if (all(CellCoord >= 0) && all(CellCoord < {ParameterName}_NumCells))
+	{
+		int LinearIndex = IndexToLinear_{ParameterName}(CellCoord);
+
+		// Single-cell: 1 slot per particle, write to first slot
+		int SlotBase = ParticleIndex * {ParameterName}_MaxCellsPerParticleValue;
+		{ParameterName}_OutputCellId[SlotBase] = LinearIndex;
+
+		Success = true;
+	}
+}
+
+// Exec-pin wrapper — Niagara appends _UEImpureCall to the generated function name for bRequiresExecPin functions
+void AddParticle_{ParameterName}_UEImpureCall(bool In_Enabled, float3 In_IndexPos, int In_ParticleIndex, out bool Out_Success)
+{
+	AddParticle_{ParameterName}(In_Enabled, In_IndexPos, In_ParticleIndex, Out_Success);
+}
+
+void AddParticleWithRadius_{ParameterName}(bool Enabled, float3 IndexPos, float Radius, int ParticleIndex, out bool Success)
+{
+	Success = false;
+	if (!Enabled) return;
+
+	int3 MinCell = clamp(int3(round(IndexPos - Radius)), 0, {ParameterName}_NumCells - 1);
+	int3 MaxCell = clamp(int3(round(IndexPos + Radius)), 0, {ParameterName}_NumCells - 1);
+
+	int SlotBase = ParticleIndex * {ParameterName}_MaxCellsPerParticleValue;
+	int Slot = 0;
+
+	for (int z = MinCell.z; z <= MaxCell.z; z++) {
+	for (int y = MinCell.y; y <= MaxCell.y; y++) {
+	for (int x = MinCell.x; x <= MaxCell.x; x++) {
+		if (Slot < {ParameterName}_MaxCellsPerParticleValue)
+		{
+			int LinearIndex = IndexToLinear_{ParameterName}(int3(x, y, z));
+			{ParameterName}_OutputCellId[SlotBase + Slot] = LinearIndex;
+			Slot++;
+			Success = true;
+		}
+	}}}
+}
+
+// Exec-pin wrapper
+void AddParticleWithRadius_{ParameterName}_UEImpureCall(bool In_Enabled, float3 In_IndexPos, float In_Radius, int In_ParticleIndex, out bool Out_Success)
+{
+	AddParticleWithRadius_{ParameterName}(In_Enabled, In_IndexPos, In_Radius, In_ParticleIndex, Out_Success);
+}
+
+void AddParticleToOverlappedCells_{ParameterName}(bool Enabled, float3 IndexPos, int ParticleIndex, out bool Success)
+{
+	Success = false;
+	if (!Enabled) return;
+
+	int3 BaseCell = int3(floor(IndexPos));
+	int SlotBase = ParticleIndex * {ParameterName}_MaxCellsPerParticleValue;
+	int Slot = 0;
+
+	for (int z = 0; z <= 1; z++) {
+	for (int y = 0; y <= 1; y++) {
+	for (int x = 0; x <= 1; x++) {
+		int3 CellCoord = BaseCell + int3(x, y, z);
+		if (Slot < {ParameterName}_MaxCellsPerParticleValue &&
+			all(CellCoord >= 0) && all(CellCoord < {ParameterName}_NumCells))
+		{
+			{ParameterName}_OutputCellId[SlotBase + Slot] = IndexToLinear_{ParameterName}(CellCoord);
+			Slot++;
+			Success = true;
+		}
+	}}}
+}
+
+// Exec-pin wrapper
+void AddParticleToOverlappedCells_{ParameterName}_UEImpureCall(bool In_Enabled, float3 In_IndexPos, int In_ParticleIndex, out bool Out_Success)
+{
+	AddParticleToOverlappedCells_{ParameterName}(In_Enabled, In_IndexPos, In_ParticleIndex, Out_Success);
+}
+
+void AddParticleToNeighborCells_{ParameterName}(bool Enabled, float3 IndexPos, int ParticleIndex, out bool Success)
+{
+	Success = false;
+	if (!Enabled) return;
+
+	int3 BaseCell = int3(round(IndexPos));
+	int SlotBase = ParticleIndex * {ParameterName}_MaxCellsPerParticleValue;
+	int Slot = 0;
+
+	for (int z = -1; z <= 1; z++) {
+	for (int y = -1; y <= 1; y++) {
+	for (int x = -1; x <= 1; x++) {
+		int3 CellCoord = BaseCell + int3(x, y, z);
+		if (Slot < {ParameterName}_MaxCellsPerParticleValue &&
+			all(CellCoord >= 0) && all(CellCoord < {ParameterName}_NumCells))
+		{
+			{ParameterName}_OutputCellId[SlotBase + Slot] = IndexToLinear_{ParameterName}(CellCoord);
+			Slot++;
+			Success = true;
+		}
+	}}}
+}
+
+// Exec-pin wrapper
+void AddParticleToNeighborCells_{ParameterName}_UEImpureCall(bool In_Enabled, float3 In_IndexPos, int In_ParticleIndex, out bool Out_Success)
+{
+	AddParticleToNeighborCells_{ParameterName}(In_Enabled, In_IndexPos, In_ParticleIndex, Out_Success);
+}
+
+/* -----------------------------------------------------------------
+ * Read functions (read stage — after counting sort passes 2-4)
+ * -----------------------------------------------------------------
+ */
+
+void GetParticleNeighborCount_{ParameterName}(int IndexX, int IndexY, int IndexZ, out int OutCount)
+{
+	int3 CellCoord = int3(IndexX, IndexY, IndexZ);
+	if (any(CellCoord < 0) || any(CellCoord >= {ParameterName}_NumCells))
+	{
+		OutCount = 0;
+		return;
+	}
+	int LinearIndex = IndexToLinear_{ParameterName}(CellCoord);
+	OutCount = {ParameterName}_CellCount[LinearIndex];
+}
+
+void GetParticleNeighbor_{ParameterName}(int IndexX, int IndexY, int IndexZ, int NeighborIdx, out int OutParticleIdx)
+{
+	int3 CellCoord = int3(IndexX, IndexY, IndexZ);
+	bool bValidCell = all(CellCoord >= 0) && all(CellCoord < {ParameterName}_NumCells);
+	int LinearIndex = bValidCell ? IndexToLinear_{ParameterName}(CellCoord) : 0;
+	int Count = bValidCell ? {ParameterName}_CellCount[LinearIndex] : 0;
+	if (!bValidCell || NeighborIdx >= Count || NeighborIdx < 0)
+	{
+		OutParticleIdx = -1;
+		return;
+	}
+	// CellOffset was consumed by the scatter pass (atomicAdd), so the original prefix sum value
+	// is CellOffset[i] - CellCount[i].
+	int StartOffset = {ParameterName}_CellOffset[LinearIndex] - Count;
+	OutParticleIdx = {ParameterName}_ParticleList[StartOffset + NeighborIdx];
+}
diff --git a/Engine/Plugins/FX/Niagara/Shaders/Private/NiagaraParticleNeighborQuerySort.usf b/Engine/Plugins/FX/Niagara/Shaders/Private/NiagaraParticleNeighborQuerySort.usf
new file mode 100644
--- /dev/null
+++ b/Engine/Plugins/FX/Niagara/Shaders/Private/NiagaraParticleNeighborQuerySort.usf
@@ -0,0 +1,446 @@
+// Copyright Epic Games, Inc. All Rights Reserved.
+
+#include "/Engine/Private/Common.ush"
+
+// Counting sort passes for NiagaraDataInterfaceParticleNeighborQuery.
+// These run between the output stage (where user modules write cell assignments)
+// and the read stage (where user modules query neighbors).
+//
+// Buffers:
+//   CellIdBuffer     — [NumSlots] cell assignment per slot (written by user in output stage)
+//   CellCountBuffer  — [NumCells] particles per cell (histogram output)
+//   CellOffsetBuffer — [NumCells] exclusive prefix sum write offset per cell
+//   ParticleListBuffer — [NumSlots] final sorted particle indices
+
+#define NEIGHBORGRID_INVALID_CELL 0xFFFFFFFF
+
+//////////////////////////////////////////////////////////////////////////
+// Pass 2: Histogram — count particles per cell
+//////////////////////////////////////////////////////////////////////////
+#if PARTICLE_NEIGHBOR_QUERY_HISTOGRAM_CS
+
+// Define USE_BITONIC_HISTOGRAM to enable bitonic sort + coalesced atomics.
+// When disabled, uses direct per-slot atomics (faster for sparse grids).
+// #define USE_BITONIC_HISTOGRAM
+
+Buffer<int>		CellIdBuffer;
+RWBuffer<int>	CellCountBuffer;
+uint			NumSlots;
+uint			GroupsX;
+
+#ifdef USE_BITONIC_HISTOGRAM
+
+groupshared int SortedCellIds[THREADGROUP_SIZE];
+
+// Serial recursive bitonic sort (reference implementation, not used on GPU):
+/*
+void BitonicSortStep(uint ArrayStart, uint Length, bool IsAscending)
+{
+	if (Length > 1)
+	{
+		int k = Length / 2;
+		BitonicSortStep(ArrayStart, k, true);
+		BitonicSortStep(ArrayStart + k, k, false);
+		BitonicMerge(ArrayStart, Length, IsAscending);
+	}
+}
+
+void BitonicMerge(uint ArrayStart, uint Length, bool IsAscending)
+{
+	if (Length > 1)
+	{
+		int k = Length / 2;
+		for (int i = ArrayStart; i < ArrayStart + k; ++i)
+		{
+			if (IsAscending == (SortedCellIds[i] > SortedCellIds[i + k]))
+			{
+				int Temp = SortedCellIds[i];
+				SortedCellIds[i] = SortedCellIds[i + k];
+				SortedCellIds[i + k] = Temp;
+			}
+		}
+		BitonicMerge(ArrayStart, k, IsAscending);
+		BitonicMerge(ArrayStart + k, k, IsAscending);
+	}
+}
+*/
+
+// Parallel bitonic sort of SortedCellIds in shared memory.
+// All threads in the group must participate (even if they hold INVALID_CELL).
+//
+// Bitonic sort works by building increasingly large "bitonic sequences"
+// (sequences that go up then down, or down then up) and merging them.
+// The outer loop doubles the block size each iteration; the inner loop
+// performs the merge steps within each block.
+//
+// For THREADGROUP_SIZE=64: 21 barrier steps. For 128: 28 barrier steps.
+void BitonicSortShared(uint LocalId)
+{
+	// Outer loop: k is the size of the bitonic block being formed.
+	// k=2: sort pairs. k=4: merge pairs into sorted quads. k=8: merge into octets. etc.
+	// This corresponds to BitonicSortStep in the serial version — each doubling of k
+	// is one level of the recursion.
+	for (uint k = 2; k <= THREADGROUP_SIZE; k <<= 1)
+	{
+		// Inner loop: j is the distance between elements being compared.
+		// Starts at k/2 and halves each step — this is the merge phase.
+		// This corresponds to BitonicMerge in the serial version — each halving of j
+		// is one level of the merge recursion.
+		for (uint j = k >> 1; j > 0; j >>= 1)
+		{
+			// XOR with j flips the j-th bit of LocalId, giving us the partner
+			// to compare with. Partners are always j positions apart.
+			// Example: j=4, LocalId=5 (binary 101) → Partner=1 (binary 001)
+			uint Partner = LocalId ^ j;
+
+			// Only the lower-indexed thread of each pair does the compare-and-swap.
+			// This avoids both threads swapping the same pair (which would undo the swap).
+			if (Partner > LocalId)
+			{
+				// The k-th bit of LocalId determines sort direction for this block.
+				// Alternating blocks sort in opposite directions to form the bitonic
+				// sequences that the next level of k will merge.
+				// Bit clear → ascending, bit set → descending.
+				bool Ascending = ((LocalId & k) == 0);
+
+				// Compare as unsigned so INVALID_CELL (0xFFFFFFFF) sorts to the end.
+				// Swap if the pair is out of order for the desired direction:
+				//   Ascending + [LocalId] > [Partner] → swap (want smaller first)
+				//   Descending + [LocalId] <= [Partner] → swap (want larger first)
+				if (Ascending == ((uint)SortedCellIds[LocalId] > (uint)SortedCellIds[Partner]))
+				{
+					int Temp = SortedCellIds[LocalId];
+					SortedCellIds[LocalId] = SortedCellIds[Partner];
+					SortedCellIds[Partner] = Temp;
+				}
+			}
+
+			// All threads must complete this step before the next comparison round,
+			// since the next step reads values that this step may have written.
+			GroupMemoryBarrierWithGroupSync();
+		}
+	}
+}
+
+#endif // USE_BITONIC_HISTOGRAM
+
+[numthreads(THREADGROUP_SIZE, 1, 1)]
+void HistogramCS(uint3 GroupId : SV_GroupID, uint3 GroupThreadId : SV_GroupThreadID)
+{
+	uint GroupIndex = GroupId.x + GroupId.y * GroupsX;
+	uint LocalId = GroupThreadId.x;
+	uint SlotIndex = GroupIndex * THREADGROUP_SIZE + LocalId;
+
+#ifdef USE_BITONIC_HISTOGRAM
+	// Bitonic sort path: sort cell IDs in shared memory, then one coalesced
+	// atomic per unique cell per group. Wins when cells are dense (many
+	// particles per cell), loses when sparse due to sort overhead.
+	SortedCellIds[LocalId] = (SlotIndex < NumSlots) ? CellIdBuffer[SlotIndex] : (int)NEIGHBORGRID_INVALID_CELL;
+	GroupMemoryBarrierWithGroupSync();
+
+	BitonicSortShared(LocalId);
+
+	int MyCellId = SortedCellIds[LocalId];
+
+	bool IsRunStart = (LocalId == 0) || (SortedCellIds[LocalId - 1] != MyCellId);
+	if (IsRunStart && MyCellId != (int)NEIGHBORGRID_INVALID_CELL)
+	{
+		uint RunCount = 1;
+		for (uint j = LocalId + 1; j < THREADGROUP_SIZE && SortedCellIds[j] == MyCellId; j++)
+		{
+			RunCount++;
+		}
+		InterlockedAdd(CellCountBuffer[MyCellId], RunCount);
+	}
+#else
+	// Direct atomic path: one InterlockedAdd per valid slot.
+	// Simpler, faster when contention is low (sparse grids, many unique cells).
+	if (SlotIndex >= NumSlots)
+	{
+		return;
+	}
+
+	int CellId = CellIdBuffer[SlotIndex];
+	if (CellId != (int)NEIGHBORGRID_INVALID_CELL)
+	{
+		InterlockedAdd(CellCountBuffer[CellId], 1);
+	}
+#endif
+}
+
+#endif
+
+//////////////////////////////////////////////////////////////////////////
+// Pass 3: Parallel exclusive prefix sum (recursive reduce-then-scan)
+//
+// Two shader sub-passes per recursion level:
+//   3a. LocalScan  — each group Blelloch-scans ELEMENTS_PER_GROUP elements,
+//                    writes tile total to BlockSumsBuffer
+//   3b. Propagate  — each group adds its scanned block prefix to local results
+//
+// Block sums are scanned by recursively applying the same algorithm.
+// Each thread processes ELEMENTS_PER_THREAD elements.
+// Tile size = THREADGROUP_SIZE * ELEMENTS_PER_THREAD = ELEMENTS_PER_GROUP.
+//////////////////////////////////////////////////////////////////////////
+
+#if PARTICLE_NEIGHBOR_QUERY_PREFIX_SUM_LOCAL_SCAN_CS
+
+Buffer<int>		InputBuffer;
+RWBuffer<int>	OutputBuffer;
+RWBuffer<int>	BlockSumsBuffer;
+uint			NumElements;
+uint			GroupsX;
+
+groupshared int SharedPartials[THREADGROUP_SIZE];
+
+[numthreads(THREADGROUP_SIZE, 1, 1)]
+void PrefixSumLocalScanCS(uint3 GroupId : SV_GroupID, uint3 GroupThreadId : SV_GroupThreadID)
+{
+	uint GroupIndex = GroupId.x + GroupId.y * GroupsX;
+	uint TileBase = GroupIndex * ELEMENTS_PER_GROUP;
+	uint LocalId = GroupThreadId.x;
+
+	// Phase 1: Each thread loads ELEMENTS_PER_THREAD elements and computes a
+	// sequential exclusive prefix sum in registers. Store the per-thread total
+	// into shared memory for the tree scan.
+	int LocalPrefixes[ELEMENTS_PER_THREAD];
+	int ThreadSum = 0;
+	for (uint e = 0; e < ELEMENTS_PER_THREAD; e++)
+	{
+		uint GlobalIdx = TileBase + LocalId * ELEMENTS_PER_THREAD + e;
+		int Val = (GlobalIdx < NumElements) ? InputBuffer[GlobalIdx] : 0;
+		LocalPrefixes[e] = ThreadSum;
+		ThreadSum += Val;
+	}
+	SharedPartials[LocalId] = ThreadSum;
+	GroupMemoryBarrierWithGroupSync();
+
+	// Phase 2: Blelloch exclusive scan on THREADGROUP_SIZE per-thread totals.
+	// This correctly uses N threads for N elements (1:1 ratio).
+	uint N = THREADGROUP_SIZE;
+
+	// Up-sweep (reduce)
+	for (uint Stride = 1; Stride < N; Stride <<= 1)
+	{
+		uint Idx = (LocalId + 1) * (Stride << 1) - 1;
+		if (Idx < N)
+		{
+			SharedPartials[Idx] += SharedPartials[Idx - Stride];
+		}
+		GroupMemoryBarrierWithGroupSync();
+	}
+
+	// Store tile total and clear last element for down-sweep
+	if (LocalId == 0)
+	{
+		BlockSumsBuffer[GroupIndex] = SharedPartials[N - 1];
+		SharedPartials[N - 1] = 0;
+	}
+	GroupMemoryBarrierWithGroupSync();
+
+	// Down-sweep (distribute)
+	for (uint Stride2 = N >> 1; Stride2 >= 1; Stride2 >>= 1)
+	{
+		uint Idx = (LocalId + 1) * (Stride2 << 1) - 1;
+		if (Idx < N)
+		{
+			int Temp = SharedPartials[Idx - Stride2];
+			SharedPartials[Idx - Stride2] = SharedPartials[Idx];
+			SharedPartials[Idx] += Temp;
+		}
+		GroupMemoryBarrierWithGroupSync();
+	}
+
+	// Phase 3: Combine the per-thread prefix (from the tree scan) with each
+	// thread's local sequential prefix sums and write final results.
+	int ThreadPrefix = SharedPartials[LocalId];
+	for (uint e2 = 0; e2 < ELEMENTS_PER_THREAD; e2++)
+	{
+		uint GlobalIdx = TileBase + LocalId * ELEMENTS_PER_THREAD + e2;
+		if (GlobalIdx < NumElements)
+		{
+			OutputBuffer[GlobalIdx] = LocalPrefixes[e2] + ThreadPrefix;
+		}
+	}
+}
+
+#endif
+
+//////////////////////////////////////////////////////////////////////////
+#if PARTICLE_NEIGHBOR_QUERY_PREFIX_SUM_PROPAGATE_CS
+
+Buffer<int>		BlockSumsBuffer;
+RWBuffer<int>	OutputBuffer;
+uint			NumElements;
+uint			GroupsX;
+
+[numthreads(THREADGROUP_SIZE, 1, 1)]
+void PrefixSumPropagateCS(uint3 GroupId : SV_GroupID, uint3 GroupThreadId : SV_GroupThreadID)
+{
+	uint GroupIndex = GroupId.x + GroupId.y * GroupsX;
+	// Group 0 has no prefix to add
+	if (GroupIndex == 0)
+	{
+		return;
+	}
+
+	int BlockPrefix = BlockSumsBuffer[GroupIndex];
+	uint TileBase = GroupIndex * ELEMENTS_PER_GROUP;
+	uint LocalId = GroupThreadId.x;
+
+	for (uint e = 0; e < ELEMENTS_PER_THREAD; e++)
+	{
+		uint GlobalIdx = TileBase + LocalId * ELEMENTS_PER_THREAD + e;
+		if (GlobalIdx < NumElements)
+		{
+			OutputBuffer[GlobalIdx] += BlockPrefix;
+		}
+	}
+}
+
+#endif
+
+//////////////////////////////////////////////////////////////////////////
+// Pass 4: Scatter — write sorted particle indices
+//////////////////////////////////////////////////////////////////////////
+#if PARTICLE_NEIGHBOR_QUERY_SCATTER_CS
+
+// Define USE_BITONIC_SCATTER to enable bitonic sort + coalesced atomics.
+// When disabled, uses direct per-slot atomics (simpler, faster for sparse grids).
+#define USE_BITONIC_SCATTER
+
+Buffer<int>		CellIdBuffer;
+RWBuffer<int>	CellOffsetBuffer;
+RWBuffer<int>	ParticleListBuffer;
+uint			NumSlots;
+uint			MaxCellsPerParticle;
+uint			GroupsX;
+
+#ifdef USE_BITONIC_SCATTER
+
+// Parallel arrays in shared memory: cell IDs and original slot indices.
+// After bitonic sort (keyed on cell ID, carrying slot index), each thread
+// can derive its local offset within the run and write to a contiguous block
+// claimed by one coalesced atomic per unique cell per group.
+groupshared int ScatterSortedCellIds[THREADGROUP_SIZE];
+groupshared uint ScatterSortedSlotIndices[THREADGROUP_SIZE];
+groupshared uint RunWriteBase[THREADGROUP_SIZE];
+
+// Parallel bitonic sort of ScatterSortedCellIds in shared memory, carrying
+// ScatterSortedSlotIndices as satellite data (swap when keys swap).
+// All threads in the group must participate.
+void BitonicSortScatterShared(uint LocalId)
+{
+	for (uint k = 2; k <= THREADGROUP_SIZE; k <<= 1)
+	{
+		for (uint j = k >> 1; j > 0; j >>= 1)
+		{
+			uint Partner = LocalId ^ j;
+			if (Partner > LocalId)
+			{
+				bool Ascending = ((LocalId & k) == 0);
+				if (Ascending == ((uint)ScatterSortedCellIds[LocalId] > (uint)ScatterSortedCellIds[Partner]))
+				{
+					int TempCell = ScatterSortedCellIds[LocalId];
+					ScatterSortedCellIds[LocalId] = ScatterSortedCellIds[Partner];
+					ScatterSortedCellIds[Partner] = TempCell;
+
+					uint TempSlot = ScatterSortedSlotIndices[LocalId];
+					ScatterSortedSlotIndices[LocalId] = ScatterSortedSlotIndices[Partner];
+					ScatterSortedSlotIndices[Partner] = TempSlot;
+				}
+			}
+			GroupMemoryBarrierWithGroupSync();
+		}
+	}
+}
+
+#endif // USE_BITONIC_SCATTER
+
+[numthreads(THREADGROUP_SIZE, 1, 1)]
+void ScatterCS(uint3 GroupId : SV_GroupID, uint3 GroupThreadId : SV_GroupThreadID)
+{
+	uint GroupIndex = GroupId.x + GroupId.y * GroupsX;
+	uint LocalId = GroupThreadId.x;
+	uint SlotIndex = GroupIndex * THREADGROUP_SIZE + LocalId;
+
+#ifdef USE_BITONIC_SCATTER
+	// Bitonic sort path: sort cell IDs in shared memory, then one coalesced
+	// atomic per unique cell per group. Wins when cells are dense (many
+	// particles per cell), loses when sparse due to sort overhead.
+
+	// Load cell IDs and slot indices into shared memory
+	// Threads past NumSlots load INVALID_CELL so they sort to the end
+	int MyCellId = (SlotIndex < NumSlots) ? CellIdBuffer[SlotIndex] : (int)NEIGHBORGRID_INVALID_CELL;
+	ScatterSortedCellIds[LocalId] = MyCellId;
+	ScatterSortedSlotIndices[LocalId] = SlotIndex;
+	GroupMemoryBarrierWithGroupSync();
+
+	// Sort by cell ID, carrying slot indices as satellite data
+	BitonicSortScatterShared(LocalId);
+
+	// After sort, read back this thread's (potentially moved) cell ID and slot index
+	MyCellId = ScatterSortedCellIds[LocalId];
+	uint MySlotIndex = ScatterSortedSlotIndices[LocalId];
+	bool IsValid = (MyCellId != (int)NEIGHBORGRID_INVALID_CELL);
+
+	// Compute local offset within the run: distance from the run start
+	uint LocalOffset = 0;
+	if (IsValid)
+	{
+		for (uint j = LocalId; j > 0 && ScatterSortedCellIds[j - 1] == MyCellId; j--)
+		{
+			LocalOffset++;
+		}
+	}
+
+	// First thread of each run claims a contiguous block via one coalesced atomic
+	// and writes the base position to shared memory for the rest of the run to read
+	bool IsRunStart = IsValid && ((LocalId == 0) || (ScatterSortedCellIds[LocalId - 1] != MyCellId));
+	if (IsRunStart)
+	{
+		// Count run length forward
+		uint RunCount = 1;
+		for (uint k = LocalId + 1; k < THREADGROUP_SIZE && ScatterSortedCellIds[k] == MyCellId; k++)
+		{
+			RunCount++;
+		}
+
+		// One atomic per unique cell per group — claim a contiguous block
+		int WriteBase;
+		InterlockedAdd(CellOffsetBuffer[MyCellId], (int)RunCount, WriteBase);
+		RunWriteBase[LocalId] = (uint)WriteBase;
+	}
+	GroupMemoryBarrierWithGroupSync();
+
+	if (IsValid)
+	{
+		// Find the run start to read the base write position
+		uint RunStartId = LocalId - LocalOffset;
+		uint WritePos = RunWriteBase[RunStartId] + LocalOffset;
+
+		// Write the particle index to the pre-computed position — no atomic needed
+		uint ParticleIndex = MySlotIndex / MaxCellsPerParticle;
+		ParticleListBuffer[WritePos] = ParticleIndex;
+	}
+#else
+	// Direct atomic path: one InterlockedAdd per valid slot.
+	// Simpler, faster when contention is low (sparse grids, many unique cells).
+	if (SlotIndex >= NumSlots)
+	{
+		return;
+	}
+
+	int CellId = CellIdBuffer[SlotIndex];
+	if (CellId != (int)NEIGHBORGRID_INVALID_CELL)
+	{
+		int WritePos;
+		InterlockedAdd(CellOffsetBuffer[CellId], 1, WritePos);
+
+		uint ParticleIndex = SlotIndex / MaxCellsPerParticle;
+		ParticleListBuffer[WritePos] = ParticleIndex;
+	}
+#endif
+}
+
+#endif
diff --git a/Engine/Plugins/FX/Niagara/Source/Niagara/Classes/NiagaraDataInterfaceParticleNeighborQuery.h b/Engine/Plugins/FX/Niagara/Source/Niagara/Classes/NiagaraDataInterfaceParticleNeighborQuery.h
new file mode 100644
--- /dev/null
+++ b/Engine/Plugins/FX/Niagara/Source/Niagara/Classes/NiagaraDataInterfaceParticleNeighborQuery.h
@@ -0,0 +1,133 @@
+// Copyright Epic Games, Inc. All Rights Reserved.
+#pragma once
+
+#include "NiagaraDataInterface.h"
+#include "NiagaraCommon.h"
+#include "NiagaraDataInterfaceRW.h"
+#include "NiagaraRenderGraphUtils.h"
+
+#include "NiagaraDataInterfaceParticleNeighborQuery.generated.h"
+
+class FNiagaraSystemInstance;
+
+struct FNDIParticleNeighborQueryInstanceData_GT
+{
+	FIntVector	NumCells = FIntVector::ZeroValue;
+	uint32		MaxCellsPerParticle = 1;
+	bool		bCountOnly = false;
+	bool		bNeedsRealloc = false;
+};
+
+struct FNDIParticleNeighborQueryInstanceData_RT
+{
+	void ResizeBuffers(FRDGBuilder& GraphBuilder);
+
+	bool ClearBeforeNonIterationStage = true;
+
+	FIntVector	NumCells = FIntVector::ZeroValue;
+	uint32		MaxCellsPerParticle = 1;
+	uint32		NumParticles = 0;
+	uint64		AllocatedNumSlots = 0;
+	uint32		AllocatedNumCells = 0;
+	bool		bCountOnly = false;
+	bool		bNeedsRealloc = false;
+
+	// Counting sort buffers
+	FNiagaraPooledRWBuffer CellIdBuffer;			// numSlots (numParticles * MaxCellsPerParticle) - cell assignment per slot
+	FNiagaraPooledRWBuffer CellCountBuffer;			// numCells - particles per cell (histogram output)
+	FNiagaraPooledRWBuffer CellOffsetBuffer;		// numCells - prefix sum write offset per cell
+	FNiagaraPooledRWBuffer ParticleListBuffer;		// numSlots - final sorted particle indices
+};
+
+struct FNiagaraDataInterfaceProxyParticleNeighborQuery : public FNiagaraDataInterfaceProxyRW
+{
+	virtual void ResetData(const FNDIGpuComputeResetContext& Context) override;
+	virtual void PreStage(const FNDIGpuComputePreStageContext& Context) override;
+	virtual void PostStage(const FNDIGpuComputePostStageContext& Context) override;
+	virtual void PostSimulate(const FNDIGpuComputePostSimulateContext& Context) override;
+
+	virtual void ConsumePerInstanceDataFromGameThread(void* PerInstanceData, const FNiagaraSystemInstanceID& Instance) override {}
+	virtual int32 PerInstanceDataPassedToRenderThreadSize() const override { return 0; }
+
+	virtual void GetDispatchArgs(const FNDIGpuComputeDispatchArgsGenContext& Context) override;
+
+	TMap<FNiagaraSystemInstanceID, FNDIParticleNeighborQueryInstanceData_RT> SystemInstancesToProxyData_RT;
+};
+
+UCLASS(EditInlineNew, Category = "Grid", CollapseCategories, meta = (DisplayName = "Particle Neighbor Query"), MinimalAPI)
+class UNiagaraDataInterfaceParticleNeighborQuery : public UNiagaraDataInterfaceGrid3D
+{
+	GENERATED_UCLASS_BODY()
+
+	BEGIN_SHADER_PARAMETER_STRUCT(FShaderParameters, )
+		SHADER_PARAMETER(FIntVector,		NumCells)
+		SHADER_PARAMETER(FVector3f,			UnitToUV)
+		SHADER_PARAMETER(int32,				MaxCellsPerParticleValue)
+
+		// Read path
+		SHADER_PARAMETER_RDG_BUFFER_SRV(Buffer<int>,	CellCount)
+		SHADER_PARAMETER_RDG_BUFFER_SRV(Buffer<int>,	CellOffset)
+		SHADER_PARAMETER_RDG_BUFFER_SRV(Buffer<int>,	ParticleList)
+
+		// Write path AddParticle calls write to this
+		SHADER_PARAMETER_RDG_BUFFER_UAV(RWBuffer<int>,	OutputCellId)
+	END_SHADER_PARAMETER_STRUCT()
+
+public:
+	UPROPERTY(EditAnywhere, Category = "Grid")
+	uint32 MaxCellsPerParticle;
+
+	/** Skip prefix sum and scatter passes — only the histogram is computed.
+	  * GetParticleNeighborCount works; GetParticleNeighbor does not.
+	  * Saves 2 GPU dispatches and ~2x buffer memory. */
+	UPROPERTY(EditAnywhere, Category = "Grid")
+	bool bCountOnly = false;
+
+	virtual void PostInitProperties() override
+	{
+		Super::PostInitProperties();
+
+		if (HasAnyFlags(RF_ClassDefaultObject))
+		{
+			ENiagaraTypeRegistryFlags Flags = ENiagaraTypeRegistryFlags::AllowAnyVariable | ENiagaraTypeRegistryFlags::AllowParameter;
+			FNiagaraTypeRegistry::Register(FNiagaraTypeDefinition(GetClass()), Flags);
+		}
+	}
+
+	//~ UNiagaraDataInterface interface
+	// VM functionality
+	NIAGARA_API virtual void GetVMExternalFunction(const FVMExternalFunctionBindingInfo& BindingInfo, void* InstanceData, FVMExternalFunction& OutFunc) override;
+
+	NIAGARA_API virtual bool Equals(const UNiagaraDataInterface* Other) const override;
+
+	// GPU sim functionality
+#if WITH_EDITORONLY_DATA
+	NIAGARA_API virtual bool AppendCompileHash(FNiagaraCompileHashVisitor* InVisitor) const override;
+	NIAGARA_API virtual void GetParameterDefinitionHLSL(const FNiagaraDataInterfaceGPUParamInfo& ParamInfo, FString& OutHLSL) override;
+	NIAGARA_API virtual bool GetFunctionHLSL(const FNiagaraDataInterfaceGPUParamInfo& ParamInfo, const FNiagaraDataInterfaceGeneratedFunction& FunctionInfo, int FunctionInstanceIndex, FString& OutHLSL) override;
+#endif
+	NIAGARA_API virtual void BuildShaderParameters(FNiagaraShaderParametersBuilder& ShaderParametersBuilder) const override;
+	NIAGARA_API virtual void SetShaderParameters(const FNiagaraDataInterfaceSetShaderParametersContext& Context) const override;
+
+	virtual void ProvidePerInstanceDataForRenderThread(void* DataForRenderThread, void* PerInstanceData, const FNiagaraSystemInstanceID& SystemInstance) override {}
+	NIAGARA_API virtual bool InitPerInstanceData(void* PerInstanceData, FNiagaraSystemInstance* SystemInstance) override;
+	NIAGARA_API virtual void DestroyPerInstanceData(void* PerInstanceData, FNiagaraSystemInstance* SystemInstance) override;
+	virtual bool PerInstanceTick(void* PerInstanceData, FNiagaraSystemInstance* SystemInstance, float DeltaSeconds) override { return false; }
+	virtual int32 PerInstanceDataSize() const override { return sizeof(FNDIParticleNeighborQueryInstanceData_GT); }
+	NIAGARA_API virtual bool PerInstanceTickPostSimulate(void* PerInstanceData, FNiagaraSystemInstance* SystemInstance, float DeltaSeconds) override;
+	virtual bool HasPostSimulateTick() const override { return true; }	
+	//~ UNiagaraDataInterface interface END
+
+	NIAGARA_API void GetNumCells(FVectorVMExternalFunctionContext& Context);
+	NIAGARA_API void GetMaxCellsPerParticle(FVectorVMExternalFunctionContext& Context);
+	NIAGARA_API void SetNumCells(FVectorVMExternalFunctionContext& Context);
+
+
+protected:
+	//~ UNiagaraDataInterface interface
+#if WITH_EDITORONLY_DATA
+	virtual void GetFunctionsInternal(TArray<FNiagaraFunctionSignature>& OutFunctions) const override;
+#endif
+	NIAGARA_API virtual bool CopyToInternal(UNiagaraDataInterface* Destination) const override;
+	//~ UNiagaraDataInterface interface END
+};
diff --git a/Engine/Plugins/FX/Niagara/Source/Niagara/Private/NiagaraDataInterfaceParticleNeighborQuery.cpp b/Engine/Plugins/FX/Niagara/Source/Niagara/Private/NiagaraDataInterfaceParticleNeighborQuery.cpp
new file mode 100644
--- /dev/null
+++ b/Engine/Plugins/FX/Niagara/Source/Niagara/Private/NiagaraDataInterfaceParticleNeighborQuery.cpp
@@ -0,0 +1,718 @@
+// Copyright Epic Games, Inc. All Rights Reserved.
+#include "NiagaraDataInterfaceParticleNeighborQuery.h"
+#include "NiagaraCompileHashVisitor.h"
+#include "NiagaraGpuComputeDispatchInterface.h"
+#include "NiagaraSimStageData.h"
+#include "NiagaraShaderParametersBuilder.h"
+#include "NiagaraSystemInstance.h"
+#include "DataDrivenShaderPlatformInfo.h"
+#include "RenderGraphUtils.h"
+#include "NiagaraParticleNeighborQuerySort.h"
+
+#include UE_INLINE_GENERATED_CPP_BY_NAME(NiagaraDataInterfaceParticleNeighborQuery)
+
+#define LOCTEXT_NAMESPACE "NiagaraDataInterfaceParticleNeighborQuery"
+
+namespace NDIParticleNeighborQueryLocal
+{
+	static const TCHAR* TemplateShaderFile = TEXT("/Plugin/FX/Niagara/Private/NiagaraDataInterfaceParticleNeighborQuery.ush");
+
+	// Function names — must match the _{ParameterName} suffix functions in the .ush	
+	static const FName SetNumCellsFunctionName(TEXT("SetNumCells"));
+	static const FName MaxCellsPerParticleFunctionName(TEXT("MaxCellsPerParticle"));
+	static const FName AddParticleFunctionName(TEXT("AddParticle"));
+	static const FName AddParticleWithRadiusFunctionName(TEXT("AddParticleWithRadius"));
+	static const FName AddParticleToOverlappedCellsFunctionName(TEXT("AddParticleToOverlappedCells"));
+	static const FName AddParticleToNeighborCellsFunctionName(TEXT("AddParticleToNeighborCells"));
+	static const FName GetParticleNeighborCountFunctionName(TEXT("GetParticleNeighborCount"));
+	static const FName GetParticleNeighborFunctionName(TEXT("GetParticleNeighbor"));
+}
+
+
+
+static int32 GMaxNiagaraParticleNeighborQueryCells = (512 * 512 * 512);
+static FAutoConsoleVariableRef CVarMaxNiagaraParticleNeighborQueryCells(
+	TEXT("fx.MaxNiagaraParticleNeighborQueryCells"),
+	GMaxNiagaraParticleNeighborQueryCells,
+	TEXT("The max number of supported grid cells for ParticleNeighborQuery. Overflowing this threshold will cause the sim to warn and fail.\n"),
+	ECVF_Default
+);
+
+//////////////////////////////////////////////////////////////////////////
+
+void FNDIParticleNeighborQueryInstanceData_RT::ResizeBuffers(FRDGBuilder& GraphBuilder)
+{
+	const uint32 NumTotalCells = NumCells.X * NumCells.Y * NumCells.Z;
+	const uint64 NumSlots = FMath::Max<uint64>((uint64) NumParticles * MaxCellsPerParticle, 1u);
+
+	// Slot-sized buffers: reallocate when particle count * MaxCellsPerParticle grows
+	if (NumSlots > AllocatedNumSlots)
+	{
+		CellIdBuffer.Release();
+
+		CellIdBuffer.Initialize(GraphBuilder, TEXT("ParticleNeighborQuery::CellId"), EPixelFormat::PF_R32_SINT, sizeof(int32), NumSlots, BUF_Static);
+
+		if (!bCountOnly)
+		{
+			ParticleListBuffer.Release();
+			ParticleListBuffer.Initialize(GraphBuilder, TEXT("ParticleNeighborQuery::ParticleList"), EPixelFormat::PF_R32_SINT, sizeof(int32), NumSlots, BUF_Static);
+		}
+
+		AllocatedNumSlots = NumSlots;
+	}
+
+	// Cell-sized buffers: reallocate when grid dimensions change
+	if (NumTotalCells != AllocatedNumCells)
+	{
+		if (NumTotalCells > (uint32)GMaxNiagaraParticleNeighborQueryCells)
+		{
+			UE_LOG(LogNiagara, Verbose, TEXT("NiagaraDataInterfaceParticleNeighborQuery - Invalid NumCells(%dx%dx%d) with total = %u exceeds the maximum of %d ."), NumCells.X, NumCells.Y, NumCells.Z, NumTotalCells, GMaxNiagaraParticleNeighborQueryCells);
+			// Do NOT clear bNeedsRealloc — cell buffers were not allocated, so subsequent
+			// frames must retry. Slot-sized buffers above may have succeeded, but the system
+			// is not functional without cell buffers.
+			return;
+		}
+
+		CellCountBuffer.Release();
+
+		const uint32 CellBufferSize = FMath::Max<uint32>(NumTotalCells, 1u);
+		CellCountBuffer.Initialize(GraphBuilder, TEXT("ParticleNeighborQuery::CellCount"), EPixelFormat::PF_R32_SINT, sizeof(int32), CellBufferSize, BUF_Static);
+
+		if (!bCountOnly)
+		{
+			CellOffsetBuffer.Release();
+			CellOffsetBuffer.Initialize(GraphBuilder, TEXT("ParticleNeighborQuery::CellOffset"), EPixelFormat::PF_R32_SINT, sizeof(int32), CellBufferSize, BUF_Static);
+		}
+
+		AllocatedNumCells = NumTotalCells;
+	}
+
+	// Both allocation blocks succeeded — clear the realloc flag.
+	bNeedsRealloc = false;
+}
+
+//////////////////////////////////////////////////////////////////////////
+
+UNiagaraDataInterfaceParticleNeighborQuery::UNiagaraDataInterfaceParticleNeighborQuery(FObjectInitializer const& ObjectInitializer)
+	: Super(ObjectInitializer)
+	, MaxCellsPerParticle(1)
+	, bCountOnly(false)
+{
+	Proxy.Reset(new FNiagaraDataInterfaceProxyParticleNeighborQuery());
+}
+
+#if WITH_EDITORONLY_DATA
+void UNiagaraDataInterfaceParticleNeighborQuery::GetFunctionsInternal(TArray<FNiagaraFunctionSignature>& OutFunctions) const
+{
+	using namespace NDIParticleNeighborQueryLocal;
+
+	Super::GetFunctionsInternal(OutFunctions);
+
+	{
+		FNiagaraFunctionSignature Sig;
+		Sig.Name = SetNumCellsFunctionName;
+		Sig.Inputs.Add(FNiagaraVariable(FNiagaraTypeDefinition(GetClass()), TEXT("Grid")));
+		Sig.Inputs.Add(FNiagaraVariable(FNiagaraTypeDefinition::GetIntDef(), TEXT("NumCellsX")));
+		Sig.Inputs.Add(FNiagaraVariable(FNiagaraTypeDefinition::GetIntDef(), TEXT("NumCellsY")));
+		Sig.Inputs.Add(FNiagaraVariable(FNiagaraTypeDefinition::GetIntDef(), TEXT("NumCellsZ")));
+		Sig.Inputs.Add(FNiagaraVariable(FNiagaraTypeDefinition::GetIntDef(), TEXT("MaxCellsPerParticle")));
+		Sig.Outputs.Add(FNiagaraVariable(FNiagaraTypeDefinition::GetBoolDef(), TEXT("Success")));
+
+		Sig.ModuleUsageBitmask = ENiagaraScriptUsageMask::Emitter | ENiagaraScriptUsageMask::System;
+		Sig.bMemberFunction = true;
+		Sig.bRequiresExecPin = true;
+		Sig.bRequiresContext = false;
+		Sig.bSupportsCPU = true;
+		Sig.bSupportsGPU = false;
+		OutFunctions.Add(Sig);
+	}
+
+	// MaxCellsPerParticle() -> MaxCellsPerParticle
+	{
+		FNiagaraFunctionSignature Sig;
+		Sig.Name = MaxCellsPerParticleFunctionName;
+		Sig.Inputs.Add(FNiagaraVariable(FNiagaraTypeDefinition(GetClass()), TEXT("Grid")));
+		Sig.Outputs.Add(FNiagaraVariable(FNiagaraTypeDefinition::GetIntDef(), TEXT("MaxCellsPerParticle")));
+
+		Sig.bMemberFunction = true;
+		Sig.bRequiresContext = false;
+		OutFunctions.Add(Sig);
+	}
+
+	// AddParticle(Enabled, IndexPos, ParticleIndex) -> Success
+	{
+		FNiagaraFunctionSignature Sig;
+		Sig.Name = NDIParticleNeighborQueryLocal::AddParticleFunctionName;
+		Sig.Inputs.Add(FNiagaraVariable(FNiagaraTypeDefinition(GetClass()), TEXT("Grid")));
+		Sig.Inputs.Add_GetRef(FNiagaraVariable(FNiagaraTypeDefinition::GetBoolDef(), TEXT("Enabled"))).SetValue(true);
+		Sig.Inputs.Add(FNiagaraVariable(FNiagaraTypeDefinition::GetVec3Def(), TEXT("IndexPos")));
+		Sig.Inputs.Add(FNiagaraVariable(FNiagaraTypeDefinition::GetIntDef(), TEXT("ParticleIndex")));
+		Sig.Outputs.Add(FNiagaraVariable(FNiagaraTypeDefinition::GetBoolDef(), TEXT("Success")));
+
+		Sig.bWriteFunction = true;
+		Sig.bMemberFunction = true;
+		Sig.bRequiresContext = false;
+		Sig.bRequiresExecPin = true;
+		OutFunctions.Add(Sig);
+	}
+
+	// AddParticleWithRadius(Enabled, IndexPos, Radius, ParticleIndex) -> Success
+	{
+		FNiagaraFunctionSignature Sig;
+		Sig.Name = AddParticleWithRadiusFunctionName;
+		Sig.Inputs.Add(FNiagaraVariable(FNiagaraTypeDefinition(GetClass()), TEXT("Grid")));
+		Sig.Inputs.Add_GetRef(FNiagaraVariable(FNiagaraTypeDefinition::GetBoolDef(), TEXT("Enabled"))).SetValue(true);
+		Sig.Inputs.Add(FNiagaraVariable(FNiagaraTypeDefinition::GetVec3Def(), TEXT("IndexPos")));
+		Sig.Inputs.Add(FNiagaraVariable(FNiagaraTypeDefinition::GetFloatDef(), TEXT("Radius")));
+		Sig.Inputs.Add(FNiagaraVariable(FNiagaraTypeDefinition::GetIntDef(), TEXT("ParticleIndex")));
+		Sig.Outputs.Add(FNiagaraVariable(FNiagaraTypeDefinition::GetBoolDef(), TEXT("Success")));
+
+		Sig.bWriteFunction = true;
+		Sig.bMemberFunction = true;
+		Sig.bRequiresContext = false;
+		Sig.bRequiresExecPin = true;
+		OutFunctions.Add(Sig);
+	}
+
+	// AddParticleToOverlappedCells(Enabled, IndexPos, ParticleIndex) -> Success
+	// Writes the particle to all 8 cells in the 2x2x2 positive neighborhood from floor(IndexPos).
+	// For asymmetric gather (P2G rasterization with trilinear kernel). Requires MaxCellsPerParticle >= 8.
+	{
+		FNiagaraFunctionSignature Sig;
+		Sig.Name = AddParticleToOverlappedCellsFunctionName;
+		Sig.Inputs.Add(FNiagaraVariable(FNiagaraTypeDefinition(GetClass()), TEXT("Grid")));
+		Sig.Inputs.Add_GetRef(FNiagaraVariable(FNiagaraTypeDefinition::GetBoolDef(), TEXT("Enabled"))).SetValue(true);
+		Sig.Inputs.Add(FNiagaraVariable(FNiagaraTypeDefinition::GetVec3Def(), TEXT("IndexPos")));
+		Sig.Inputs.Add(FNiagaraVariable(FNiagaraTypeDefinition::GetIntDef(), TEXT("ParticleIndex")));
+		Sig.Outputs.Add(FNiagaraVariable(FNiagaraTypeDefinition::GetBoolDef(), TEXT("Success")));
+
+		Sig.bWriteFunction = true;
+		Sig.bMemberFunction = true;
+		Sig.bRequiresContext = false;
+		Sig.bRequiresExecPin = true;
+		OutFunctions.Add(Sig);
+	}
+
+	// AddParticleToNeighborCells(Enabled, IndexPos, ParticleIndex) -> Success
+	// Writes the particle to all 27 cells in the 3x3x3 neighborhood centered on floor(IndexPos).
+	// For symmetric bilateral queries (PBD constraints) with single-cell query. Requires MaxCellsPerParticle >= 27.
+	{
+		FNiagaraFunctionSignature Sig;
+		Sig.Name = AddParticleToNeighborCellsFunctionName;
+		Sig.Inputs.Add(FNiagaraVariable(FNiagaraTypeDefinition(GetClass()), TEXT("Grid")));
+		Sig.Inputs.Add_GetRef(FNiagaraVariable(FNiagaraTypeDefinition::GetBoolDef(), TEXT("Enabled"))).SetValue(true);
+		Sig.Inputs.Add(FNiagaraVariable(FNiagaraTypeDefinition::GetVec3Def(), TEXT("IndexPos")));
+		Sig.Inputs.Add(FNiagaraVariable(FNiagaraTypeDefinition::GetIntDef(), TEXT("ParticleIndex")));
+		Sig.Outputs.Add(FNiagaraVariable(FNiagaraTypeDefinition::GetBoolDef(), TEXT("Success")));
+
+		Sig.bWriteFunction = true;
+		Sig.bMemberFunction = true;
+		Sig.bRequiresContext = false;
+		Sig.bRequiresExecPin = true;
+		OutFunctions.Add(Sig);
+	}
+
+	// GetParticleNeighborCount(IndexX, IndexY, IndexZ) -> Count
+	{
+		FNiagaraFunctionSignature Sig;
+		Sig.Name = NDIParticleNeighborQueryLocal::GetParticleNeighborCountFunctionName;
+		Sig.Inputs.Add(FNiagaraVariable(FNiagaraTypeDefinition(GetClass()), TEXT("Grid")));
+		Sig.Inputs.Add(FNiagaraVariable(FNiagaraTypeDefinition::GetIntDef(), TEXT("IndexX")));
+		Sig.Inputs.Add(FNiagaraVariable(FNiagaraTypeDefinition::GetIntDef(), TEXT("IndexY")));
+		Sig.Inputs.Add(FNiagaraVariable(FNiagaraTypeDefinition::GetIntDef(), TEXT("IndexZ")));
+		Sig.Outputs.Add(FNiagaraVariable(FNiagaraTypeDefinition::GetIntDef(), TEXT("Count")));
+
+		Sig.bMemberFunction = true;
+		Sig.bRequiresContext = false;
+		OutFunctions.Add(Sig);
+	}
+
+	// GetParticleNeighbor(IndexX, IndexY, IndexZ, NeighborIdx) -> ParticleIdx
+	{
+		FNiagaraFunctionSignature Sig;
+		Sig.Name = NDIParticleNeighborQueryLocal::GetParticleNeighborFunctionName;
+		Sig.Inputs.Add(FNiagaraVariable(FNiagaraTypeDefinition(GetClass()), TEXT("Grid")));
+		Sig.Inputs.Add(FNiagaraVariable(FNiagaraTypeDefinition::GetIntDef(), TEXT("IndexX")));
+		Sig.Inputs.Add(FNiagaraVariable(FNiagaraTypeDefinition::GetIntDef(), TEXT("IndexY")));
+		Sig.Inputs.Add(FNiagaraVariable(FNiagaraTypeDefinition::GetIntDef(), TEXT("IndexZ")));
+		Sig.Inputs.Add(FNiagaraVariable(FNiagaraTypeDefinition::GetIntDef(), TEXT("NeighborIdx")));
+		Sig.Outputs.Add(FNiagaraVariable(FNiagaraTypeDefinition::GetIntDef(), TEXT("ParticleIdx")));
+
+		Sig.bMemberFunction = true;
+		Sig.bRequiresContext = false;
+		OutFunctions.Add(Sig);
+	}
+
+}
+#endif
+
+//////////////////////////////////////////////////////////////////////////
+// VM function bindings
+
+DEFINE_NDI_DIRECT_FUNC_BINDER(UNiagaraDataInterfaceParticleNeighborQuery, SetNumCells);
+void UNiagaraDataInterfaceParticleNeighborQuery::GetVMExternalFunction(const FVMExternalFunctionBindingInfo& BindingInfo, void* InstanceData, FVMExternalFunction& OutFunc)
+{
+	using namespace NDIParticleNeighborQueryLocal;
+
+	Super::GetVMExternalFunction(BindingInfo, InstanceData, OutFunc);
+
+	if (BindingInfo.Name == UNiagaraDataInterfaceRWBase::NumCellsFunctionName)
+	{
+		check(BindingInfo.GetNumInputs() == 1 && BindingInfo.GetNumOutputs() == 3);
+		OutFunc.BindUObject(this, &UNiagaraDataInterfaceParticleNeighborQuery::GetNumCells);
+	}
+	else if (BindingInfo.Name == MaxCellsPerParticleFunctionName)
+	{
+		check(BindingInfo.GetNumInputs() == 1 && BindingInfo.GetNumOutputs() == 1);
+		OutFunc.BindUObject(this, &UNiagaraDataInterfaceParticleNeighborQuery::GetMaxCellsPerParticle);
+	}
+	else if (BindingInfo.Name == SetNumCellsFunctionName)
+	{
+		check(BindingInfo.GetNumInputs() == 5 && BindingInfo.GetNumOutputs() == 1);
+		NDI_FUNC_BINDER(UNiagaraDataInterfaceParticleNeighborQuery, SetNumCells)::Bind(this, OutFunc);
+	}
+}
+
+//////////////////////////////////////////////////////////////////////////
+// VM function implementations
+
+void UNiagaraDataInterfaceParticleNeighborQuery::GetNumCells(FVectorVMExternalFunctionContext& Context)
+{
+	VectorVM::FUserPtrHandler<FNDIParticleNeighborQueryInstanceData_GT> InstData(Context);
+
+	FNDIOutputParam<int32> NumCellsX(Context);
+	FNDIOutputParam<int32> NumCellsY(Context);
+	FNDIOutputParam<int32> NumCellsZ(Context);
+
+	int32 TmpNumCellsX = InstData->NumCells.X;
+	int32 TmpNumCellsY = InstData->NumCells.Y;
+	int32 TmpNumCellsZ = InstData->NumCells.Z;
+
+	for (int32 InstanceIdx = 0; InstanceIdx < Context.GetNumInstances(); ++InstanceIdx)
+	{
+		NumCellsX.SetAndAdvance(TmpNumCellsX);
+		NumCellsY.SetAndAdvance(TmpNumCellsY);
+		NumCellsZ.SetAndAdvance(TmpNumCellsZ);
+	}
+}
+
+void UNiagaraDataInterfaceParticleNeighborQuery::GetMaxCellsPerParticle(FVectorVMExternalFunctionContext& Context)
+{
+	VectorVM::FUserPtrHandler<FNDIParticleNeighborQueryInstanceData_GT> InstData(Context);
+
+	FNDIOutputParam<int32> OutMaxCellsPerParticle(Context);
+
+	for (int32 InstanceIdx = 0; InstanceIdx < Context.GetNumInstances(); ++InstanceIdx)
+	{
+		OutMaxCellsPerParticle.SetAndAdvance(InstData->MaxCellsPerParticle);
+	}
+}
+
+bool UNiagaraDataInterfaceParticleNeighborQuery::Equals(const UNiagaraDataInterface* Other) const
+{
+	if (!Super::Equals(Other))
+	{
+		return false;
+	}
+	const UNiagaraDataInterfaceParticleNeighborQuery* OtherTyped = CastChecked<const UNiagaraDataInterfaceParticleNeighborQuery>(Other);
+
+	return OtherTyped->MaxCellsPerParticle == MaxCellsPerParticle
+		&& OtherTyped->bCountOnly == bCountOnly;
+}
+
+//////////////////////////////////////////////////////////////////////////
+// HLSL generation — .ush template approach
+
+#if WITH_EDITORONLY_DATA
+bool UNiagaraDataInterfaceParticleNeighborQuery::AppendCompileHash(FNiagaraCompileHashVisitor* InVisitor) const
+{
+	if (!Super::AppendCompileHash(InVisitor))
+	{
+		return false;
+	}
+
+	InVisitor->UpdateShaderFile(NDIParticleNeighborQueryLocal::TemplateShaderFile);
+	InVisitor->UpdateShaderParameters<FShaderParameters>();
+
+	return true;
+}
+
+void UNiagaraDataInterfaceParticleNeighborQuery::GetParameterDefinitionHLSL(const FNiagaraDataInterfaceGPUParamInfo& ParamInfo, FString& OutHLSL)
+{	
+	Super::GetParameterDefinitionHLSL(ParamInfo, OutHLSL);
+	
+	const TMap<FString, FStringFormatArg> TemplateArgs =
+	{
+		{TEXT("ParameterName"), ParamInfo.DataInterfaceHLSLSymbol},
+	};
+	AppendTemplateHLSL(OutHLSL, NDIParticleNeighborQueryLocal::TemplateShaderFile, TemplateArgs);
+}
+
+bool UNiagaraDataInterfaceParticleNeighborQuery::GetFunctionHLSL(const FNiagaraDataInterfaceGPUParamInfo& ParamInfo, const FNiagaraDataInterfaceGeneratedFunction& FunctionInfo, int FunctionInstanceIndex, FString& OutHLSL)
+{
+	using namespace NDIParticleNeighborQueryLocal;
+
+	bool ParentRet = Super::GetFunctionHLSL(ParamInfo, FunctionInfo, FunctionInstanceIndex, OutHLSL);
+	if (ParentRet)
+	{
+		return true;
+	}
+
+	// All functions are handled by the .ush template (including _UEImpureCall wrappers for exec-pin functions).
+	if (FunctionInfo.DefinitionName == MaxCellsPerParticleFunctionName ||
+		FunctionInfo.DefinitionName == NDIParticleNeighborQueryLocal::AddParticleFunctionName ||
+		FunctionInfo.DefinitionName == AddParticleWithRadiusFunctionName ||
+		FunctionInfo.DefinitionName == AddParticleToOverlappedCellsFunctionName ||
+		FunctionInfo.DefinitionName == AddParticleToNeighborCellsFunctionName ||
+		FunctionInfo.DefinitionName == NDIParticleNeighborQueryLocal::GetParticleNeighborCountFunctionName ||
+		FunctionInfo.DefinitionName == NDIParticleNeighborQueryLocal::GetParticleNeighborFunctionName)
+	{
+		return true;
+	}
+
+	return false;
+}
+#endif
+
+//////////////////////////////////////////////////////////////////////////
+// Shader parameter binding
+
+void UNiagaraDataInterfaceParticleNeighborQuery::BuildShaderParameters(FNiagaraShaderParametersBuilder& ShaderParametersBuilder) const
+{
+	ShaderParametersBuilder.AddNestedStruct<FShaderParameters>();
+}
+
+void UNiagaraDataInterfaceParticleNeighborQuery::SetShaderParameters(const FNiagaraDataInterfaceSetShaderParametersContext& Context) const
+{
+	FNiagaraDataInterfaceProxyParticleNeighborQuery& DIProxy = Context.GetProxy<FNiagaraDataInterfaceProxyParticleNeighborQuery>();
+	FNDIParticleNeighborQueryInstanceData_RT* ProxyData = DIProxy.SystemInstancesToProxyData_RT.Find(Context.GetSystemInstanceID());
+	FRDGBuilder& GraphBuilder = Context.GetGraphBuilder();
+
+	FShaderParameters* ShaderParameters = Context.GetParameterNestedStruct<FShaderParameters>();
+	if (ProxyData && ProxyData->CellCountBuffer.IsValid())
+	{
+		ShaderParameters->NumCells = ProxyData->NumCells;
+		ShaderParameters->UnitToUV = FVector3f(1.0f) / FVector3f(ProxyData->NumCells);
+		ShaderParameters->MaxCellsPerParticleValue = ProxyData->MaxCellsPerParticle;
+
+		if (Context.IsOutputStage())
+		{
+			// Output stage: user modules write CellIdBuffer, read buffers are empty
+			ShaderParameters->CellCount = Context.GetComputeDispatchInterface().GetEmptyBufferSRV(GraphBuilder, PF_R32_SINT);
+			ShaderParameters->CellOffset = Context.GetComputeDispatchInterface().GetEmptyBufferSRV(GraphBuilder, PF_R32_SINT);
+			ShaderParameters->ParticleList = Context.GetComputeDispatchInterface().GetEmptyBufferSRV(GraphBuilder, PF_R32_SINT);
+			ShaderParameters->OutputCellId = ProxyData->CellIdBuffer.GetOrCreateUAV(GraphBuilder);
+		}
+		else
+		{
+			// Read stage: counting sort has completed, read results
+			ShaderParameters->CellCount = ProxyData->CellCountBuffer.GetOrCreateSRV(GraphBuilder);
+			if (ProxyData->bCountOnly)
+			{
+				ShaderParameters->CellOffset = Context.GetComputeDispatchInterface().GetEmptyBufferSRV(GraphBuilder, PF_R32_SINT);
+				ShaderParameters->ParticleList = Context.GetComputeDispatchInterface().GetEmptyBufferSRV(GraphBuilder, PF_R32_SINT);
+			}
+			else
+			{
+				ShaderParameters->CellOffset = ProxyData->CellOffsetBuffer.GetOrCreateSRV(GraphBuilder);
+				ShaderParameters->ParticleList = ProxyData->ParticleListBuffer.GetOrCreateSRV(GraphBuilder);
+			}
+			ShaderParameters->OutputCellId = Context.GetComputeDispatchInterface().GetEmptyBufferUAV(GraphBuilder, PF_R32_SINT);
+		}
+	}
+	else
+	{
+		ShaderParameters->NumCells = FIntVector::ZeroValue;
+		ShaderParameters->UnitToUV = FVector3f::ZeroVector;
+		ShaderParameters->MaxCellsPerParticleValue = 0;
+
+		ShaderParameters->CellCount = Context.GetComputeDispatchInterface().GetEmptyBufferSRV(GraphBuilder, PF_R32_SINT);
+		ShaderParameters->CellOffset = Context.GetComputeDispatchInterface().GetEmptyBufferSRV(GraphBuilder, PF_R32_SINT);
+		ShaderParameters->ParticleList = Context.GetComputeDispatchInterface().GetEmptyBufferSRV(GraphBuilder, PF_R32_SINT);
+		ShaderParameters->OutputCellId = Context.GetComputeDispatchInterface().GetEmptyBufferUAV(GraphBuilder, PF_R32_SINT);
+	}
+}
+
+//////////////////////////////////////////////////////////////////////////
+// Instance data lifecycle
+
+bool UNiagaraDataInterfaceParticleNeighborQuery::InitPerInstanceData(void* PerInstanceData, FNiagaraSystemInstance* SystemInstance)
+{
+	if (UE::PixelFormat::HasCapabilities(EPixelFormat::PF_R32_SINT, EPixelFormatCapabilities::TypedUAVLoad | EPixelFormatCapabilities::TypedUAVStore) == false)
+	{
+		return false;
+	}
+
+	FNDIParticleNeighborQueryInstanceData_GT* InstanceData = new (PerInstanceData) FNDIParticleNeighborQueryInstanceData_GT();
+
+	// Defaults — expected to be overridden by SetNumCells from emitter graph
+	InstanceData->NumCells = FIntVector(1, 1, 1);
+	InstanceData->MaxCellsPerParticle = MaxCellsPerParticle;
+	InstanceData->bCountOnly = bCountOnly;
+
+	FNiagaraDataInterfaceProxyParticleNeighborQuery* RT_Proxy = GetProxyAs<FNiagaraDataInterfaceProxyParticleNeighborQuery>();
+	const FIntVector RT_NumCells = InstanceData->NumCells;
+	const uint32 RT_MaxCellsPerParticle = InstanceData->MaxCellsPerParticle;
+	const bool RT_bCountOnly = InstanceData->bCountOnly;
+
+	ENQUEUE_RENDER_COMMAND(FUpdateData)(
+		[RT_Proxy, RT_NumCells, RT_MaxCellsPerParticle, RT_bCountOnly, InstanceID = SystemInstance->GetId()](FRHICommandListImmediate& RHICmdList)
+		{
+			check(!RT_Proxy->SystemInstancesToProxyData_RT.Contains(InstanceID));
+			FNDIParticleNeighborQueryInstanceData_RT* TargetData = &RT_Proxy->SystemInstancesToProxyData_RT.Add(InstanceID);
+
+			TargetData->NumCells = RT_NumCells;
+			TargetData->MaxCellsPerParticle = RT_MaxCellsPerParticle;
+			TargetData->bCountOnly = RT_bCountOnly;
+			TargetData->bNeedsRealloc = true;
+		}
+	);
+
+	return true;
+}
+
+void UNiagaraDataInterfaceParticleNeighborQuery::SetNumCells(FVectorVMExternalFunctionContext& Context)
+{
+	VectorVM::FUserPtrHandler<FNDIParticleNeighborQueryInstanceData_GT> InstData(Context);
+	VectorVM::FExternalFuncInputHandler<int> InNumCellsX(Context);
+	VectorVM::FExternalFuncInputHandler<int> InNumCellsY(Context);
+	VectorVM::FExternalFuncInputHandler<int> InNumCellsZ(Context);
+	VectorVM::FExternalFuncInputHandler<int> InMaxCellsPerParticle(Context);
+	VectorVM::FExternalFuncRegisterHandler<FNiagaraBool> OutSuccess(Context);
+
+	for (int32 InstanceIdx = 0; InstanceIdx < Context.GetNumInstances(); ++InstanceIdx)
+	{
+		const int NewNumCellsX = InNumCellsX.GetAndAdvance();
+		const int NewNumCellsY = InNumCellsY.GetAndAdvance();
+		const int NewNumCellsZ = InNumCellsZ.GetAndAdvance();
+		const int NewMaxCellsPerParticle = InMaxCellsPerParticle.GetAndAdvance();
+		bool bSuccess = (InstData.Get() != nullptr && Context.GetNumInstances() == 1 && NewNumCellsX >= 1 && NewNumCellsY >= 1 && NewNumCellsZ >= 1 && NewMaxCellsPerParticle >= 1);
+
+		const uint64 NumTotalCells = (uint64) NewNumCellsX * NewNumCellsY * NewNumCellsZ;
+		if (NumTotalCells == 0 || NumTotalCells > (uint64)GMaxNiagaraParticleNeighborQueryCells)
+		{
+			UE_LOG(LogNiagara, Error, TEXT("NiagaraDataInterfaceParticleNeighborQuery - Invalid NumCells(%dx%dx%d) with total = %llu exceeds the maximum of %d for '%s'."), NewNumCellsX, NewNumCellsY, NewNumCellsZ, NumTotalCells, GMaxNiagaraParticleNeighborQueryCells, *GetFullNameSafe(this));
+			bSuccess = false;
+		}
+
+		*OutSuccess.GetDestAndAdvance() = bSuccess;
+		if (bSuccess)
+		{
+			FIntVector OldNumCells = InstData->NumCells;
+			int OldMaxCellsPerParticle = InstData->MaxCellsPerParticle;
+
+			InstData->NumCells.X = FMath::Max(1, NewNumCellsX);
+			InstData->NumCells.Y = FMath::Max(1, NewNumCellsY);
+			InstData->NumCells.Z = FMath::Max(1, NewNumCellsZ);
+			InstData->MaxCellsPerParticle = FMath::Max(1, NewMaxCellsPerParticle);
+
+			InstData->bNeedsRealloc = OldNumCells != InstData->NumCells || OldMaxCellsPerParticle != InstData->MaxCellsPerParticle;
+		}
+	}
+}
+
+bool UNiagaraDataInterfaceParticleNeighborQuery::PerInstanceTickPostSimulate(void* PerInstanceData, FNiagaraSystemInstance* SystemInstance, float DeltaSeconds)
+{
+	FNDIParticleNeighborQueryInstanceData_GT* InstanceData = static_cast<FNDIParticleNeighborQueryInstanceData_GT*>(PerInstanceData);
+
+	if (InstanceData->bNeedsRealloc && InstanceData->NumCells.X > 0 && InstanceData->NumCells.Y > 0 && InstanceData->NumCells.Z > 0)
+	{
+		InstanceData->bNeedsRealloc = false;
+
+		FNiagaraDataInterfaceProxyParticleNeighborQuery* RT_Proxy = GetProxyAs<FNiagaraDataInterfaceProxyParticleNeighborQuery>();
+		ENQUEUE_RENDER_COMMAND(FUpdateData)(
+			[RT_Proxy, RT_NumCells = InstanceData->NumCells, RT_MaxCellsPerParticle = InstanceData->MaxCellsPerParticle, RT_bCountOnly = InstanceData->bCountOnly, InstanceID = SystemInstance->GetId()](FRHICommandListImmediate& RHICmdList)
+		{
+			check(RT_Proxy->SystemInstancesToProxyData_RT.Contains(InstanceID));
+			FNDIParticleNeighborQueryInstanceData_RT* TargetData = RT_Proxy->SystemInstancesToProxyData_RT.Find(InstanceID);
+
+			TargetData->NumCells = RT_NumCells;
+			TargetData->MaxCellsPerParticle = RT_MaxCellsPerParticle;
+			TargetData->bCountOnly = RT_bCountOnly;
+			TargetData->bNeedsRealloc = true;
+				
+			TargetData->NumParticles = 0;
+			TargetData->AllocatedNumSlots = 0;
+			TargetData->AllocatedNumCells = 0;
+	
+
+		});
+	}
+
+	return false;
+}
+
+void UNiagaraDataInterfaceParticleNeighborQuery::DestroyPerInstanceData(void* PerInstanceData, FNiagaraSystemInstance* SystemInstance)
+{
+	FNDIParticleNeighborQueryInstanceData_GT* InstanceData = static_cast<FNDIParticleNeighborQueryInstanceData_GT*>(PerInstanceData);
+	InstanceData->~FNDIParticleNeighborQueryInstanceData_GT();
+
+	FNiagaraDataInterfaceProxyParticleNeighborQuery* ThisProxy = GetProxyAs<FNiagaraDataInterfaceProxyParticleNeighborQuery>();
+	if (!ThisProxy)
+		return;
+
+	ENQUEUE_RENDER_COMMAND(FNiagaraDIDestroyInstanceData)(
+		[ThisProxy, InstanceID = SystemInstance->GetId()](FRHICommandListImmediate& CmdList)
+	{
+		ThisProxy->SystemInstancesToProxyData_RT.Remove(InstanceID);
+	}
+	);
+}
+
+//////////////////////////////////////////////////////////////////////////
+// Proxy methods
+
+void FNiagaraDataInterfaceProxyParticleNeighborQuery::ResetData(const FNDIGpuComputeResetContext& Context)
+{
+	FNDIParticleNeighborQueryInstanceData_RT* ProxyData = SystemInstancesToProxyData_RT.Find(Context.GetSystemInstanceID());
+	if (!ProxyData)
+	{
+		return;
+	}
+
+	FRDGBuilder& GraphBuilder = Context.GetGraphBuilder();
+	if (ProxyData->CellIdBuffer.IsValid())
+	{
+		AddClearUAVPass(GraphBuilder, ProxyData->CellIdBuffer.GetOrCreateUAV(GraphBuilder), 0xFFFFFFFF);  // INVALID_CELL sentinel
+	}
+	if (ProxyData->CellCountBuffer.IsValid())
+	{
+		AddClearUAVPass(GraphBuilder, ProxyData->CellCountBuffer.GetOrCreateUAV(GraphBuilder), 0);
+	}
+}
+
+void FNiagaraDataInterfaceProxyParticleNeighborQuery::PreStage(const FNDIGpuComputePreStageContext& Context)
+{
+	FRDGBuilder& GraphBuilder = Context.GetGraphBuilder();
+	FNDIParticleNeighborQueryInstanceData_RT& ProxyData = SystemInstancesToProxyData_RT.FindChecked(Context.GetSystemInstanceID());
+
+	if (Context.IsOutputStage())
+	{
+		// Use SourceNumInstances — this is the actual dispatch thread count (pre-compaction).
+		// DestinationNumInstances may be smaller if this stage kills particles, causing buffer overflow.
+		const uint32 NewNumParticles = Context.GetSimStageData().SourceNumInstances;
+		const uint64 NeededSlots = (uint64) NewNumParticles * ProxyData.MaxCellsPerParticle;
+
+		ProxyData.NumParticles = NewNumParticles;
+
+		// Realloc if grid changed or slot buffers need to grow.
+		if (ProxyData.bNeedsRealloc || NeededSlots > ProxyData.AllocatedNumSlots)
+		{
+			ProxyData.ResizeBuffers(GraphBuilder);
+		}
+
+		if (ProxyData.ClearBeforeNonIterationStage)
+		{
+			RDG_RHI_EVENT_SCOPE(GraphBuilder, NiagaraParticleNeighborQueryClear);
+			
+			if (ProxyData.CellIdBuffer.IsValid())
+			{
+				AddClearUAVPass(GraphBuilder, ProxyData.CellIdBuffer.GetOrCreateUAV(GraphBuilder), 0xFFFFFFFF);  // INVALID_CELL sentinel
+			}
+			if (ProxyData.CellCountBuffer.IsValid())
+			{
+				AddClearUAVPass(GraphBuilder, ProxyData.CellCountBuffer.GetOrCreateUAV(GraphBuilder), 0);
+			}
+		}
+	}
+	else if (ProxyData.bNeedsRealloc)
+	{
+		ProxyData.ResizeBuffers(GraphBuilder);
+	}
+}
+
+void FNiagaraDataInterfaceProxyParticleNeighborQuery::PostStage(const FNDIGpuComputePostStageContext& Context)
+{
+	if (!Context.IsOutputStage())
+	{
+		return;
+	}
+
+	FNDIParticleNeighborQueryInstanceData_RT& ProxyData = SystemInstancesToProxyData_RT.FindChecked(Context.GetSystemInstanceID());
+	if (!ProxyData.CellIdBuffer.IsValid() || !ProxyData.CellCountBuffer.IsValid())
+	{
+		return;
+	}
+
+	FRDGBuilder& GraphBuilder = Context.GetGraphBuilder();
+
+	const uint32 NumTotalCells = ProxyData.NumCells.X * ProxyData.NumCells.Y * ProxyData.NumCells.Z;
+	const uint32 NumSlots = (uint32)FMath::Min<uint64>((uint64)ProxyData.NumParticles * ProxyData.MaxCellsPerParticle, MAX_uint32);
+
+	FRDGBufferSRVRef CellIdSRV = ProxyData.CellIdBuffer.GetOrCreateSRV(GraphBuilder);
+	FRDGBufferUAVRef CellCountUAV = ProxyData.CellCountBuffer.GetOrCreateUAV(GraphBuilder);
+
+	RDG_EVENT_SCOPE(GraphBuilder, "NiagaraNeighborSort Particles(%u) Cells(%dx%dx%d) Slots(%u)%s",
+		ProxyData.NumParticles, ProxyData.NumCells.X, ProxyData.NumCells.Y, ProxyData.NumCells.Z, NumSlots,
+		ProxyData.bCountOnly ? TEXT(" CountOnly") : TEXT(""));
+
+	// Pass 2: Histogram — count particles per cell
+	NiagaraParticleNeighborQuerySort::Histogram(GraphBuilder, CellIdSRV, CellCountUAV, NumSlots);
+
+	if (!ProxyData.bCountOnly)
+	{
+		FRDGBufferUAVRef CellOffsetUAV = ProxyData.CellOffsetBuffer.GetOrCreateUAV(GraphBuilder);
+		FRDGBufferUAVRef ParticleListUAV = ProxyData.ParticleListBuffer.GetOrCreateUAV(GraphBuilder);
+
+		// Pass 3: Prefix Sum — exclusive scan over cell counts
+		NiagaraParticleNeighborQuerySort::PrefixSum(GraphBuilder, CellCountUAV, CellOffsetUAV, NumTotalCells);
+
+		// Pass 4: Scatter — write sorted particle indices into ParticleListBuffer.
+		// NOTE: Scatter atomically increments CellOffsetBuffer, so after this pass
+		// CellOffset[i] = original_prefix_sum[i] + CellCount[i]. The read functions
+		// in the .ush account for this by subtracting CellCount.
+		NiagaraParticleNeighborQuerySort::Scatter(GraphBuilder, CellIdSRV, CellOffsetUAV, ParticleListUAV, NumSlots, ProxyData.MaxCellsPerParticle);
+	}
+}
+
+void FNiagaraDataInterfaceProxyParticleNeighborQuery::PostSimulate(const FNDIGpuComputePostSimulateContext& Context)
+{
+	if (Context.IsFinalPostSimulate())
+	{
+		FNDIParticleNeighborQueryInstanceData_RT& ProxyData = SystemInstancesToProxyData_RT.FindChecked(Context.GetSystemInstanceID());
+		if (ProxyData.CellIdBuffer.IsValid())
+		{
+			ProxyData.CellIdBuffer.EndGraphUsage();
+		}
+		if (ProxyData.CellCountBuffer.IsValid())
+		{
+			ProxyData.CellCountBuffer.EndGraphUsage();
+		}
+		if (!ProxyData.bCountOnly)
+		{
+			if (ProxyData.CellOffsetBuffer.IsValid())
+			{
+				ProxyData.CellOffsetBuffer.EndGraphUsage();
+			}
+			if (ProxyData.ParticleListBuffer.IsValid())
+			{
+				ProxyData.ParticleListBuffer.EndGraphUsage();
+			}
+		}
+	}
+}
+
+void FNiagaraDataInterfaceProxyParticleNeighborQuery::GetDispatchArgs(const FNDIGpuComputeDispatchArgsGenContext& Context)
+{
+	if (const FNDIParticleNeighborQueryInstanceData_RT* TargetData = SystemInstancesToProxyData_RT.Find(Context.GetSystemInstanceID()))
+	{
+		Context.SetDirect(TargetData->NumCells);
+	}
+}
+
+bool UNiagaraDataInterfaceParticleNeighborQuery::CopyToInternal(UNiagaraDataInterface* Destination) const
+{
+	if (!Super::CopyToInternal(Destination))
+	{
+		return false;
+	}
+
+	UNiagaraDataInterfaceParticleNeighborQuery* OtherTyped = CastChecked<UNiagaraDataInterfaceParticleNeighborQuery>(Destination);
+
+	OtherTyped->MaxCellsPerParticle = MaxCellsPerParticle;
+	OtherTyped->bCountOnly = bCountOnly;
+
+	return true;
+}
+
+#undef LOCTEXT_NAMESPACE
diff --git a/Engine/Plugins/FX/Niagara/Source/Niagara/Private/NiagaraParticleNeighborQuerySort.cpp b/Engine/Plugins/FX/Niagara/Source/Niagara/Private/NiagaraParticleNeighborQuerySort.cpp
new file mode 100644
--- /dev/null
+++ b/Engine/Plugins/FX/Niagara/Source/Niagara/Private/NiagaraParticleNeighborQuerySort.cpp
@@ -0,0 +1,298 @@
+// Copyright Epic Games, Inc. All Rights Reserved.
+
+#include "NiagaraParticleNeighborQuerySort.h"
+
+#include "GlobalShader.h"
+#include "RenderGraphBuilder.h"
+#include "RenderGraphUtils.h"
+#include "ShaderParameterStruct.h"
+
+//////////////////////////////////////////////////////////////////////////
+// Pass 2: Histogram
+
+class FNiagaraParticleNeighborQueryHistogramCS : public FGlobalShader
+{
+	DECLARE_GLOBAL_SHADER(FNiagaraParticleNeighborQueryHistogramCS);
+	SHADER_USE_PARAMETER_STRUCT(FNiagaraParticleNeighborQueryHistogramCS, FGlobalShader);
+
+	static constexpr uint32 ThreadGroupSize = 64;
+
+	static void ModifyCompilationEnvironment(const FGlobalShaderPermutationParameters& Parameters, FShaderCompilerEnvironment& OutEnvironment)
+	{
+		FGlobalShader::ModifyCompilationEnvironment(Parameters, OutEnvironment);
+		OutEnvironment.SetDefine(TEXT("PARTICLE_NEIGHBOR_QUERY_HISTOGRAM_CS"), 1);
+		OutEnvironment.SetDefine(TEXT("THREADGROUP_SIZE"), ThreadGroupSize);
+	}
+
+	BEGIN_SHADER_PARAMETER_STRUCT(FParameters, )
+		SHADER_PARAMETER_RDG_BUFFER_SRV(Buffer<int>,	CellIdBuffer)
+		SHADER_PARAMETER_RDG_BUFFER_UAV(RWBuffer<int>,	CellCountBuffer)
+		SHADER_PARAMETER(uint32,						NumSlots)
+		SHADER_PARAMETER(uint32,						GroupsX)
+	END_SHADER_PARAMETER_STRUCT()
+};
+
+IMPLEMENT_GLOBAL_SHADER(FNiagaraParticleNeighborQueryHistogramCS, "/Plugin/FX/Niagara/Private/NiagaraParticleNeighborQuerySort.usf", "HistogramCS", SF_Compute);
+
+//////////////////////////////////////////////////////////////////////////
+// Pass 3: Parallel Prefix Sum (recursive reduce-then-scan)
+//
+// Two shader sub-passes per recursion level:
+//   3a. LocalScan  — each group scans a tile, writes tile sum to BlockSumsBuffer
+//   3b. Propagate  — each group adds its scanned block prefix to local results
+//
+// Block sums are scanned by recursively applying the same algorithm.
+
+namespace PrefixSumConstants
+{
+	static constexpr uint32 ThreadGroupSize = 128;
+	static constexpr uint32 ElementsPerThread = 4;
+	static constexpr uint32 ElementsPerGroup = ThreadGroupSize * ElementsPerThread; // 512
+}
+
+class FNiagaraParticleNeighborQueryPrefixSumLocalScanCS : public FGlobalShader
+{
+	DECLARE_GLOBAL_SHADER(FNiagaraParticleNeighborQueryPrefixSumLocalScanCS);
+	SHADER_USE_PARAMETER_STRUCT(FNiagaraParticleNeighborQueryPrefixSumLocalScanCS, FGlobalShader);
+
+	static void ModifyCompilationEnvironment(const FGlobalShaderPermutationParameters& Parameters, FShaderCompilerEnvironment& OutEnvironment)
+	{
+		FGlobalShader::ModifyCompilationEnvironment(Parameters, OutEnvironment);
+		OutEnvironment.SetDefine(TEXT("PARTICLE_NEIGHBOR_QUERY_PREFIX_SUM_LOCAL_SCAN_CS"), 1);
+		OutEnvironment.SetDefine(TEXT("THREADGROUP_SIZE"), PrefixSumConstants::ThreadGroupSize);
+		OutEnvironment.SetDefine(TEXT("ELEMENTS_PER_THREAD"), PrefixSumConstants::ElementsPerThread);
+		OutEnvironment.SetDefine(TEXT("ELEMENTS_PER_GROUP"), PrefixSumConstants::ElementsPerGroup);
+	}
+
+	BEGIN_SHADER_PARAMETER_STRUCT(FParameters, )
+		SHADER_PARAMETER_RDG_BUFFER_SRV(Buffer<int>,	InputBuffer)
+		SHADER_PARAMETER_RDG_BUFFER_UAV(RWBuffer<int>,	OutputBuffer)
+		SHADER_PARAMETER_RDG_BUFFER_UAV(RWBuffer<int>,	BlockSumsBuffer)
+		SHADER_PARAMETER(uint32,						NumElements)
+		SHADER_PARAMETER(uint32,						GroupsX)
+	END_SHADER_PARAMETER_STRUCT()
+};
+
+IMPLEMENT_GLOBAL_SHADER(FNiagaraParticleNeighborQueryPrefixSumLocalScanCS, "/Plugin/FX/Niagara/Private/NiagaraParticleNeighborQuerySort.usf", "PrefixSumLocalScanCS", SF_Compute);
+
+class FNiagaraParticleNeighborQueryPrefixSumPropagateCS : public FGlobalShader
+{
+	DECLARE_GLOBAL_SHADER(FNiagaraParticleNeighborQueryPrefixSumPropagateCS);
+	SHADER_USE_PARAMETER_STRUCT(FNiagaraParticleNeighborQueryPrefixSumPropagateCS, FGlobalShader);
+
+	static void ModifyCompilationEnvironment(const FGlobalShaderPermutationParameters& Parameters, FShaderCompilerEnvironment& OutEnvironment)
+	{
+		FGlobalShader::ModifyCompilationEnvironment(Parameters, OutEnvironment);
+		OutEnvironment.SetDefine(TEXT("PARTICLE_NEIGHBOR_QUERY_PREFIX_SUM_PROPAGATE_CS"), 1);
+		OutEnvironment.SetDefine(TEXT("THREADGROUP_SIZE"), PrefixSumConstants::ThreadGroupSize);
+		OutEnvironment.SetDefine(TEXT("ELEMENTS_PER_THREAD"), PrefixSumConstants::ElementsPerThread);
+		OutEnvironment.SetDefine(TEXT("ELEMENTS_PER_GROUP"), PrefixSumConstants::ElementsPerGroup);
+	}
+
+	BEGIN_SHADER_PARAMETER_STRUCT(FParameters, )
+		SHADER_PARAMETER_RDG_BUFFER_SRV(Buffer<int>,	BlockSumsBuffer)
+		SHADER_PARAMETER_RDG_BUFFER_UAV(RWBuffer<int>,	OutputBuffer)
+		SHADER_PARAMETER(uint32,						NumElements)
+		SHADER_PARAMETER(uint32,						GroupsX)
+	END_SHADER_PARAMETER_STRUCT()
+};
+
+IMPLEMENT_GLOBAL_SHADER(FNiagaraParticleNeighborQueryPrefixSumPropagateCS, "/Plugin/FX/Niagara/Private/NiagaraParticleNeighborQuerySort.usf", "PrefixSumPropagateCS", SF_Compute);
+
+//////////////////////////////////////////////////////////////////////////
+// Pass 4: Scatter
+
+class FNiagaraParticleNeighborQueryScatterCS : public FGlobalShader
+{
+	DECLARE_GLOBAL_SHADER(FNiagaraParticleNeighborQueryScatterCS);
+	SHADER_USE_PARAMETER_STRUCT(FNiagaraParticleNeighborQueryScatterCS, FGlobalShader);
+
+	static constexpr uint32 ThreadGroupSize = 64;
+
+	static void ModifyCompilationEnvironment(const FGlobalShaderPermutationParameters& Parameters, FShaderCompilerEnvironment& OutEnvironment)
+	{
+		FGlobalShader::ModifyCompilationEnvironment(Parameters, OutEnvironment);
+		OutEnvironment.SetDefine(TEXT("PARTICLE_NEIGHBOR_QUERY_SCATTER_CS"), 1);
+		OutEnvironment.SetDefine(TEXT("THREADGROUP_SIZE"), ThreadGroupSize);
+	}
+
+	BEGIN_SHADER_PARAMETER_STRUCT(FParameters, )
+		SHADER_PARAMETER_RDG_BUFFER_SRV(Buffer<int>,	CellIdBuffer)
+		SHADER_PARAMETER_RDG_BUFFER_UAV(RWBuffer<int>,	CellOffsetBuffer)
+		SHADER_PARAMETER_RDG_BUFFER_UAV(RWBuffer<int>,	ParticleListBuffer)
+		SHADER_PARAMETER(uint32,						NumSlots)
+		SHADER_PARAMETER(uint32,						MaxCellsPerParticle)
+		SHADER_PARAMETER(uint32,						GroupsX)
+	END_SHADER_PARAMETER_STRUCT()
+};
+
+IMPLEMENT_GLOBAL_SHADER(FNiagaraParticleNeighborQueryScatterCS, "/Plugin/FX/Niagara/Private/NiagaraParticleNeighborQuerySort.usf", "ScatterCS", SF_Compute);
+
+//////////////////////////////////////////////////////////////////////////
+// Dispatch functions
+
+// Split a 1D group count into a 2D dispatch to stay within MaxDispatchThreadGroupsPerDimension (65535).
+static FIntVector GetDispatchSize(uint32 NumGroups)
+{
+	if (NumGroups <= 65535)
+	{
+		return FIntVector(NumGroups, 1, 1);
+	}
+	const uint32 GroupsX = FMath::DivideAndRoundUp(NumGroups, (uint32)FMath::CeilToInt32(FMath::Sqrt((float)NumGroups)));
+	const uint32 GroupsY = FMath::DivideAndRoundUp(NumGroups, GroupsX);
+	return FIntVector(GroupsX, GroupsY, 1);
+}
+
+void NiagaraParticleNeighborQuerySort::Histogram(FRDGBuilder& GraphBuilder, FRDGBufferSRVRef CellIdSRV, FRDGBufferUAVRef CellCountUAV, uint32 NumSlots)
+{
+	if (NumSlots == 0)
+	{
+		return;
+	}
+
+	TShaderMapRef<FNiagaraParticleNeighborQueryHistogramCS> ComputeShader(GetGlobalShaderMap(GMaxRHIFeatureLevel));
+	const uint32 NumThreadGroups = FMath::DivideAndRoundUp(NumSlots, FNiagaraParticleNeighborQueryHistogramCS::ThreadGroupSize);
+	const FIntVector DispatchSize = GetDispatchSize(NumThreadGroups);
+
+	FNiagaraParticleNeighborQueryHistogramCS::FParameters* Parameters = GraphBuilder.AllocParameters<FNiagaraParticleNeighborQueryHistogramCS::FParameters>();
+	Parameters->CellIdBuffer = CellIdSRV;
+	Parameters->CellCountBuffer = CellCountUAV;
+	Parameters->NumSlots = NumSlots;
+	Parameters->GroupsX = DispatchSize.X;
+
+	FComputeShaderUtils::AddPass(
+		GraphBuilder,
+		RDG_EVENT_NAME("NiagaraNeighborSort::Histogram NumSlots(%u) DispatchCount(%ux%ux1) NumThreads(%ux1x1)", NumSlots, DispatchSize.X, DispatchSize.Y, FNiagaraParticleNeighborQueryHistogramCS::ThreadGroupSize),
+		ERDGPassFlags::Compute,
+		ComputeShader,
+		Parameters,
+		DispatchSize
+	);
+}
+
+// Internal recursive prefix sum implementation.
+// InputBuffer is read as SRV, results written to OutputUAV.
+// Recurses on block sums when NumElements > ElementsPerGroup.
+static void PrefixSumRecursive(FRDGBuilder& GraphBuilder, FRDGBufferRef InputBuffer, FRDGBufferUAVRef OutputUAV, uint32 NumElements)
+{
+	if (NumElements == 0)
+	{
+		return;
+	}
+
+	using namespace PrefixSumConstants;
+
+	const uint32 NumGroups = FMath::DivideAndRoundUp(NumElements, ElementsPerGroup);
+
+	// Single tile — scan directly, no block sums needed
+	if (NumGroups == 1)
+	{
+		TShaderMapRef<FNiagaraParticleNeighborQueryPrefixSumLocalScanCS> ComputeShader(GetGlobalShaderMap(GMaxRHIFeatureLevel));
+
+		// Create a dummy 1-element block sums buffer (required by shader signature)
+		FRDGBufferRef DummyBlockSums = GraphBuilder.CreateBuffer(FRDGBufferDesc::CreateBufferDesc(sizeof(int32), 1), TEXT("NiagaraNeighborSort::DummyBlockSums"));
+
+		auto* Parameters = GraphBuilder.AllocParameters<FNiagaraParticleNeighborQueryPrefixSumLocalScanCS::FParameters>();
+		Parameters->InputBuffer = GraphBuilder.CreateSRV(InputBuffer, PF_R32_SINT);
+		Parameters->OutputBuffer = OutputUAV;
+		Parameters->BlockSumsBuffer = GraphBuilder.CreateUAV(DummyBlockSums, PF_R32_SINT);
+		Parameters->NumElements = NumElements;
+		Parameters->GroupsX = 1;
+
+		FComputeShaderUtils::AddPass(
+			GraphBuilder,
+			RDG_EVENT_NAME("NiagaraNeighborSort::PrefixSum NumElements(%u) DispatchCount(1x1x1) NumThreads(%ux1x1)", NumElements, ThreadGroupSize),
+			ERDGPassFlags::Compute,
+			ComputeShader,
+			Parameters,
+			FIntVector(1, 1, 1)
+		);
+		return;
+	}
+
+	// Multi-tile: reduce-then-scan with recursive block sum scanning
+	const FIntVector LocalScanDispatch = GetDispatchSize(NumGroups);
+
+	FRDGBufferRef BlockSumsBuffer = GraphBuilder.CreateBuffer(FRDGBufferDesc::CreateBufferDesc(sizeof(int32), NumGroups), TEXT("NiagaraNeighborSort::BlockSums"));
+
+	// LocalScan — each group scans its tile, writes tile total to BlockSumsBuffer
+	{
+		TShaderMapRef<FNiagaraParticleNeighborQueryPrefixSumLocalScanCS> ComputeShader(GetGlobalShaderMap(GMaxRHIFeatureLevel));
+
+		auto* Parameters = GraphBuilder.AllocParameters<FNiagaraParticleNeighborQueryPrefixSumLocalScanCS::FParameters>();
+		Parameters->InputBuffer = GraphBuilder.CreateSRV(InputBuffer, PF_R32_SINT);
+		Parameters->OutputBuffer = OutputUAV;
+		Parameters->BlockSumsBuffer = GraphBuilder.CreateUAV(BlockSumsBuffer, PF_R32_SINT);
+		Parameters->NumElements = NumElements;
+		Parameters->GroupsX = LocalScanDispatch.X;
+
+		FComputeShaderUtils::AddPass(
+			GraphBuilder,
+			RDG_EVENT_NAME("NiagaraNeighborSort::PrefixSum::LocalScan NumElements(%u) DispatchCount(%ux%ux1) NumThreads(%ux1x1)", NumElements, LocalScanDispatch.X, LocalScanDispatch.Y, ThreadGroupSize),
+			ERDGPassFlags::Compute,
+			ComputeShader,
+			Parameters,
+			LocalScanDispatch
+		);
+	}
+
+	// Recursively scan the block sums
+	FRDGBufferRef ScannedBlockSums = GraphBuilder.CreateBuffer(FRDGBufferDesc::CreateBufferDesc(sizeof(int32), NumGroups), TEXT("NiagaraNeighborSort::ScannedBlockSums"));
+	FRDGBufferUAVRef ScannedBlockSumsUAV = GraphBuilder.CreateUAV(ScannedBlockSums, PF_R32_SINT);
+	PrefixSumRecursive(GraphBuilder, BlockSumsBuffer, ScannedBlockSumsUAV, NumGroups);
+
+	// Propagate — each group adds its scanned block prefix to local results
+	{
+		TShaderMapRef<FNiagaraParticleNeighborQueryPrefixSumPropagateCS> ComputeShader(GetGlobalShaderMap(GMaxRHIFeatureLevel));
+		const FIntVector PropagateDispatch = GetDispatchSize(NumGroups);
+
+		auto* Parameters = GraphBuilder.AllocParameters<FNiagaraParticleNeighborQueryPrefixSumPropagateCS::FParameters>();
+		Parameters->BlockSumsBuffer = GraphBuilder.CreateSRV(ScannedBlockSums, PF_R32_SINT);
+		Parameters->OutputBuffer = OutputUAV;
+		Parameters->NumElements = NumElements;
+		Parameters->GroupsX = PropagateDispatch.X;
+
+		FComputeShaderUtils::AddPass(
+			GraphBuilder,
+			RDG_EVENT_NAME("NiagaraNeighborSort::PrefixSum::Propagate NumElements(%u) DispatchCount(%ux%ux1) NumThreads(%ux1x1)", NumElements, PropagateDispatch.X, PropagateDispatch.Y, ThreadGroupSize),
+			ERDGPassFlags::Compute,
+			ComputeShader,
+			Parameters,
+			PropagateDispatch
+		);
+	}
+}
+
+void NiagaraParticleNeighborQuerySort::PrefixSum(FRDGBuilder& GraphBuilder, FRDGBufferUAVRef CellCountUAV, FRDGBufferUAVRef CellOffsetUAV, uint32 NumCells)
+{
+	PrefixSumRecursive(GraphBuilder, CellCountUAV->GetParent(), CellOffsetUAV, NumCells);
+}
+
+void NiagaraParticleNeighborQuerySort::Scatter(FRDGBuilder& GraphBuilder, FRDGBufferSRVRef CellIdSRV, FRDGBufferUAVRef CellOffsetUAV, FRDGBufferUAVRef ParticleListUAV, uint32 NumSlots, uint32 MaxCellsPerParticle)
+{
+	if (NumSlots == 0)
+	{
+		return;
+	}
+
+	TShaderMapRef<FNiagaraParticleNeighborQueryScatterCS> ComputeShader(GetGlobalShaderMap(GMaxRHIFeatureLevel));
+	const uint32 NumThreadGroups = FMath::DivideAndRoundUp(NumSlots, FNiagaraParticleNeighborQueryScatterCS::ThreadGroupSize);
+	const FIntVector DispatchSize = GetDispatchSize(NumThreadGroups);
+
+	FNiagaraParticleNeighborQueryScatterCS::FParameters* Parameters = GraphBuilder.AllocParameters<FNiagaraParticleNeighborQueryScatterCS::FParameters>();
+	Parameters->CellIdBuffer = CellIdSRV;
+	Parameters->CellOffsetBuffer = CellOffsetUAV;
+	Parameters->ParticleListBuffer = ParticleListUAV;
+	Parameters->NumSlots = NumSlots;
+	Parameters->MaxCellsPerParticle = MaxCellsPerParticle;
+	Parameters->GroupsX = DispatchSize.X;
+
+	FComputeShaderUtils::AddPass(
+		GraphBuilder,
+		RDG_EVENT_NAME("NiagaraNeighborSort::Scatter NumSlots(%u) MaxCellsPerParticle(%u) DispatchCount(%ux%ux1) NumThreads(%ux1x1)", NumSlots, MaxCellsPerParticle, DispatchSize.X, DispatchSize.Y, FNiagaraParticleNeighborQueryScatterCS::ThreadGroupSize),
+		ERDGPassFlags::Compute,
+		ComputeShader,
+		Parameters,
+		DispatchSize
+	);
+}
diff --git a/Engine/Plugins/FX/Niagara/Source/Niagara/Private/NiagaraParticleNeighborQuerySort.h b/Engine/Plugins/FX/Niagara/Source/Niagara/Private/NiagaraParticleNeighborQuerySort.h
new file mode 100644
--- /dev/null
+++ b/Engine/Plugins/FX/Niagara/Source/Niagara/Private/NiagaraParticleNeighborQuerySort.h
@@ -0,0 +1,21 @@
+// Copyright Epic Games, Inc. All Rights Reserved.
+
+#pragma once
+
+#include "RenderGraphFwd.h"
+
+class FRDGBuilder;
+
+// Counting sort dispatch functions for NiagaraDataInterfaceParticleNeighborQuery.
+// Called from PostStage after the output stage has written cell assignments.
+namespace NiagaraParticleNeighborQuerySort
+{
+	// Pass 2: Build histogram — count particles per cell
+	void Histogram(FRDGBuilder& GraphBuilder, FRDGBufferSRVRef CellIdSRV, FRDGBufferUAVRef CellCountUAV, uint32 NumSlots);
+
+	// Pass 3: Exclusive prefix sum over cell counts
+	void PrefixSum(FRDGBuilder& GraphBuilder, FRDGBufferUAVRef CellCountUAV, FRDGBufferUAVRef CellOffsetUAV, uint32 NumCells);
+
+	// Pass 4: Scatter sorted particle indices
+	void Scatter(FRDGBuilder& GraphBuilder, FRDGBufferSRVRef CellIdSRV, FRDGBufferUAVRef CellOffsetUAV, FRDGBufferUAVRef ParticleListUAV, uint32 NumSlots, uint32 MaxCellsPerParticle);
+}
