Fixed indeterministic asset registry queries during cook in FNiagaraEditorOnlyDataUtilities::MoveUserAssetTagsToMetaData

#Synced-CL 52466916

diff --git a/Engine/Plugins/FX/Niagara/Source/NiagaraEditor/Private/NiagaraEditorModule.cpp b/Engine/Plugins/FX/Niagara/Source/NiagaraEditor/Private/NiagaraEditorModule.cpp
--- a/Engine/Plugins/FX/Niagara/Source/NiagaraEditor/Private/NiagaraEditorModule.cpp
+++ b/Engine/Plugins/FX/Niagara/Source/NiagaraEditor/Private/NiagaraEditorModule.cpp
@@ -420,6 +420,8 @@ PRAGMA_DISABLE_DEPRECATION_WARNINGS
 			}
 PRAGMA_ENABLE_DEPRECATION_WARNINGS
 
+			Filter.bIncludeOnlyOnDiskAssets = IsRunningCookCommandlet();	// When cooking, only enumerate on-disk assets- avoids indeterministic results
+
 			TArray<FAssetData> AssetTagDefinitionAssets;
 			AssetRegistryModule.Get().GetAssets(Filter, AssetTagDefinitionAssets);
 
