PR #14619: Fix increasing NDI memory for each run

#Synced-CL 52445503

diff --git a/Engine/Plugins/FX/Niagara/Source/Niagara/Private/NiagaraDataInterfaceLandscape.cpp b/Engine/Plugins/FX/Niagara/Source/Niagara/Private/NiagaraDataInterfaceLandscape.cpp
--- a/Engine/Plugins/FX/Niagara/Source/Niagara/Private/NiagaraDataInterfaceLandscape.cpp
+++ b/Engine/Plugins/FX/Niagara/Source/Niagara/Private/NiagaraDataInterfaceLandscape.cpp
@@ -767,7 +767,7 @@ void FLandscapeTextureResource::ReleaseRHI()
 	PhysMatTexture.Release();
 
 #if STATS
-	INC_MEMORY_STAT_BY(STAT_NiagaraGPUDataInterfaceMemory, GpuMemoryUsage);
+	DEC_MEMORY_STAT_BY(STAT_NiagaraGPUDataInterfaceMemory, GpuMemoryUsage);
 	GpuMemoryUsage = 0;
 #endif
 }
