- Replace GetShaderFileHash with visitor UpdateShaderFile

#Synced-CL 52445503

diff --git a/Engine/Plugins/FX/Niagara/Source/Niagara/Private/DataInterface/NiagaraDataInterfaceDataChannelRead.cpp b/Engine/Plugins/FX/Niagara/Source/Niagara/Private/DataInterface/NiagaraDataInterfaceDataChannelRead.cpp
--- a/Engine/Plugins/FX/Niagara/Source/Niagara/Private/DataInterface/NiagaraDataInterfaceDataChannelRead.cpp
+++ b/Engine/Plugins/FX/Niagara/Source/Niagara/Private/DataInterface/NiagaraDataInterfaceDataChannelRead.cpp
@@ -2079,11 +2079,11 @@ void UNiagaraDataInterfaceDataChannelRead::ScaleSpawnCount(FVectorVMExternalFunc
 bool UNiagaraDataInterfaceDataChannelRead::AppendCompileHash(FNiagaraCompileHashVisitor* InVisitor) const
 {
 	bool bSuccess = Super::AppendCompileHash(InVisitor);
-	bSuccess &= InVisitor->UpdateString(TEXT("UNiagaraDataInterfaceDataChannelCommon"), GetShaderFileHash(NDIDataChannelReadLocal::CommonShaderFile, EShaderPlatform::SP_PCD3D_SM5).ToString());
-	bSuccess &= InVisitor->UpdateString(TEXT("UNiagaraDataInterfaceDataChannelTemplateCommon"), GetShaderFileHash(NDIDataChannelReadLocal::TemplateShaderFile_Common, EShaderPlatform::SP_PCD3D_SM5).ToString());
-	bSuccess &= InVisitor->UpdateString(TEXT("UNiagaraDataInterfaceDataChannelRead_Common"), GetShaderFileHash(NDIDataChannelReadLocal::TemplateShaderFile_ReadCommon, EShaderPlatform::SP_PCD3D_SM5).ToString());
-	bSuccess &= InVisitor->UpdateString(TEXT("UNiagaraDataInterfaceDataChannelRead_Read"), GetShaderFileHash(NDIDataChannelReadLocal::TemplateShaderFile_Read, EShaderPlatform::SP_PCD3D_SM5).ToString());
-	bSuccess &= InVisitor->UpdateString(TEXT("UNiagaraDataInterfaceDataChannelRead_Consume"), GetShaderFileHash(NDIDataChannelReadLocal::TemplateShaderFile_Consume, EShaderPlatform::SP_PCD3D_SM5).ToString());
+	bSuccess &= InVisitor->UpdateShaderFile(NDIDataChannelReadLocal::CommonShaderFile);
+	bSuccess &= InVisitor->UpdateShaderFile(NDIDataChannelReadLocal::TemplateShaderFile_Common);
+	bSuccess &= InVisitor->UpdateShaderFile(NDIDataChannelReadLocal::TemplateShaderFile_ReadCommon);
+	bSuccess &= InVisitor->UpdateShaderFile(NDIDataChannelReadLocal::TemplateShaderFile_Read);
+	bSuccess &= InVisitor->UpdateShaderFile(NDIDataChannelReadLocal::TemplateShaderFile_Consume);
 
 	bSuccess &= InVisitor->UpdateShaderParameters<NDIDataChannelReadLocal::FShaderParameters>();
 	return bSuccess;
diff --git a/Engine/Plugins/FX/Niagara/Source/Niagara/Private/DataInterface/NiagaraDataInterfaceDataChannelWrite.cpp b/Engine/Plugins/FX/Niagara/Source/Niagara/Private/DataInterface/NiagaraDataInterfaceDataChannelWrite.cpp
--- a/Engine/Plugins/FX/Niagara/Source/Niagara/Private/DataInterface/NiagaraDataInterfaceDataChannelWrite.cpp
+++ b/Engine/Plugins/FX/Niagara/Source/Niagara/Private/DataInterface/NiagaraDataInterfaceDataChannelWrite.cpp
@@ -1629,11 +1629,11 @@ void UNiagaraDataInterfaceDataChannelWrite::Append(FVectorVMExternalFunctionCont
 bool UNiagaraDataInterfaceDataChannelWrite::AppendCompileHash(FNiagaraCompileHashVisitor* InVisitor) const
 { 
 	bool bSuccess = Super::AppendCompileHash(InVisitor);
- 	bSuccess &= InVisitor->UpdateString(TEXT("UNiagaraDataInterfaceDataChannelCommon"), GetShaderFileHash(NDIDataChannelWriteLocal::CommonShaderFile, EShaderPlatform::SP_PCD3D_SM5).ToString());
-	bSuccess &= InVisitor->UpdateString(TEXT("UNiagaraDataInterfaceDataChannelTemplateCommon"), GetShaderFileHash(NDIDataChannelWriteLocal::TemplateShaderFile_Common, EShaderPlatform::SP_PCD3D_SM5).ToString());
-	bSuccess &= InVisitor->UpdateString(TEXT("UNiagaraDataInterfaceDataChannelWrite_WriteCommon"), GetShaderFileHash(NDIDataChannelWriteLocal::TemplateShaderFile_WriteCommon, EShaderPlatform::SP_PCD3D_SM5).ToString());
- 	//bSuccess &= InVisitor->UpdateString(TEXT("UNiagaraDataInterfaceDataChannelWrite_Write"), GetShaderFileHash(NDIDataChannelWriteLocal::TemplateShaderFile_Write, EShaderPlatform::SP_PCD3D_SM5).ToString());
- 	bSuccess &= InVisitor->UpdateString(TEXT("UNiagaraDataInterfaceDataChannelWrite_Append"), GetShaderFileHash(NDIDataChannelWriteLocal::TemplateShaderFile_Append, EShaderPlatform::SP_PCD3D_SM5).ToString());
+	bSuccess &= InVisitor->UpdateShaderFile(NDIDataChannelWriteLocal::CommonShaderFile);
+	bSuccess &= InVisitor->UpdateShaderFile(NDIDataChannelWriteLocal::TemplateShaderFile_Common);
+	bSuccess &= InVisitor->UpdateShaderFile(NDIDataChannelWriteLocal::TemplateShaderFile_WriteCommon);
+	//bSuccess &= InVisitor->UpdateShaderFile(NDIDataChannelWriteLocal::TemplateShaderFile_Write);
+	bSuccess &= InVisitor->UpdateShaderFile(NDIDataChannelWriteLocal::TemplateShaderFile_Append);
  
  	bSuccess &= InVisitor->UpdateShaderParameters<NDIDataChannelWriteLocal::FShaderParameters>();
  	return bSuccess;
diff --git a/Engine/Plugins/FX/Niagara/Source/Niagara/Private/NiagaraDataInterfaceRenderTargetVolume.cpp b/Engine/Plugins/FX/Niagara/Source/Niagara/Private/NiagaraDataInterfaceRenderTargetVolume.cpp
--- a/Engine/Plugins/FX/Niagara/Source/Niagara/Private/NiagaraDataInterfaceRenderTargetVolume.cpp
+++ b/Engine/Plugins/FX/Niagara/Source/Niagara/Private/NiagaraDataInterfaceRenderTargetVolume.cpp
@@ -402,9 +402,8 @@ bool UNiagaraDataInterfaceRenderTargetVolume::ShouldCompile(EShaderPlatform Shad
 bool UNiagaraDataInterfaceRenderTargetVolume::AppendCompileHash(FNiagaraCompileHashVisitor* InVisitor) const
 {
 	bool bSuccess = Super::AppendCompileHash(InVisitor);
-	FSHAHash Hash = GetShaderFileHash(NDIRenderTargetVolumeLocal::TemplateShaderFile, EShaderPlatform::SP_PCD3D_SM5);
-	InVisitor->UpdateString(TEXT("UNiagaraDataInterfaceRenderTargetVolumeTemplateHLSLSource"), Hash.ToString());
-	InVisitor->UpdateShaderParameters<NDIRenderTargetVolumeLocal::FShaderParameters>();
+	bSuccess &= InVisitor->UpdateShaderFile(NDIRenderTargetVolumeLocal::TemplateShaderFile);
+	bSuccess &= InVisitor->UpdateShaderParameters<NDIRenderTargetVolumeLocal::FShaderParameters>();
 	return bSuccess;
 }
 
