move shaders to fix initialization

#Synced-CL 52108856

diff --git a/Engine/Plugins/FX/Niagara/Source/Niagara/Private/NiagaraParticleNeighborQuerySort.cpp b/Engine/Plugins/FX/Niagara/Source/Niagara/Private/NiagaraParticleNeighborQuerySort.cpp
deleted file mode 100644
diff --git a/Engine/Plugins/FX/Niagara/Source/Niagara/Private/NiagaraParticleNeighborQuerySort.h b/Engine/Plugins/FX/Niagara/Source/Niagara/Private/NiagaraParticleNeighborQuerySort.h
deleted file mode 100644
diff --git a/Engine/Plugins/FX/Niagara/Source/NiagaraShader/Private/NiagaraParticleNeighborQuerySort.cpp b/Engine/Plugins/FX/Niagara/Source/NiagaraShader/Private/NiagaraParticleNeighborQuerySort.cpp
new file mode 100644
--- /dev/null
+++ b/Engine/Plugins/FX/Niagara/Source/NiagaraShader/Private/NiagaraParticleNeighborQuerySort.cpp
@@ -0,0 +1,298 @@
+// Copyright Epic Games, Inc. All Rights Reserved.
+
+#include "NiagaraParticleNeighborQuerySort.h"
+
+#include "GlobalShader.h"
+#include "RenderGraphBuilder.h"
+#include "RenderGraphUtils.h"
+#include "ShaderParameterStruct.h"
+
+//////////////////////////////////////////////////////////////////////////
+// Pass 2: Histogram
+
+class FNiagaraParticleNeighborQueryHistogramCS : public FGlobalShader
+{
+	DECLARE_GLOBAL_SHADER(FNiagaraParticleNeighborQueryHistogramCS);
+	SHADER_USE_PARAMETER_STRUCT(FNiagaraParticleNeighborQueryHistogramCS, FGlobalShader);
+
+	static constexpr uint32 ThreadGroupSize = 64;
+
+	static void ModifyCompilationEnvironment(const FGlobalShaderPermutationParameters& Parameters, FShaderCompilerEnvironment& OutEnvironment)
+	{
+		FGlobalShader::ModifyCompilationEnvironment(Parameters, OutEnvironment);
+		OutEnvironment.SetDefine(TEXT("PARTICLE_NEIGHBOR_QUERY_HISTOGRAM_CS"), 1);
+		OutEnvironment.SetDefine(TEXT("THREADGROUP_SIZE"), ThreadGroupSize);
+	}
+
+	BEGIN_SHADER_PARAMETER_STRUCT(FParameters, )
+		SHADER_PARAMETER_RDG_BUFFER_SRV(Buffer<int>,	CellIdBuffer)
+		SHADER_PARAMETER_RDG_BUFFER_UAV(RWBuffer<int>,	CellCountBuffer)
+		SHADER_PARAMETER(uint32,						NumSlots)
+		SHADER_PARAMETER(uint32,						GroupsX)
+	END_SHADER_PARAMETER_STRUCT()
+};
+
+IMPLEMENT_GLOBAL_SHADER(FNiagaraParticleNeighborQueryHistogramCS, "/Plugin/FX/Niagara/Private/NiagaraParticleNeighborQuerySort.usf", "HistogramCS", SF_Compute);
+
+//////////////////////////////////////////////////////////////////////////
+// Pass 3: Parallel Prefix Sum (recursive reduce-then-scan)
+//
+// Two shader sub-passes per recursion level:
+//   3a. LocalScan  — each group scans a tile, writes tile sum to BlockSumsBuffer
+//   3b. Propagate  — each group adds its scanned block prefix to local results
+//
+// Block sums are scanned by recursively applying the same algorithm.
+
+namespace PrefixSumConstants
+{
+	static constexpr uint32 ThreadGroupSize = 128;
+	static constexpr uint32 ElementsPerThread = 4;
+	static constexpr uint32 ElementsPerGroup = ThreadGroupSize * ElementsPerThread; // 512
+}
+
+class FNiagaraParticleNeighborQueryPrefixSumLocalScanCS : public FGlobalShader
+{
+	DECLARE_GLOBAL_SHADER(FNiagaraParticleNeighborQueryPrefixSumLocalScanCS);
+	SHADER_USE_PARAMETER_STRUCT(FNiagaraParticleNeighborQueryPrefixSumLocalScanCS, FGlobalShader);
+
+	static void ModifyCompilationEnvironment(const FGlobalShaderPermutationParameters& Parameters, FShaderCompilerEnvironment& OutEnvironment)
+	{
+		FGlobalShader::ModifyCompilationEnvironment(Parameters, OutEnvironment);
+		OutEnvironment.SetDefine(TEXT("PARTICLE_NEIGHBOR_QUERY_PREFIX_SUM_LOCAL_SCAN_CS"), 1);
+		OutEnvironment.SetDefine(TEXT("THREADGROUP_SIZE"), PrefixSumConstants::ThreadGroupSize);
+		OutEnvironment.SetDefine(TEXT("ELEMENTS_PER_THREAD"), PrefixSumConstants::ElementsPerThread);
+		OutEnvironment.SetDefine(TEXT("ELEMENTS_PER_GROUP"), PrefixSumConstants::ElementsPerGroup);
+	}
+
+	BEGIN_SHADER_PARAMETER_STRUCT(FParameters, )
+		SHADER_PARAMETER_RDG_BUFFER_SRV(Buffer<int>,	InputBuffer)
+		SHADER_PARAMETER_RDG_BUFFER_UAV(RWBuffer<int>,	OutputBuffer)
+		SHADER_PARAMETER_RDG_BUFFER_UAV(RWBuffer<int>,	BlockSumsBuffer)
+		SHADER_PARAMETER(uint32,						NumElements)
+		SHADER_PARAMETER(uint32,						GroupsX)
+	END_SHADER_PARAMETER_STRUCT()
+};
+
+IMPLEMENT_GLOBAL_SHADER(FNiagaraParticleNeighborQueryPrefixSumLocalScanCS, "/Plugin/FX/Niagara/Private/NiagaraParticleNeighborQuerySort.usf", "PrefixSumLocalScanCS", SF_Compute);
+
+class FNiagaraParticleNeighborQueryPrefixSumPropagateCS : public FGlobalShader
+{
+	DECLARE_GLOBAL_SHADER(FNiagaraParticleNeighborQueryPrefixSumPropagateCS);
+	SHADER_USE_PARAMETER_STRUCT(FNiagaraParticleNeighborQueryPrefixSumPropagateCS, FGlobalShader);
+
+	static void ModifyCompilationEnvironment(const FGlobalShaderPermutationParameters& Parameters, FShaderCompilerEnvironment& OutEnvironment)
+	{
+		FGlobalShader::ModifyCompilationEnvironment(Parameters, OutEnvironment);
+		OutEnvironment.SetDefine(TEXT("PARTICLE_NEIGHBOR_QUERY_PREFIX_SUM_PROPAGATE_CS"), 1);
+		OutEnvironment.SetDefine(TEXT("THREADGROUP_SIZE"), PrefixSumConstants::ThreadGroupSize);
+		OutEnvironment.SetDefine(TEXT("ELEMENTS_PER_THREAD"), PrefixSumConstants::ElementsPerThread);
+		OutEnvironment.SetDefine(TEXT("ELEMENTS_PER_GROUP"), PrefixSumConstants::ElementsPerGroup);
+	}
+
+	BEGIN_SHADER_PARAMETER_STRUCT(FParameters, )
+		SHADER_PARAMETER_RDG_BUFFER_SRV(Buffer<int>,	BlockSumsBuffer)
+		SHADER_PARAMETER_RDG_BUFFER_UAV(RWBuffer<int>,	OutputBuffer)
+		SHADER_PARAMETER(uint32,						NumElements)
+		SHADER_PARAMETER(uint32,						GroupsX)
+	END_SHADER_PARAMETER_STRUCT()
+};
+
+IMPLEMENT_GLOBAL_SHADER(FNiagaraParticleNeighborQueryPrefixSumPropagateCS, "/Plugin/FX/Niagara/Private/NiagaraParticleNeighborQuerySort.usf", "PrefixSumPropagateCS", SF_Compute);
+
+//////////////////////////////////////////////////////////////////////////
+// Pass 4: Scatter
+
+class FNiagaraParticleNeighborQueryScatterCS : public FGlobalShader
+{
+	DECLARE_GLOBAL_SHADER(FNiagaraParticleNeighborQueryScatterCS);
+	SHADER_USE_PARAMETER_STRUCT(FNiagaraParticleNeighborQueryScatterCS, FGlobalShader);
+
+	static constexpr uint32 ThreadGroupSize = 64;
+
+	static void ModifyCompilationEnvironment(const FGlobalShaderPermutationParameters& Parameters, FShaderCompilerEnvironment& OutEnvironment)
+	{
+		FGlobalShader::ModifyCompilationEnvironment(Parameters, OutEnvironment);
+		OutEnvironment.SetDefine(TEXT("PARTICLE_NEIGHBOR_QUERY_SCATTER_CS"), 1);
+		OutEnvironment.SetDefine(TEXT("THREADGROUP_SIZE"), ThreadGroupSize);
+	}
+
+	BEGIN_SHADER_PARAMETER_STRUCT(FParameters, )
+		SHADER_PARAMETER_RDG_BUFFER_SRV(Buffer<int>,	CellIdBuffer)
+		SHADER_PARAMETER_RDG_BUFFER_UAV(RWBuffer<int>,	CellOffsetBuffer)
+		SHADER_PARAMETER_RDG_BUFFER_UAV(RWBuffer<int>,	ParticleListBuffer)
+		SHADER_PARAMETER(uint32,						NumSlots)
+		SHADER_PARAMETER(uint32,						MaxCellsPerParticle)
+		SHADER_PARAMETER(uint32,						GroupsX)
+	END_SHADER_PARAMETER_STRUCT()
+};
+
+IMPLEMENT_GLOBAL_SHADER(FNiagaraParticleNeighborQueryScatterCS, "/Plugin/FX/Niagara/Private/NiagaraParticleNeighborQuerySort.usf", "ScatterCS", SF_Compute);
+
+//////////////////////////////////////////////////////////////////////////
+// Dispatch functions
+
+// Split a 1D group count into a 2D dispatch to stay within MaxDispatchThreadGroupsPerDimension (65535).
+static FIntVector GetDispatchSize(uint32 NumGroups)
+{
+	if (NumGroups <= 65535)
+	{
+		return FIntVector(NumGroups, 1, 1);
+	}
+	const uint32 GroupsX = FMath::DivideAndRoundUp(NumGroups, (uint32)FMath::CeilToInt32(FMath::Sqrt((float)NumGroups)));
+	const uint32 GroupsY = FMath::DivideAndRoundUp(NumGroups, GroupsX);
+	return FIntVector(GroupsX, GroupsY, 1);
+}
+
+void NiagaraParticleNeighborQuerySort::Histogram(FRDGBuilder& GraphBuilder, FRDGBufferSRVRef CellIdSRV, FRDGBufferUAVRef CellCountUAV, uint32 NumSlots)
+{
+	if (NumSlots == 0)
+	{
+		return;
+	}
+
+	TShaderMapRef<FNiagaraParticleNeighborQueryHistogramCS> ComputeShader(GetGlobalShaderMap(GMaxRHIFeatureLevel));
+	const uint32 NumThreadGroups = FMath::DivideAndRoundUp(NumSlots, FNiagaraParticleNeighborQueryHistogramCS::ThreadGroupSize);
+	const FIntVector DispatchSize = GetDispatchSize(NumThreadGroups);
+
+	FNiagaraParticleNeighborQueryHistogramCS::FParameters* Parameters = GraphBuilder.AllocParameters<FNiagaraParticleNeighborQueryHistogramCS::FParameters>();
+	Parameters->CellIdBuffer = CellIdSRV;
+	Parameters->CellCountBuffer = CellCountUAV;
+	Parameters->NumSlots = NumSlots;
+	Parameters->GroupsX = DispatchSize.X;
+
+	FComputeShaderUtils::AddPass(
+		GraphBuilder,
+		RDG_EVENT_NAME("NiagaraNeighborSort::Histogram NumSlots(%u) DispatchCount(%ux%ux1) NumThreads(%ux1x1)", NumSlots, DispatchSize.X, DispatchSize.Y, FNiagaraParticleNeighborQueryHistogramCS::ThreadGroupSize),
+		ERDGPassFlags::Compute,
+		ComputeShader,
+		Parameters,
+		DispatchSize
+	);
+}
+
+// Internal recursive prefix sum implementation.
+// InputBuffer is read as SRV, results written to OutputUAV.
+// Recurses on block sums when NumElements > ElementsPerGroup.
+static void PrefixSumRecursive(FRDGBuilder& GraphBuilder, FRDGBufferRef InputBuffer, FRDGBufferUAVRef OutputUAV, uint32 NumElements)
+{
+	if (NumElements == 0)
+	{
+		return;
+	}
+
+	using namespace PrefixSumConstants;
+
+	const uint32 NumGroups = FMath::DivideAndRoundUp(NumElements, ElementsPerGroup);
+
+	// Single tile — scan directly, no block sums needed
+	if (NumGroups == 1)
+	{
+		TShaderMapRef<FNiagaraParticleNeighborQueryPrefixSumLocalScanCS> ComputeShader(GetGlobalShaderMap(GMaxRHIFeatureLevel));
+
+		// Create a dummy 1-element block sums buffer (required by shader signature)
+		FRDGBufferRef DummyBlockSums = GraphBuilder.CreateBuffer(FRDGBufferDesc::CreateBufferDesc(sizeof(int32), 1), TEXT("NiagaraNeighborSort::DummyBlockSums"));
+
+		auto* Parameters = GraphBuilder.AllocParameters<FNiagaraParticleNeighborQueryPrefixSumLocalScanCS::FParameters>();
+		Parameters->InputBuffer = GraphBuilder.CreateSRV(InputBuffer, PF_R32_SINT);
+		Parameters->OutputBuffer = OutputUAV;
+		Parameters->BlockSumsBuffer = GraphBuilder.CreateUAV(DummyBlockSums, PF_R32_SINT);
+		Parameters->NumElements = NumElements;
+		Parameters->GroupsX = 1;
+
+		FComputeShaderUtils::AddPass(
+			GraphBuilder,
+			RDG_EVENT_NAME("NiagaraNeighborSort::PrefixSum NumElements(%u) DispatchCount(1x1x1) NumThreads(%ux1x1)", NumElements, ThreadGroupSize),
+			ERDGPassFlags::Compute,
+			ComputeShader,
+			Parameters,
+			FIntVector(1, 1, 1)
+		);
+		return;
+	}
+
+	// Multi-tile: reduce-then-scan with recursive block sum scanning
+	const FIntVector LocalScanDispatch = GetDispatchSize(NumGroups);
+
+	FRDGBufferRef BlockSumsBuffer = GraphBuilder.CreateBuffer(FRDGBufferDesc::CreateBufferDesc(sizeof(int32), NumGroups), TEXT("NiagaraNeighborSort::BlockSums"));
+
+	// LocalScan — each group scans its tile, writes tile total to BlockSumsBuffer
+	{
+		TShaderMapRef<FNiagaraParticleNeighborQueryPrefixSumLocalScanCS> ComputeShader(GetGlobalShaderMap(GMaxRHIFeatureLevel));
+
+		auto* Parameters = GraphBuilder.AllocParameters<FNiagaraParticleNeighborQueryPrefixSumLocalScanCS::FParameters>();
+		Parameters->InputBuffer = GraphBuilder.CreateSRV(InputBuffer, PF_R32_SINT);
+		Parameters->OutputBuffer = OutputUAV;
+		Parameters->BlockSumsBuffer = GraphBuilder.CreateUAV(BlockSumsBuffer, PF_R32_SINT);
+		Parameters->NumElements = NumElements;
+		Parameters->GroupsX = LocalScanDispatch.X;
+
+		FComputeShaderUtils::AddPass(
+			GraphBuilder,
+			RDG_EVENT_NAME("NiagaraNeighborSort::PrefixSum::LocalScan NumElements(%u) DispatchCount(%ux%ux1) NumThreads(%ux1x1)", NumElements, LocalScanDispatch.X, LocalScanDispatch.Y, ThreadGroupSize),
+			ERDGPassFlags::Compute,
+			ComputeShader,
+			Parameters,
+			LocalScanDispatch
+		);
+	}
+
+	// Recursively scan the block sums
+	FRDGBufferRef ScannedBlockSums = GraphBuilder.CreateBuffer(FRDGBufferDesc::CreateBufferDesc(sizeof(int32), NumGroups), TEXT("NiagaraNeighborSort::ScannedBlockSums"));
+	FRDGBufferUAVRef ScannedBlockSumsUAV = GraphBuilder.CreateUAV(ScannedBlockSums, PF_R32_SINT);
+	PrefixSumRecursive(GraphBuilder, BlockSumsBuffer, ScannedBlockSumsUAV, NumGroups);
+
+	// Propagate — each group adds its scanned block prefix to local results
+	{
+		TShaderMapRef<FNiagaraParticleNeighborQueryPrefixSumPropagateCS> ComputeShader(GetGlobalShaderMap(GMaxRHIFeatureLevel));
+		const FIntVector PropagateDispatch = GetDispatchSize(NumGroups);
+
+		auto* Parameters = GraphBuilder.AllocParameters<FNiagaraParticleNeighborQueryPrefixSumPropagateCS::FParameters>();
+		Parameters->BlockSumsBuffer = GraphBuilder.CreateSRV(ScannedBlockSums, PF_R32_SINT);
+		Parameters->OutputBuffer = OutputUAV;
+		Parameters->NumElements = NumElements;
+		Parameters->GroupsX = PropagateDispatch.X;
+
+		FComputeShaderUtils::AddPass(
+			GraphBuilder,
+			RDG_EVENT_NAME("NiagaraNeighborSort::PrefixSum::Propagate NumElements(%u) DispatchCount(%ux%ux1) NumThreads(%ux1x1)", NumElements, PropagateDispatch.X, PropagateDispatch.Y, ThreadGroupSize),
+			ERDGPassFlags::Compute,
+			ComputeShader,
+			Parameters,
+			PropagateDispatch
+		);
+	}
+}
+
+void NiagaraParticleNeighborQuerySort::PrefixSum(FRDGBuilder& GraphBuilder, FRDGBufferUAVRef CellCountUAV, FRDGBufferUAVRef CellOffsetUAV, uint32 NumCells)
+{
+	PrefixSumRecursive(GraphBuilder, CellCountUAV->GetParent(), CellOffsetUAV, NumCells);
+}
+
+void NiagaraParticleNeighborQuerySort::Scatter(FRDGBuilder& GraphBuilder, FRDGBufferSRVRef CellIdSRV, FRDGBufferUAVRef CellOffsetUAV, FRDGBufferUAVRef ParticleListUAV, uint32 NumSlots, uint32 MaxCellsPerParticle)
+{
+	if (NumSlots == 0)
+	{
+		return;
+	}
+
+	TShaderMapRef<FNiagaraParticleNeighborQueryScatterCS> ComputeShader(GetGlobalShaderMap(GMaxRHIFeatureLevel));
+	const uint32 NumThreadGroups = FMath::DivideAndRoundUp(NumSlots, FNiagaraParticleNeighborQueryScatterCS::ThreadGroupSize);
+	const FIntVector DispatchSize = GetDispatchSize(NumThreadGroups);
+
+	FNiagaraParticleNeighborQueryScatterCS::FParameters* Parameters = GraphBuilder.AllocParameters<FNiagaraParticleNeighborQueryScatterCS::FParameters>();
+	Parameters->CellIdBuffer = CellIdSRV;
+	Parameters->CellOffsetBuffer = CellOffsetUAV;
+	Parameters->ParticleListBuffer = ParticleListUAV;
+	Parameters->NumSlots = NumSlots;
+	Parameters->MaxCellsPerParticle = MaxCellsPerParticle;
+	Parameters->GroupsX = DispatchSize.X;
+
+	FComputeShaderUtils::AddPass(
+		GraphBuilder,
+		RDG_EVENT_NAME("NiagaraNeighborSort::Scatter NumSlots(%u) MaxCellsPerParticle(%u) DispatchCount(%ux%ux1) NumThreads(%ux1x1)", NumSlots, MaxCellsPerParticle, DispatchSize.X, DispatchSize.Y, FNiagaraParticleNeighborQueryScatterCS::ThreadGroupSize),
+		ERDGPassFlags::Compute,
+		ComputeShader,
+		Parameters,
+		DispatchSize
+	);
+}
diff --git a/Engine/Plugins/FX/Niagara/Source/NiagaraShader/Public/NiagaraParticleNeighborQuerySort.h b/Engine/Plugins/FX/Niagara/Source/NiagaraShader/Public/NiagaraParticleNeighborQuerySort.h
new file mode 100644
--- /dev/null
+++ b/Engine/Plugins/FX/Niagara/Source/NiagaraShader/Public/NiagaraParticleNeighborQuerySort.h
@@ -0,0 +1,21 @@
+// Copyright Epic Games, Inc. All Rights Reserved.
+
+#pragma once
+
+#include "RenderGraphFwd.h"
+
+class FRDGBuilder;
+
+// Counting sort dispatch functions for NiagaraDataInterfaceParticleNeighborQuery.
+// Called from PostStage after the output stage has written cell assignments.
+namespace NiagaraParticleNeighborQuerySort
+{
+	// Pass 2: Build histogram — count particles per cell
+	NIAGARASHADER_API void Histogram(FRDGBuilder& GraphBuilder, FRDGBufferSRVRef CellIdSRV, FRDGBufferUAVRef CellCountUAV, uint32 NumSlots);
+
+	// Pass 3: Exclusive prefix sum over cell counts
+	NIAGARASHADER_API void PrefixSum(FRDGBuilder& GraphBuilder, FRDGBufferUAVRef CellCountUAV, FRDGBufferUAVRef CellOffsetUAV, uint32 NumCells);
+
+	// Pass 4: Scatter sorted particle indices
+	NIAGARASHADER_API void Scatter(FRDGBuilder& GraphBuilder, FRDGBufferSRVRef CellIdSRV, FRDGBufferUAVRef CellOffsetUAV, FRDGBufferUAVRef ParticleListUAV, uint32 NumSlots, uint32 MaxCellsPerParticle);
+}
