Support IDs for Particle Neighbor Query
Update related content

#virtualized
#Synced-CL 52548380

diff --git a/Engine/Plugins/FX/Niagara/Shaders/Private/NiagaraDataInterfaceParticleNeighborQuery.ush b/Engine/Plugins/FX/Niagara/Shaders/Private/NiagaraDataInterfaceParticleNeighborQuery.ush
--- a/Engine/Plugins/FX/Niagara/Shaders/Private/NiagaraDataInterfaceParticleNeighborQuery.ush
+++ b/Engine/Plugins/FX/Niagara/Shaders/Private/NiagaraDataInterfaceParticleNeighborQuery.ush
@@ -15,6 +15,8 @@ Buffer<int>		{ParameterName}_CellCount;
 Buffer<int>		{ParameterName}_CellOffset;
 Buffer<int>		{ParameterName}_ParticleList;
 RWBuffer<int>	{ParameterName}_OutputCellId;
+RWBuffer<int>	{ParameterName}_OutputAcquireTag;
+Buffer<int>		{ParameterName}_AcquireTagList;
 
 #ifndef NEIGHBORGRID_INVALID_CELL
 #define NEIGHBORGRID_INVALID_CELL 0xFFFFFFFF
@@ -46,7 +48,7 @@ void MaxCellsPerParticle_{ParameterName}(out int OutMaxCellsPerParticle)
  * -----------------------------------------------------------------
  */
 
-void AddParticle_{ParameterName}(bool Enabled, float3 IndexPos, int ParticleIndex, out bool Success)
+void AddParticle_{ParameterName}(bool Enabled, float3 IndexPos, int ParticleIndex, int AcquireTag, out bool Success)
 {
 	Success = false;
 	if (!Enabled) return;
@@ -60,18 +62,19 @@ void AddParticle_{ParameterName}(bool Enabled, float3 IndexPos, int ParticleInde
 		// Single-cell: 1 slot per particle, write to first slot
 		int SlotBase = ParticleIndex * {ParameterName}_MaxCellsPerParticleValue;
 		{ParameterName}_OutputCellId[SlotBase] = LinearIndex;
+		{ParameterName}_OutputAcquireTag[SlotBase] = AcquireTag;
 
 		Success = true;
 	}
 }
 
 // Exec-pin wrapper — Niagara appends _UEImpureCall to the generated function name for bRequiresExecPin functions
-void AddParticle_{ParameterName}_UEImpureCall(bool In_Enabled, float3 In_IndexPos, int In_ParticleIndex, out bool Out_Success)
+void AddParticle_{ParameterName}_UEImpureCall(bool In_Enabled, float3 In_IndexPos, int In_ParticleIndex, int In_AcquireTag, out bool Out_Success)
 {
-	AddParticle_{ParameterName}(In_Enabled, In_IndexPos, In_ParticleIndex, Out_Success);
+	AddParticle_{ParameterName}(In_Enabled, In_IndexPos, In_ParticleIndex, In_AcquireTag, Out_Success);
 }
 
-void AddParticleWithRadius_{ParameterName}(bool Enabled, float3 IndexPos, float Radius, int ParticleIndex, out bool Success)
+void AddParticleWithRadius_{ParameterName}(bool Enabled, float3 IndexPos, float Radius, int ParticleIndex, int AcquireTag, out bool Success)
 {
 	Success = false;
 	if (!Enabled) return;
@@ -89,6 +92,7 @@ void AddParticleWithRadius_{ParameterName}(bool Enabled, float3 IndexPos, float
 		{
 			int LinearIndex = IndexToLinear_{ParameterName}(int3(x, y, z));
 			{ParameterName}_OutputCellId[SlotBase + Slot] = LinearIndex;
+			{ParameterName}_OutputAcquireTag[SlotBase + Slot] = AcquireTag;
 			Slot++;
 			Success = true;
 		}
@@ -96,12 +100,12 @@ void AddParticleWithRadius_{ParameterName}(bool Enabled, float3 IndexPos, float
 }
 
 // Exec-pin wrapper
-void AddParticleWithRadius_{ParameterName}_UEImpureCall(bool In_Enabled, float3 In_IndexPos, float In_Radius, int In_ParticleIndex, out bool Out_Success)
+void AddParticleWithRadius_{ParameterName}_UEImpureCall(bool In_Enabled, float3 In_IndexPos, float In_Radius, int In_ParticleIndex, int In_AcquireTag, out bool Out_Success)
 {
-	AddParticleWithRadius_{ParameterName}(In_Enabled, In_IndexPos, In_Radius, In_ParticleIndex, Out_Success);
+	AddParticleWithRadius_{ParameterName}(In_Enabled, In_IndexPos, In_Radius, In_ParticleIndex, In_AcquireTag, Out_Success);
 }
 
-void AddParticleToOverlappedCells_{ParameterName}(bool Enabled, float3 IndexPos, int ParticleIndex, out bool Success)
+void AddParticleToOverlappedCells_{ParameterName}(bool Enabled, float3 IndexPos, int ParticleIndex, int AcquireTag, out bool Success)
 {
 	Success = false;
 	if (!Enabled) return;
@@ -118,6 +122,7 @@ void AddParticleToOverlappedCells_{ParameterName}(bool Enabled, float3 IndexPos,
 			all(CellCoord >= 0) && all(CellCoord < {ParameterName}_NumCells))
 		{
 			{ParameterName}_OutputCellId[SlotBase + Slot] = IndexToLinear_{ParameterName}(CellCoord);
+			{ParameterName}_OutputAcquireTag[SlotBase + Slot] = AcquireTag;
 			Slot++;
 			Success = true;
 		}
@@ -125,12 +130,12 @@ void AddParticleToOverlappedCells_{ParameterName}(bool Enabled, float3 IndexPos,
 }
 
 // Exec-pin wrapper
-void AddParticleToOverlappedCells_{ParameterName}_UEImpureCall(bool In_Enabled, float3 In_IndexPos, int In_ParticleIndex, out bool Out_Success)
+void AddParticleToOverlappedCells_{ParameterName}_UEImpureCall(bool In_Enabled, float3 In_IndexPos, int In_ParticleIndex, int In_AcquireTag, out bool Out_Success)
 {
-	AddParticleToOverlappedCells_{ParameterName}(In_Enabled, In_IndexPos, In_ParticleIndex, Out_Success);
+	AddParticleToOverlappedCells_{ParameterName}(In_Enabled, In_IndexPos, In_ParticleIndex, In_AcquireTag, Out_Success);
 }
 
-void AddParticleToNeighborCells_{ParameterName}(bool Enabled, float3 IndexPos, int ParticleIndex, out bool Success)
+void AddParticleToNeighborCells_{ParameterName}(bool Enabled, float3 IndexPos, int ParticleIndex, int AcquireTag, out bool Success)
 {
 	Success = false;
 	if (!Enabled) return;
@@ -147,6 +152,7 @@ void AddParticleToNeighborCells_{ParameterName}(bool Enabled, float3 IndexPos, i
 			all(CellCoord >= 0) && all(CellCoord < {ParameterName}_NumCells))
 		{
 			{ParameterName}_OutputCellId[SlotBase + Slot] = IndexToLinear_{ParameterName}(CellCoord);
+			{ParameterName}_OutputAcquireTag[SlotBase + Slot] = AcquireTag;
 			Slot++;
 			Success = true;
 		}
@@ -154,9 +160,9 @@ void AddParticleToNeighborCells_{ParameterName}(bool Enabled, float3 IndexPos, i
 }
 
 // Exec-pin wrapper
-void AddParticleToNeighborCells_{ParameterName}_UEImpureCall(bool In_Enabled, float3 In_IndexPos, int In_ParticleIndex, out bool Out_Success)
+void AddParticleToNeighborCells_{ParameterName}_UEImpureCall(bool In_Enabled, float3 In_IndexPos, int In_ParticleIndex, int In_AcquireTag, out bool Out_Success)
 {
-	AddParticleToNeighborCells_{ParameterName}(In_Enabled, In_IndexPos, In_ParticleIndex, Out_Success);
+	AddParticleToNeighborCells_{ParameterName}(In_Enabled, In_IndexPos, In_ParticleIndex, In_AcquireTag, Out_Success);
 }
 
 /* -----------------------------------------------------------------
@@ -176,7 +182,7 @@ void GetParticleNeighborCount_{ParameterName}(int IndexX, int IndexY, int IndexZ
 	OutCount = {ParameterName}_CellCount[LinearIndex];
 }
 
-void GetParticleNeighbor_{ParameterName}(int IndexX, int IndexY, int IndexZ, int NeighborIdx, out int OutParticleIdx)
+void GetParticleNeighbor_{ParameterName}(int IndexX, int IndexY, int IndexZ, int NeighborIdx, out int OutParticleIdx, out int OutAcquireTag)
 {
 	int3 CellCoord = int3(IndexX, IndexY, IndexZ);
 	bool bValidCell = all(CellCoord >= 0) && all(CellCoord < {ParameterName}_NumCells);
@@ -185,10 +191,12 @@ void GetParticleNeighbor_{ParameterName}(int IndexX, int IndexY, int IndexZ, int
 	if (!bValidCell || NeighborIdx >= Count || NeighborIdx < 0)
 	{
 		OutParticleIdx = -1;
+		OutAcquireTag = -1;
 		return;
 	}
 	// CellOffset was consumed by the scatter pass (atomicAdd), so the original prefix sum value
 	// is CellOffset[i] - CellCount[i].
 	int StartOffset = {ParameterName}_CellOffset[LinearIndex] - Count;
 	OutParticleIdx = {ParameterName}_ParticleList[StartOffset + NeighborIdx];
+	OutAcquireTag = {ParameterName}_AcquireTagList[StartOffset + NeighborIdx];
 }
diff --git a/Engine/Plugins/FX/Niagara/Shaders/Private/NiagaraParticleNeighborQuerySort.usf b/Engine/Plugins/FX/Niagara/Shaders/Private/NiagaraParticleNeighborQuerySort.usf
--- a/Engine/Plugins/FX/Niagara/Shaders/Private/NiagaraParticleNeighborQuerySort.usf
+++ b/Engine/Plugins/FX/Niagara/Shaders/Private/NiagaraParticleNeighborQuerySort.usf
@@ -312,9 +312,12 @@ void PrefixSumPropagateCS(uint3 GroupId : SV_GroupID, uint3 GroupThreadId : SV_G
 Buffer<int>		CellIdBuffer;
 RWBuffer<int>	CellOffsetBuffer;
 RWBuffer<int>	ParticleListBuffer;
+Buffer<int>		AcquireTagBuffer;
+RWBuffer<int>	AcquireTagListBuffer;
 uint			NumSlots;
 uint			MaxCellsPerParticle;
 uint			GroupsX;
+uint			bCarryAcquireTag;
 
 #ifdef USE_BITONIC_SCATTER
 
@@ -419,9 +422,13 @@ void ScatterCS(uint3 GroupId : SV_GroupID, uint3 GroupThreadId : SV_GroupThreadI
 		uint RunStartId = LocalId - LocalOffset;
 		uint WritePos = RunWriteBase[RunStartId] + LocalOffset;
 
-		// Write the particle index to the pre-computed position — no atomic needed
+		// Write the particle index and AcquireTag to the pre-computed position — no atomic needed
 		uint ParticleIndex = MySlotIndex / MaxCellsPerParticle;
 		ParticleListBuffer[WritePos] = ParticleIndex;
+		if (bCarryAcquireTag)
+		{
+			AcquireTagListBuffer[WritePos] = AcquireTagBuffer[MySlotIndex];
+		}
 	}
 #else
 	// Direct atomic path: one InterlockedAdd per valid slot.
@@ -439,6 +446,10 @@ void ScatterCS(uint3 GroupId : SV_GroupID, uint3 GroupThreadId : SV_GroupThreadI
 
 		uint ParticleIndex = SlotIndex / MaxCellsPerParticle;
 		ParticleListBuffer[WritePos] = ParticleIndex;
+		if (bCarryAcquireTag)
+		{
+			AcquireTagListBuffer[WritePos] = AcquireTagBuffer[SlotIndex];
+		}
 	}
 #endif
 }
diff --git a/Engine/Plugins/FX/Niagara/Source/Niagara/Classes/NiagaraDataInterfaceParticleNeighborQuery.h b/Engine/Plugins/FX/Niagara/Source/Niagara/Classes/NiagaraDataInterfaceParticleNeighborQuery.h
--- a/Engine/Plugins/FX/Niagara/Source/Niagara/Classes/NiagaraDataInterfaceParticleNeighborQuery.h
+++ b/Engine/Plugins/FX/Niagara/Source/Niagara/Classes/NiagaraDataInterfaceParticleNeighborQuery.h
@@ -15,6 +15,7 @@ struct FNDIParticleNeighborQueryInstanceData_GT
 	FIntVector	NumCells = FIntVector::ZeroValue;
 	uint32		MaxCellsPerParticle = 1;
 	bool		bCountOnly = false;
+	bool		bUsePersistentIDs = true;
 	bool		bNeedsRealloc = false;
 };
 
@@ -30,6 +31,7 @@ struct FNDIParticleNeighborQueryInstanceData_RT
 	uint64		AllocatedNumSlots = 0;
 	uint32		AllocatedNumCells = 0;
 	bool		bCountOnly = false;
+	bool		bUsePersistentIDs = true;
 	bool		bNeedsRealloc = false;
 
 	// Counting sort buffers
@@ -37,6 +39,8 @@ struct FNDIParticleNeighborQueryInstanceData_RT
 	FNiagaraPooledRWBuffer CellCountBuffer;			// numCells - particles per cell (histogram output)
 	FNiagaraPooledRWBuffer CellOffsetBuffer;		// numCells - prefix sum write offset per cell
 	FNiagaraPooledRWBuffer ParticleListBuffer;		// numSlots - final sorted particle indices
+	FNiagaraPooledRWBuffer AcquireTagBuffer;		// numSlots - AcquireTag per slot, written by AddParticle
+	FNiagaraPooledRWBuffer AcquireTagListBuffer;	// numSlots - sorted AcquireTag, written by Scatter, read by GetParticleNeighbor
 };
 
 struct FNiagaraDataInterfaceProxyParticleNeighborQuery : public FNiagaraDataInterfaceProxyRW
@@ -71,6 +75,10 @@ class UNiagaraDataInterfaceParticleNeighborQuery : public UNiagaraDataInterfaceG
 
 		// Write path AddParticle calls write to this
 		SHADER_PARAMETER_RDG_BUFFER_UAV(RWBuffer<int>,	OutputCellId)
+		SHADER_PARAMETER_RDG_BUFFER_UAV(RWBuffer<int>,	OutputAcquireTag)
+
+		// Read path for sorted AcquireTag (alongside ParticleList)
+		SHADER_PARAMETER_RDG_BUFFER_SRV(Buffer<int>,	AcquireTagList)
 	END_SHADER_PARAMETER_STRUCT()
 
 public:
@@ -83,6 +91,13 @@ public:
 	UPROPERTY(EditAnywhere, Category = "Grid")
 	bool bCountOnly = false;
 
+	/** When enabled (default), AcquireTag is tracked through the counting sort for compaction-safe
+	  * neighbor lookups via GetPositionByID. The emitter must have Persistent IDs enabled.
+	  * When disabled, AcquireTag buffers are not allocated and GetParticleNeighbor returns 0 for
+	  * AcquireTag. Only disable if you know compaction will not reorder particles. */
+	UPROPERTY(EditAnywhere, Category = "Grid")
+	bool bUsePersistentIDs = true;
+
 	virtual void PostInitProperties() override
 	{
 		Super::PostInitProperties();
@@ -109,6 +124,10 @@ public:
 	NIAGARA_API virtual void BuildShaderParameters(FNiagaraShaderParametersBuilder& ShaderParametersBuilder) const override;
 	NIAGARA_API virtual void SetShaderParameters(const FNiagaraDataInterfaceSetShaderParametersContext& Context) const override;
 
+#if WITH_EDITOR
+	NIAGARA_API virtual void GetFeedback(UNiagaraSystem* InAsset, UNiagaraComponent* InComponent, TArray<FNiagaraDataInterfaceError>& OutErrors, TArray<FNiagaraDataInterfaceFeedback>& OutWarnings, TArray<FNiagaraDataInterfaceFeedback>& OutInfo) override;
+#endif
+
 	virtual void ProvidePerInstanceDataForRenderThread(void* DataForRenderThread, void* PerInstanceData, const FNiagaraSystemInstanceID& SystemInstance) override {}
 	NIAGARA_API virtual bool InitPerInstanceData(void* PerInstanceData, FNiagaraSystemInstance* SystemInstance) override;
 	NIAGARA_API virtual void DestroyPerInstanceData(void* PerInstanceData, FNiagaraSystemInstance* SystemInstance) override;
diff --git a/Engine/Plugins/FX/Niagara/Source/Niagara/Private/NiagaraDataInterfaceParticleNeighborQuery.cpp b/Engine/Plugins/FX/Niagara/Source/Niagara/Private/NiagaraDataInterfaceParticleNeighborQuery.cpp
--- a/Engine/Plugins/FX/Niagara/Source/Niagara/Private/NiagaraDataInterfaceParticleNeighborQuery.cpp
+++ b/Engine/Plugins/FX/Niagara/Source/Niagara/Private/NiagaraDataInterfaceParticleNeighborQuery.cpp
@@ -5,6 +5,8 @@
 #include "NiagaraSimStageData.h"
 #include "NiagaraShaderParametersBuilder.h"
 #include "NiagaraSystemInstance.h"
+#include "NiagaraEmitter.h"
+#include "NiagaraScript.h"
 #include "DataDrivenShaderPlatformInfo.h"
 #include "RenderGraphUtils.h"
 #include "NiagaraParticleNeighborQuerySort.h"
@@ -56,11 +58,27 @@ void FNDIParticleNeighborQueryInstanceData_RT::ResizeBuffers(FRDGBuilder& GraphB
 		{
 			ParticleListBuffer.Release();
 			ParticleListBuffer.Initialize(GraphBuilder, TEXT("ParticleNeighborQuery::ParticleList"), EPixelFormat::PF_R32_SINT, sizeof(int32), NumSlots, BUF_Static);
+
+			if (bUsePersistentIDs)
+			{
+				AcquireTagBuffer.Release();
+				AcquireTagBuffer.Initialize(GraphBuilder, TEXT("ParticleNeighborQuery::AcquireTag"), EPixelFormat::PF_R32_SINT, sizeof(int32), NumSlots, BUF_Static);
+
+				AcquireTagListBuffer.Release();
+				AcquireTagListBuffer.Initialize(GraphBuilder, TEXT("ParticleNeighborQuery::AcquireTagList"), EPixelFormat::PF_R32_SINT, sizeof(int32), NumSlots, BUF_Static);
+			}
 		}
 
 		AllocatedNumSlots = NumSlots;
 	}
 
+	// Release AcquireTag buffers when persistent IDs are disabled (mode transition)
+	if (!bUsePersistentIDs)
+	{
+		AcquireTagBuffer.Release();
+		AcquireTagListBuffer.Release();
+	}
+
 	// Cell-sized buffers: reallocate when grid dimensions change
 	if (NumTotalCells != AllocatedNumCells)
 	{
@@ -97,6 +115,7 @@ UNiagaraDataInterfaceParticleNeighborQuery::UNiagaraDataInterfaceParticleNeighbo
 	: Super(ObjectInitializer)
 	, MaxCellsPerParticle(1)
 	, bCountOnly(false)
+	, bUsePersistentIDs(true)
 {
 	Proxy.Reset(new FNiagaraDataInterfaceProxyParticleNeighborQuery());
 }
@@ -139,7 +158,7 @@ void UNiagaraDataInterfaceParticleNeighborQuery::GetFunctionsInternal(TArray<FNi
 		OutFunctions.Add(Sig);
 	}
 
-	// AddParticle(Enabled, IndexPos, ParticleIndex) -> Success
+	// AddParticle(Enabled, IndexPos, ParticleIndex, AcquireTag) -> Success
 	{
 		FNiagaraFunctionSignature Sig;
 		Sig.Name = NDIParticleNeighborQueryLocal::AddParticleFunctionName;
@@ -147,6 +166,7 @@ void UNiagaraDataInterfaceParticleNeighborQuery::GetFunctionsInternal(TArray<FNi
 		Sig.Inputs.Add_GetRef(FNiagaraVariable(FNiagaraTypeDefinition::GetBoolDef(), TEXT("Enabled"))).SetValue(true);
 		Sig.Inputs.Add(FNiagaraVariable(FNiagaraTypeDefinition::GetVec3Def(), TEXT("IndexPos")));
 		Sig.Inputs.Add(FNiagaraVariable(FNiagaraTypeDefinition::GetIntDef(), TEXT("ParticleIndex")));
+		Sig.Inputs.Add_GetRef(FNiagaraVariable(FNiagaraTypeDefinition::GetIntDef(), TEXT("AcquireTag"))).SetValue(0);
 		Sig.Outputs.Add(FNiagaraVariable(FNiagaraTypeDefinition::GetBoolDef(), TEXT("Success")));
 
 		Sig.bWriteFunction = true;
@@ -156,7 +176,7 @@ void UNiagaraDataInterfaceParticleNeighborQuery::GetFunctionsInternal(TArray<FNi
 		OutFunctions.Add(Sig);
 	}
 
-	// AddParticleWithRadius(Enabled, IndexPos, Radius, ParticleIndex) -> Success
+	// AddParticleWithRadius(Enabled, IndexPos, Radius, ParticleIndex, AcquireTag) -> Success
 	{
 		FNiagaraFunctionSignature Sig;
 		Sig.Name = AddParticleWithRadiusFunctionName;
@@ -165,6 +185,7 @@ void UNiagaraDataInterfaceParticleNeighborQuery::GetFunctionsInternal(TArray<FNi
 		Sig.Inputs.Add(FNiagaraVariable(FNiagaraTypeDefinition::GetVec3Def(), TEXT("IndexPos")));
 		Sig.Inputs.Add(FNiagaraVariable(FNiagaraTypeDefinition::GetFloatDef(), TEXT("Radius")));
 		Sig.Inputs.Add(FNiagaraVariable(FNiagaraTypeDefinition::GetIntDef(), TEXT("ParticleIndex")));
+		Sig.Inputs.Add_GetRef(FNiagaraVariable(FNiagaraTypeDefinition::GetIntDef(), TEXT("AcquireTag"))).SetValue(0);
 		Sig.Outputs.Add(FNiagaraVariable(FNiagaraTypeDefinition::GetBoolDef(), TEXT("Success")));
 
 		Sig.bWriteFunction = true;
@@ -174,7 +195,7 @@ void UNiagaraDataInterfaceParticleNeighborQuery::GetFunctionsInternal(TArray<FNi
 		OutFunctions.Add(Sig);
 	}
 
-	// AddParticleToOverlappedCells(Enabled, IndexPos, ParticleIndex) -> Success
+	// AddParticleToOverlappedCells(Enabled, IndexPos, ParticleIndex, AcquireTag) -> Success
 	// Writes the particle to all 8 cells in the 2x2x2 positive neighborhood from floor(IndexPos).
 	// For asymmetric gather (P2G rasterization with trilinear kernel). Requires MaxCellsPerParticle >= 8.
 	{
@@ -184,6 +205,7 @@ void UNiagaraDataInterfaceParticleNeighborQuery::GetFunctionsInternal(TArray<FNi
 		Sig.Inputs.Add_GetRef(FNiagaraVariable(FNiagaraTypeDefinition::GetBoolDef(), TEXT("Enabled"))).SetValue(true);
 		Sig.Inputs.Add(FNiagaraVariable(FNiagaraTypeDefinition::GetVec3Def(), TEXT("IndexPos")));
 		Sig.Inputs.Add(FNiagaraVariable(FNiagaraTypeDefinition::GetIntDef(), TEXT("ParticleIndex")));
+		Sig.Inputs.Add_GetRef(FNiagaraVariable(FNiagaraTypeDefinition::GetIntDef(), TEXT("AcquireTag"))).SetValue(0);
 		Sig.Outputs.Add(FNiagaraVariable(FNiagaraTypeDefinition::GetBoolDef(), TEXT("Success")));
 
 		Sig.bWriteFunction = true;
@@ -193,7 +215,7 @@ void UNiagaraDataInterfaceParticleNeighborQuery::GetFunctionsInternal(TArray<FNi
 		OutFunctions.Add(Sig);
 	}
 
-	// AddParticleToNeighborCells(Enabled, IndexPos, ParticleIndex) -> Success
+	// AddParticleToNeighborCells(Enabled, IndexPos, ParticleIndex, AcquireTag) -> Success
 	// Writes the particle to all 27 cells in the 3x3x3 neighborhood centered on floor(IndexPos).
 	// For symmetric bilateral queries (PBD constraints) with single-cell query. Requires MaxCellsPerParticle >= 27.
 	{
@@ -203,6 +225,7 @@ void UNiagaraDataInterfaceParticleNeighborQuery::GetFunctionsInternal(TArray<FNi
 		Sig.Inputs.Add_GetRef(FNiagaraVariable(FNiagaraTypeDefinition::GetBoolDef(), TEXT("Enabled"))).SetValue(true);
 		Sig.Inputs.Add(FNiagaraVariable(FNiagaraTypeDefinition::GetVec3Def(), TEXT("IndexPos")));
 		Sig.Inputs.Add(FNiagaraVariable(FNiagaraTypeDefinition::GetIntDef(), TEXT("ParticleIndex")));
+		Sig.Inputs.Add_GetRef(FNiagaraVariable(FNiagaraTypeDefinition::GetIntDef(), TEXT("AcquireTag"))).SetValue(0);
 		Sig.Outputs.Add(FNiagaraVariable(FNiagaraTypeDefinition::GetBoolDef(), TEXT("Success")));
 
 		Sig.bWriteFunction = true;
@@ -227,7 +250,7 @@ void UNiagaraDataInterfaceParticleNeighborQuery::GetFunctionsInternal(TArray<FNi
 		OutFunctions.Add(Sig);
 	}
 
-	// GetParticleNeighbor(IndexX, IndexY, IndexZ, NeighborIdx) -> ParticleIdx
+	// GetParticleNeighbor(IndexX, IndexY, IndexZ, NeighborIdx) -> (ParticleIdx, AcquireTag)
 	{
 		FNiagaraFunctionSignature Sig;
 		Sig.Name = NDIParticleNeighborQueryLocal::GetParticleNeighborFunctionName;
@@ -237,6 +260,7 @@ void UNiagaraDataInterfaceParticleNeighborQuery::GetFunctionsInternal(TArray<FNi
 		Sig.Inputs.Add(FNiagaraVariable(FNiagaraTypeDefinition::GetIntDef(), TEXT("IndexZ")));
 		Sig.Inputs.Add(FNiagaraVariable(FNiagaraTypeDefinition::GetIntDef(), TEXT("NeighborIdx")));
 		Sig.Outputs.Add(FNiagaraVariable(FNiagaraTypeDefinition::GetIntDef(), TEXT("ParticleIdx")));
+		Sig.Outputs.Add(FNiagaraVariable(FNiagaraTypeDefinition::GetIntDef(), TEXT("AcquireTag")));
 
 		Sig.bMemberFunction = true;
 		Sig.bRequiresContext = false;
@@ -317,7 +341,8 @@ bool UNiagaraDataInterfaceParticleNeighborQuery::Equals(const UNiagaraDataInterf
 	const UNiagaraDataInterfaceParticleNeighborQuery* OtherTyped = CastChecked<const UNiagaraDataInterfaceParticleNeighborQuery>(Other);
 
 	return OtherTyped->MaxCellsPerParticle == MaxCellsPerParticle
-		&& OtherTyped->bCountOnly == bCountOnly;
+		&& OtherTyped->bCountOnly == bCountOnly
+		&& OtherTyped->bUsePersistentIDs == bUsePersistentIDs;
 }
 
 //////////////////////////////////////////////////////////////////////////
@@ -397,11 +422,15 @@ void UNiagaraDataInterfaceParticleNeighborQuery::SetShaderParameters(const FNiag
 
 		if (Context.IsOutputStage())
 		{
-			// Output stage: user modules write CellIdBuffer, read buffers are empty
+			// Output stage: user modules write CellIdBuffer and AcquireTagBuffer, read buffers are empty
 			ShaderParameters->CellCount = Context.GetComputeDispatchInterface().GetEmptyBufferSRV(GraphBuilder, PF_R32_SINT);
 			ShaderParameters->CellOffset = Context.GetComputeDispatchInterface().GetEmptyBufferSRV(GraphBuilder, PF_R32_SINT);
 			ShaderParameters->ParticleList = Context.GetComputeDispatchInterface().GetEmptyBufferSRV(GraphBuilder, PF_R32_SINT);
 			ShaderParameters->OutputCellId = ProxyData->CellIdBuffer.GetOrCreateUAV(GraphBuilder);
+			ShaderParameters->OutputAcquireTag = ProxyData->AcquireTagBuffer.IsValid()
+				? ProxyData->AcquireTagBuffer.GetOrCreateUAV(GraphBuilder)
+				: Context.GetComputeDispatchInterface().GetEmptyBufferUAV(GraphBuilder, PF_R32_SINT);
+			ShaderParameters->AcquireTagList = Context.GetComputeDispatchInterface().GetEmptyBufferSRV(GraphBuilder, PF_R32_SINT);
 		}
 		else
 		{
@@ -411,13 +440,18 @@ void UNiagaraDataInterfaceParticleNeighborQuery::SetShaderParameters(const FNiag
 			{
 				ShaderParameters->CellOffset = Context.GetComputeDispatchInterface().GetEmptyBufferSRV(GraphBuilder, PF_R32_SINT);
 				ShaderParameters->ParticleList = Context.GetComputeDispatchInterface().GetEmptyBufferSRV(GraphBuilder, PF_R32_SINT);
+				ShaderParameters->AcquireTagList = Context.GetComputeDispatchInterface().GetEmptyBufferSRV(GraphBuilder, PF_R32_SINT);
 			}
 			else
 			{
 				ShaderParameters->CellOffset = ProxyData->CellOffsetBuffer.GetOrCreateSRV(GraphBuilder);
 				ShaderParameters->ParticleList = ProxyData->ParticleListBuffer.GetOrCreateSRV(GraphBuilder);
+				ShaderParameters->AcquireTagList = ProxyData->AcquireTagListBuffer.IsValid()
+					? ProxyData->AcquireTagListBuffer.GetOrCreateSRV(GraphBuilder)
+					: Context.GetComputeDispatchInterface().GetEmptyBufferSRV(GraphBuilder, PF_R32_SINT);
 			}
 			ShaderParameters->OutputCellId = Context.GetComputeDispatchInterface().GetEmptyBufferUAV(GraphBuilder, PF_R32_SINT);
+			ShaderParameters->OutputAcquireTag = Context.GetComputeDispatchInterface().GetEmptyBufferUAV(GraphBuilder, PF_R32_SINT);
 		}
 	}
 	else
@@ -430,6 +464,8 @@ void UNiagaraDataInterfaceParticleNeighborQuery::SetShaderParameters(const FNiag
 		ShaderParameters->CellOffset = Context.GetComputeDispatchInterface().GetEmptyBufferSRV(GraphBuilder, PF_R32_SINT);
 		ShaderParameters->ParticleList = Context.GetComputeDispatchInterface().GetEmptyBufferSRV(GraphBuilder, PF_R32_SINT);
 		ShaderParameters->OutputCellId = Context.GetComputeDispatchInterface().GetEmptyBufferUAV(GraphBuilder, PF_R32_SINT);
+		ShaderParameters->OutputAcquireTag = Context.GetComputeDispatchInterface().GetEmptyBufferUAV(GraphBuilder, PF_R32_SINT);
+		ShaderParameters->AcquireTagList = Context.GetComputeDispatchInterface().GetEmptyBufferSRV(GraphBuilder, PF_R32_SINT);
 	}
 }
 
@@ -449,14 +485,16 @@ bool UNiagaraDataInterfaceParticleNeighborQuery::InitPerInstanceData(void* PerIn
 	InstanceData->NumCells = FIntVector(1, 1, 1);
 	InstanceData->MaxCellsPerParticle = MaxCellsPerParticle;
 	InstanceData->bCountOnly = bCountOnly;
+	InstanceData->bUsePersistentIDs = bUsePersistentIDs;
 
 	FNiagaraDataInterfaceProxyParticleNeighborQuery* RT_Proxy = GetProxyAs<FNiagaraDataInterfaceProxyParticleNeighborQuery>();
 	const FIntVector RT_NumCells = InstanceData->NumCells;
 	const uint32 RT_MaxCellsPerParticle = InstanceData->MaxCellsPerParticle;
 	const bool RT_bCountOnly = InstanceData->bCountOnly;
+	const bool RT_bUsePersistentIDs = InstanceData->bUsePersistentIDs;
 
 	ENQUEUE_RENDER_COMMAND(FUpdateData)(
-		[RT_Proxy, RT_NumCells, RT_MaxCellsPerParticle, RT_bCountOnly, InstanceID = SystemInstance->GetId()](FRHICommandListImmediate& RHICmdList)
+		[RT_Proxy, RT_NumCells, RT_MaxCellsPerParticle, RT_bCountOnly, RT_bUsePersistentIDs, InstanceID = SystemInstance->GetId()](FRHICommandListImmediate& RHICmdList)
 		{
 			check(!RT_Proxy->SystemInstancesToProxyData_RT.Contains(InstanceID));
 			FNDIParticleNeighborQueryInstanceData_RT* TargetData = &RT_Proxy->SystemInstancesToProxyData_RT.Add(InstanceID);
@@ -464,6 +502,7 @@ bool UNiagaraDataInterfaceParticleNeighborQuery::InitPerInstanceData(void* PerIn
 			TargetData->NumCells = RT_NumCells;
 			TargetData->MaxCellsPerParticle = RT_MaxCellsPerParticle;
 			TargetData->bCountOnly = RT_bCountOnly;
+			TargetData->bUsePersistentIDs = RT_bUsePersistentIDs;
 			TargetData->bNeedsRealloc = true;
 		}
 	);
@@ -521,7 +560,7 @@ bool UNiagaraDataInterfaceParticleNeighborQuery::PerInstanceTickPostSimulate(voi
 
 		FNiagaraDataInterfaceProxyParticleNeighborQuery* RT_Proxy = GetProxyAs<FNiagaraDataInterfaceProxyParticleNeighborQuery>();
 		ENQUEUE_RENDER_COMMAND(FUpdateData)(
-			[RT_Proxy, RT_NumCells = InstanceData->NumCells, RT_MaxCellsPerParticle = InstanceData->MaxCellsPerParticle, RT_bCountOnly = InstanceData->bCountOnly, InstanceID = SystemInstance->GetId()](FRHICommandListImmediate& RHICmdList)
+			[RT_Proxy, RT_NumCells = InstanceData->NumCells, RT_MaxCellsPerParticle = InstanceData->MaxCellsPerParticle, RT_bCountOnly = InstanceData->bCountOnly, RT_bUsePersistentIDs = InstanceData->bUsePersistentIDs, InstanceID = SystemInstance->GetId()](FRHICommandListImmediate& RHICmdList)
 		{
 			check(RT_Proxy->SystemInstancesToProxyData_RT.Contains(InstanceID));
 			FNDIParticleNeighborQueryInstanceData_RT* TargetData = RT_Proxy->SystemInstancesToProxyData_RT.Find(InstanceID);
@@ -529,6 +568,7 @@ bool UNiagaraDataInterfaceParticleNeighborQuery::PerInstanceTickPostSimulate(voi
 			TargetData->NumCells = RT_NumCells;
 			TargetData->MaxCellsPerParticle = RT_MaxCellsPerParticle;
 			TargetData->bCountOnly = RT_bCountOnly;
+			TargetData->bUsePersistentIDs = RT_bUsePersistentIDs;
 			TargetData->bNeedsRealloc = true;
 				
 			TargetData->NumParticles = 0;
@@ -575,6 +615,10 @@ void FNiagaraDataInterfaceProxyParticleNeighborQuery::ResetData(const FNDIGpuCom
 	{
 		AddClearUAVPass(GraphBuilder, ProxyData->CellIdBuffer.GetOrCreateUAV(GraphBuilder), 0xFFFFFFFF);  // INVALID_CELL sentinel
 	}
+	if (ProxyData->AcquireTagBuffer.IsValid())
+	{
+		AddClearUAVPass(GraphBuilder, ProxyData->AcquireTagBuffer.GetOrCreateUAV(GraphBuilder), 0);
+	}
 	if (ProxyData->CellCountBuffer.IsValid())
 	{
 		AddClearUAVPass(GraphBuilder, ProxyData->CellCountBuffer.GetOrCreateUAV(GraphBuilder), 0);
@@ -604,11 +648,15 @@ void FNiagaraDataInterfaceProxyParticleNeighborQuery::PreStage(const FNDIGpuComp
 		if (ProxyData.ClearBeforeNonIterationStage)
 		{
 			RDG_RHI_EVENT_SCOPE(GraphBuilder, NiagaraParticleNeighborQueryClear);
-			
+
 			if (ProxyData.CellIdBuffer.IsValid())
 			{
 				AddClearUAVPass(GraphBuilder, ProxyData.CellIdBuffer.GetOrCreateUAV(GraphBuilder), 0xFFFFFFFF);  // INVALID_CELL sentinel
 			}
+			if (ProxyData.AcquireTagBuffer.IsValid())
+			{
+				AddClearUAVPass(GraphBuilder, ProxyData.AcquireTagBuffer.GetOrCreateUAV(GraphBuilder), 0);
+			}
 			if (ProxyData.CellCountBuffer.IsValid())
 			{
 				AddClearUAVPass(GraphBuilder, ProxyData.CellCountBuffer.GetOrCreateUAV(GraphBuilder), 0);
@@ -654,14 +702,22 @@ void FNiagaraDataInterfaceProxyParticleNeighborQuery::PostStage(const FNDIGpuCom
 		FRDGBufferUAVRef CellOffsetUAV = ProxyData.CellOffsetBuffer.GetOrCreateUAV(GraphBuilder);
 		FRDGBufferUAVRef ParticleListUAV = ProxyData.ParticleListBuffer.GetOrCreateUAV(GraphBuilder);
 
+		const bool bCarryAcquireTag = ProxyData.bUsePersistentIDs && ProxyData.AcquireTagBuffer.IsValid();
+		FRDGBufferSRVRef AcquireTagSRV = bCarryAcquireTag
+			? ProxyData.AcquireTagBuffer.GetOrCreateSRV(GraphBuilder)
+			: Context.GetComputeDispatchInterface().GetEmptyBufferSRV(GraphBuilder, PF_R32_SINT);
+		FRDGBufferUAVRef AcquireTagListUAV = bCarryAcquireTag
+			? ProxyData.AcquireTagListBuffer.GetOrCreateUAV(GraphBuilder)
+			: Context.GetComputeDispatchInterface().GetEmptyBufferUAV(GraphBuilder, PF_R32_SINT);
+
 		// Pass 3: Prefix Sum — exclusive scan over cell counts
 		NiagaraParticleNeighborQuerySort::PrefixSum(GraphBuilder, CellCountUAV, CellOffsetUAV, NumTotalCells);
 
-		// Pass 4: Scatter — write sorted particle indices into ParticleListBuffer.
+		// Pass 4: Scatter — write sorted particle indices (and AcquireTag) into ParticleListBuffer.
 		// NOTE: Scatter atomically increments CellOffsetBuffer, so after this pass
 		// CellOffset[i] = original_prefix_sum[i] + CellCount[i]. The read functions
 		// in the .ush account for this by subtracting CellCount.
-		NiagaraParticleNeighborQuerySort::Scatter(GraphBuilder, CellIdSRV, CellOffsetUAV, ParticleListUAV, NumSlots, ProxyData.MaxCellsPerParticle);
+		NiagaraParticleNeighborQuerySort::Scatter(GraphBuilder, CellIdSRV, CellOffsetUAV, ParticleListUAV, AcquireTagSRV, AcquireTagListUAV, NumSlots, ProxyData.MaxCellsPerParticle, bCarryAcquireTag);
 	}
 }
 
@@ -674,6 +730,10 @@ void FNiagaraDataInterfaceProxyParticleNeighborQuery::PostSimulate(const FNDIGpu
 		{
 			ProxyData.CellIdBuffer.EndGraphUsage();
 		}
+		if (ProxyData.AcquireTagBuffer.IsValid())
+		{
+			ProxyData.AcquireTagBuffer.EndGraphUsage();
+		}
 		if (ProxyData.CellCountBuffer.IsValid())
 		{
 			ProxyData.CellCountBuffer.EndGraphUsage();
@@ -688,6 +748,10 @@ void FNiagaraDataInterfaceProxyParticleNeighborQuery::PostSimulate(const FNDIGpu
 			{
 				ProxyData.ParticleListBuffer.EndGraphUsage();
 			}
+			if (ProxyData.AcquireTagListBuffer.IsValid())
+			{
+				ProxyData.AcquireTagListBuffer.EndGraphUsage();
+			}
 		}
 	}
 }
@@ -711,8 +775,72 @@ bool UNiagaraDataInterfaceParticleNeighborQuery::CopyToInternal(UNiagaraDataInte
 
 	OtherTyped->MaxCellsPerParticle = MaxCellsPerParticle;
 	OtherTyped->bCountOnly = bCountOnly;
+	OtherTyped->bUsePersistentIDs = bUsePersistentIDs;
 
 	return true;
 }
 
+#if WITH_EDITOR
+void UNiagaraDataInterfaceParticleNeighborQuery::GetFeedback(UNiagaraSystem* InAsset, UNiagaraComponent* InComponent, TArray<FNiagaraDataInterfaceError>& OutErrors, TArray<FNiagaraDataInterfaceFeedback>& OutWarnings, TArray<FNiagaraDataInterfaceFeedback>& OutInfo)
+{
+	Super::GetFeedback(InAsset, InComponent, OutErrors, OutWarnings, OutInfo);
+	if (!InAsset || !bUsePersistentIDs)
+	{
+		return;
+	}
+
+	for (const FNiagaraEmitterHandle& EmitterHandle : InAsset->GetEmitterHandles())
+	{
+		FVersionedNiagaraEmitterData* EmitterData = EmitterHandle.GetEmitterData();
+		if (!EmitterData)
+		{
+			continue;
+		}
+
+		// Check if this emitter references this DI
+		bool bFound = false;
+		TArray<UNiagaraScript*> Scripts;
+		EmitterData->GetScripts(Scripts, false);
+		for (const UNiagaraScript* Script : Scripts)
+		{
+			if (!Script)
+			{
+				continue;
+			}
+			for (const FNiagaraScriptDataInterfaceInfo& DIInfo : Script->GetCachedDefaultDataInterfaces())
+			{
+				if (DIInfo.DataInterface && DIInfo.DataInterface->Equals(this))
+				{
+					bFound = true;
+					break;
+				}
+			}
+			if (bFound)
+			{
+				break;
+			}
+		}
+
+		if (bFound && !EmitterData->RequiresPersistentIDs())
+		{
+			FNiagaraDataInterfaceFix FixDelegate;
+			FixDelegate.BindLambda([EmitterData]()
+			{
+				EmitterData->bRequiresPersistentIDs = true;
+				return true;
+			});
+
+			OutErrors.Add(FNiagaraDataInterfaceError(
+				FText::Format(LOCTEXT("NeedsPersistentIDs",
+					"Emitter '{0}' must have Persistent IDs enabled for ParticleNeighborQuery. "
+					"Either enable Persistent IDs on the emitter, or disable 'Use Persistent IDs' on this data interface (not recommended)."),
+					FText::FromString(EmitterHandle.GetUniqueInstanceName())),
+				LOCTEXT("NeedsPersistentIDsSummary", "Emitter needs Persistent IDs"),
+				FixDelegate
+			));
+		}
+	}
+}
+#endif
+
 #undef LOCTEXT_NAMESPACE
diff --git a/Engine/Plugins/FX/Niagara/Source/NiagaraShader/Private/NiagaraParticleNeighborQuerySort.cpp b/Engine/Plugins/FX/Niagara/Source/NiagaraShader/Private/NiagaraParticleNeighborQuerySort.cpp
--- a/Engine/Plugins/FX/Niagara/Source/NiagaraShader/Private/NiagaraParticleNeighborQuerySort.cpp
+++ b/Engine/Plugins/FX/Niagara/Source/NiagaraShader/Private/NiagaraParticleNeighborQuerySort.cpp
@@ -120,9 +120,12 @@ class FNiagaraParticleNeighborQueryScatterCS : public FGlobalShader
 		SHADER_PARAMETER_RDG_BUFFER_SRV(Buffer<int>,	CellIdBuffer)
 		SHADER_PARAMETER_RDG_BUFFER_UAV(RWBuffer<int>,	CellOffsetBuffer)
 		SHADER_PARAMETER_RDG_BUFFER_UAV(RWBuffer<int>,	ParticleListBuffer)
+		SHADER_PARAMETER_RDG_BUFFER_SRV(Buffer<int>,	AcquireTagBuffer)
+		SHADER_PARAMETER_RDG_BUFFER_UAV(RWBuffer<int>,	AcquireTagListBuffer)
 		SHADER_PARAMETER(uint32,						NumSlots)
 		SHADER_PARAMETER(uint32,						MaxCellsPerParticle)
 		SHADER_PARAMETER(uint32,						GroupsX)
+		SHADER_PARAMETER(uint32,						bCarryAcquireTag)
 	END_SHADER_PARAMETER_STRUCT()
 };
 
@@ -268,7 +271,7 @@ void NiagaraParticleNeighborQuerySort::PrefixSum(FRDGBuilder& GraphBuilder, FRDG
 	PrefixSumRecursive(GraphBuilder, CellCountUAV->GetParent(), CellOffsetUAV, NumCells);
 }
 
-void NiagaraParticleNeighborQuerySort::Scatter(FRDGBuilder& GraphBuilder, FRDGBufferSRVRef CellIdSRV, FRDGBufferUAVRef CellOffsetUAV, FRDGBufferUAVRef ParticleListUAV, uint32 NumSlots, uint32 MaxCellsPerParticle)
+void NiagaraParticleNeighborQuerySort::Scatter(FRDGBuilder& GraphBuilder, FRDGBufferSRVRef CellIdSRV, FRDGBufferUAVRef CellOffsetUAV, FRDGBufferUAVRef ParticleListUAV, FRDGBufferSRVRef AcquireTagSRV, FRDGBufferUAVRef AcquireTagListUAV, uint32 NumSlots, uint32 MaxCellsPerParticle, bool bCarryAcquireTag)
 {
 	if (NumSlots == 0)
 	{
@@ -283,9 +286,12 @@ void NiagaraParticleNeighborQuerySort::Scatter(FRDGBuilder& GraphBuilder, FRDGBu
 	Parameters->CellIdBuffer = CellIdSRV;
 	Parameters->CellOffsetBuffer = CellOffsetUAV;
 	Parameters->ParticleListBuffer = ParticleListUAV;
+	Parameters->AcquireTagBuffer = AcquireTagSRV;
+	Parameters->AcquireTagListBuffer = AcquireTagListUAV;
 	Parameters->NumSlots = NumSlots;
 	Parameters->MaxCellsPerParticle = MaxCellsPerParticle;
 	Parameters->GroupsX = DispatchSize.X;
+	Parameters->bCarryAcquireTag = bCarryAcquireTag ? 1u : 0u;
 
 	FComputeShaderUtils::AddPass(
 		GraphBuilder,
diff --git a/Engine/Plugins/FX/Niagara/Source/NiagaraShader/Public/NiagaraParticleNeighborQuerySort.h b/Engine/Plugins/FX/Niagara/Source/NiagaraShader/Public/NiagaraParticleNeighborQuerySort.h
--- a/Engine/Plugins/FX/Niagara/Source/NiagaraShader/Public/NiagaraParticleNeighborQuerySort.h
+++ b/Engine/Plugins/FX/Niagara/Source/NiagaraShader/Public/NiagaraParticleNeighborQuerySort.h
@@ -16,6 +16,6 @@ namespace NiagaraParticleNeighborQuerySort
 	// Pass 3: Exclusive prefix sum over cell counts
 	NIAGARASHADER_API void PrefixSum(FRDGBuilder& GraphBuilder, FRDGBufferUAVRef CellCountUAV, FRDGBufferUAVRef CellOffsetUAV, uint32 NumCells);
 
-	// Pass 4: Scatter sorted particle indices
-	NIAGARASHADER_API void Scatter(FRDGBuilder& GraphBuilder, FRDGBufferSRVRef CellIdSRV, FRDGBufferUAVRef CellOffsetUAV, FRDGBufferUAVRef ParticleListUAV, uint32 NumSlots, uint32 MaxCellsPerParticle);
+	// Pass 4: Scatter sorted particle indices (and AcquireTag satellite data)
+	NIAGARASHADER_API void Scatter(FRDGBuilder& GraphBuilder, FRDGBufferSRVRef CellIdSRV, FRDGBufferUAVRef CellOffsetUAV, FRDGBufferUAVRef ParticleListUAV, FRDGBufferSRVRef AcquireTagSRV, FRDGBufferUAVRef AcquireTagListUAV, uint32 NumSlots, uint32 MaxCellsPerParticle, bool bCarryAcquireTag = true);
 }
