[NiagaraDebugHud] Add FiltersEnabled console command tuple

Adds a MakeTuple entry for "FiltersEnabled=" to NiagaraDebugHud.cpp,
matching the pattern used by OverviewEnabled and other existing settings.
This allows bFiltersEnabled to be toggled via the Niagara.DebugHUD console
command, enabling global filter control without requiring per-filter commands.

#Synced-CL 52567348

diff --git a/Engine/Plugins/FX/Niagara/Source/Niagara/Private/NiagaraDebugHud.cpp b/Engine/Plugins/FX/Niagara/Source/Niagara/Private/NiagaraDebugHud.cpp
--- a/Engine/Plugins/FX/Niagara/Source/Niagara/Private/NiagaraDebugHud.cpp
+++ b/Engine/Plugins/FX/Niagara/Source/Niagara/Private/NiagaraDebugHud.cpp
@@ -139,6 +139,7 @@ namespace NiagaraDebugLocal
 		MakeTuple(TEXT("OverviewMode="), TEXT("Change the mode of the debug overivew"), [](FString Arg) {Settings.OverviewMode = (ENiagaraDebugHUDOverviewMode)FCString::Atoi(*Arg); }),
 
 		MakeTuple(TEXT("OverviewLocation="), TEXT("Set the overview location"), [](FString Arg) { Settings.OverviewLocation = FVector2D(StringToVector2f(Arg, FVector2f(Settings.OverviewLocation))); }),
+		MakeTuple(TEXT("FiltersEnabled="), TEXT("Enable or disable the filters."), [](FString Arg) {Settings.bFiltersEnabled = FCString::Atoi(*Arg) != 0; }),
 		MakeTuple(TEXT("SystemFilter="), TEXT("Set the system filter"), [](FString Arg) {Settings.SystemFilter = Arg; Settings.bSystemFilterEnabled = !Arg.IsEmpty(); }),
 		MakeTuple(TEXT("EmitterFilter="), TEXT("Set the emitter filter"), [](FString Arg) {Settings.EmitterFilter = Arg; Settings.bEmitterFilterEnabled = !Arg.IsEmpty(); GCachedSystemVariables.Empty(); }),
 		MakeTuple(TEXT("ActorFilter="), TEXT("Set the actor filter"), [](FString Arg) {Settings.ActorFilter = Arg; Settings.bActorFilterEnabled = !Arg.IsEmpty(); }),
