single grid mode to speedup pressure solver

#Synced-CL 52589973

diff --git a/Engine/Plugins/FX/Niagara/Source/Niagara/Classes/NiagaraDataInterfaceGrid3DCollection.h b/Engine/Plugins/FX/Niagara/Source/Niagara/Classes/NiagaraDataInterfaceGrid3DCollection.h
--- a/Engine/Plugins/FX/Niagara/Source/Niagara/Classes/NiagaraDataInterfaceGrid3DCollection.h
+++ b/Engine/Plugins/FX/Niagara/Source/Niagara/Classes/NiagaraDataInterfaceGrid3DCollection.h
@@ -290,6 +290,11 @@ public:
 	static NIAGARA_API const FName GetPreviousFloatValueFunctionName;
 	static NIAGARA_API const FName SamplePreviousGridFloatFunctionName;
 	static NIAGARA_API const FName CubicSamplePreviousGridFloatFunctionName;
+
+	static NIAGARA_API const FName GetOutputGridFloatValueFunctionName;
+	static NIAGARA_API const FName GetOutputGridVector2ValueFunctionName;
+	static NIAGARA_API const FName GetOutputGridVectorValueFunctionName;
+	static NIAGARA_API const FName GetOutputGridVector4ValueFunctionName;
 	static NIAGARA_API const FString AttributeIndicesBaseName;
 	static NIAGARA_API const FString PerAttributeDataName;
 	static NIAGARA_API const TCHAR* VectorComponentNames[];
@@ -326,6 +331,7 @@ protected:
 	NIAGARA_API void WriteSetHLSL(const FNiagaraDataInterfaceGPUParamInfo& ParamInfo, const FNiagaraDataInterfaceGeneratedFunction& FunctionInfo, int FunctionInstanceIndex, int32 InNumChannels, FString& OutHLSL);
 	NIAGARA_API void WriteSetAtIndexHLSL(const FNiagaraDataInterfaceGPUParamInfo& ParamInfo, const FNiagaraDataInterfaceGeneratedFunction& FunctionInfo, int32 InNumChannels, FString& OutHLSL);
 	NIAGARA_API void WriteGetHLSL(const FNiagaraDataInterfaceGPUParamInfo& ParamInfo, const FNiagaraDataInterfaceGeneratedFunction& FunctionInfo, int FunctionInstanceIndex, int32 InNumChannels, FString& OutHLSL);
+	NIAGARA_API void WriteGetOutputGridHLSL(const FNiagaraDataInterfaceGPUParamInfo& ParamInfo, const FNiagaraDataInterfaceGeneratedFunction& FunctionInfo, int FunctionInstanceIndex, int32 InNumChannels, FString& OutHLSL);
 	NIAGARA_API void WriteGetAtIndexHLSL(const FNiagaraDataInterfaceGPUParamInfo& ParamInfo, const FNiagaraDataInterfaceGeneratedFunction& FunctionInfo, int32 InNumChannels, FString& OutHLSL);
 	NIAGARA_API void WriteSampleHLSL(const FNiagaraDataInterfaceGPUParamInfo& ParamInfo, const FNiagaraDataInterfaceGeneratedFunction& FunctionInfo, int FunctionInstanceIndex, int32 InNumChannels, FString SampleFunction, FString& OutHLSL);
 	NIAGARA_API void WriteSampleAtIndexHLSL(const FNiagaraDataInterfaceGPUParamInfo& ParamInfo, const FNiagaraDataInterfaceGeneratedFunction& FunctionInfo, int32 InNumChannels, FString SampleFunction, FString& OutHLSL);
diff --git a/Engine/Plugins/FX/Niagara/Source/Niagara/Private/NiagaraDataInterfaceGrid3DCollection.cpp b/Engine/Plugins/FX/Niagara/Source/Niagara/Private/NiagaraDataInterfaceGrid3DCollection.cpp
--- a/Engine/Plugins/FX/Niagara/Source/Niagara/Private/NiagaraDataInterfaceGrid3DCollection.cpp
+++ b/Engine/Plugins/FX/Niagara/Source/Niagara/Private/NiagaraDataInterfaceGrid3DCollection.cpp
@@ -134,6 +134,11 @@ const FName UNiagaraDataInterfaceGrid3DCollection::GetPreviousFloatValueFunction
 const FName UNiagaraDataInterfaceGrid3DCollection::SamplePreviousGridFloatFunctionName("SamplePreviousGridFloatValue");
 const FName UNiagaraDataInterfaceGrid3DCollection::CubicSamplePreviousGridFloatFunctionName("CubicSamplePreviousGridFloatValue");
 
+const FName UNiagaraDataInterfaceGrid3DCollection::GetOutputGridFloatValueFunctionName("GetOutputGridFloatValue");
+const FName UNiagaraDataInterfaceGrid3DCollection::GetOutputGridVector2ValueFunctionName("GetOutputGridVector2Value");
+const FName UNiagaraDataInterfaceGrid3DCollection::GetOutputGridVectorValueFunctionName("GetOutputGridVectorValue");
+const FName UNiagaraDataInterfaceGrid3DCollection::GetOutputGridVector4ValueFunctionName("GetOutputGridVector4Value");
+
 const FName UNiagaraDataInterfaceGrid3DCollection::GetVector4AttributeIndexFunctionName("GetVector4AttributeIndex");
 const FName UNiagaraDataInterfaceGrid3DCollection::GetVectorAttributeIndexFunctionName("GetVectorAttributeIndex");
 const FName UNiagaraDataInterfaceGrid3DCollection::GetVector2DAttributeIndexFunctionName("GetVector2DAttributeIndex");
@@ -157,6 +162,7 @@ struct FNiagaraGridCollection3DDIFunctionVersion
 	enum Type
 	{
 		InitialVersion = 0,
+		AddGetOutputGridValue,
 
 		VersionPlusOne,
 		LatestVersion = VersionPlusOne - 1
@@ -174,7 +180,8 @@ bool UNiagaraDataInterfaceGrid3DCollection::CanCreateVarFromFuncName(const FName
 		FuncName == UNiagaraDataInterfaceGrid3DCollection::SetVector4ValueAtIndexFunctionName ||
 		FuncName == UNiagaraDataInterfaceGrid3DCollection::GetPreviousVector4ValueAtIndexFunctionName ||
 		FuncName == UNiagaraDataInterfaceGrid3DCollection::SamplePreviousGridVector4AtIndexFunctionName ||
-		FuncName == UNiagaraDataInterfaceGrid3DCollection::CubicSamplePreviousGridVector4AtIndexFunctionName)
+		FuncName == UNiagaraDataInterfaceGrid3DCollection::CubicSamplePreviousGridVector4AtIndexFunctionName ||
+		FuncName == UNiagaraDataInterfaceGrid3DCollection::GetOutputGridVector4ValueFunctionName)
 
 		return true;
 	else if (
@@ -186,7 +193,8 @@ bool UNiagaraDataInterfaceGrid3DCollection::CanCreateVarFromFuncName(const FName
 		FuncName == UNiagaraDataInterfaceGrid3DCollection::SetVector3ValueAtIndexFunctionName ||
 		FuncName == UNiagaraDataInterfaceGrid3DCollection::GetPreviousVector3ValueAtIndexFunctionName ||
 		FuncName == UNiagaraDataInterfaceGrid3DCollection::SamplePreviousGridVector3AtIndexFunctionName ||
-		FuncName == UNiagaraDataInterfaceGrid3DCollection::CubicSamplePreviousGridVector3AtIndexFunctionName)
+		FuncName == UNiagaraDataInterfaceGrid3DCollection::CubicSamplePreviousGridVector3AtIndexFunctionName ||
+		FuncName == UNiagaraDataInterfaceGrid3DCollection::GetOutputGridVectorValueFunctionName)
 		return true;
 	else if (
 		FuncName == UNiagaraDataInterfaceGrid3DCollection::SetVector2ValueFunctionName ||
@@ -197,7 +205,8 @@ bool UNiagaraDataInterfaceGrid3DCollection::CanCreateVarFromFuncName(const FName
 		FuncName == UNiagaraDataInterfaceGrid3DCollection::SetVector2ValueAtIndexFunctionName ||
 		FuncName == UNiagaraDataInterfaceGrid3DCollection::GetPreviousVector2ValueAtIndexFunctionName ||
 		FuncName == UNiagaraDataInterfaceGrid3DCollection::SamplePreviousGridVector2AtIndexFunctionName ||
-		FuncName == UNiagaraDataInterfaceGrid3DCollection::CubicSamplePreviousGridVector2AtIndexFunctionName)
+		FuncName == UNiagaraDataInterfaceGrid3DCollection::CubicSamplePreviousGridVector2AtIndexFunctionName ||
+		FuncName == UNiagaraDataInterfaceGrid3DCollection::GetOutputGridVector2ValueFunctionName)
 		return true;
 	else if (
 		FuncName == UNiagaraDataInterfaceGrid3DCollection::SetFloatValueFunctionName ||
@@ -207,7 +216,8 @@ bool UNiagaraDataInterfaceGrid3DCollection::CanCreateVarFromFuncName(const FName
 		FuncName == UNiagaraDataInterfaceGrid3DCollection::SetFloatValueAtIndexFunctionName ||
 		FuncName == UNiagaraDataInterfaceGrid3DCollection::GetPreviousFloatValueAtIndexFunctionName ||
 		FuncName == UNiagaraDataInterfaceGrid3DCollection::SamplePreviousGridFloatAtIndexFunctionName ||
-		FuncName == UNiagaraDataInterfaceGrid3DCollection::CubicSamplePreviousGridFloatAtIndexFunctionName)
+		FuncName == UNiagaraDataInterfaceGrid3DCollection::CubicSamplePreviousGridFloatAtIndexFunctionName ||
+		FuncName == UNiagaraDataInterfaceGrid3DCollection::GetOutputGridFloatValueFunctionName)
 		return true;
 	return false;
 }
@@ -222,7 +232,8 @@ FNiagaraTypeDefinition UNiagaraDataInterfaceGrid3DCollection::GetValueTypeFromFu
 		FuncName == UNiagaraDataInterfaceGrid3DCollection::SetVector4ValueAtIndexFunctionName ||
 		FuncName == UNiagaraDataInterfaceGrid3DCollection::GetPreviousVector4ValueAtIndexFunctionName ||
 		FuncName == UNiagaraDataInterfaceGrid3DCollection::SamplePreviousGridVector4AtIndexFunctionName ||
-		FuncName == UNiagaraDataInterfaceGrid3DCollection::CubicSamplePreviousGridVector4AtIndexFunctionName)
+		FuncName == UNiagaraDataInterfaceGrid3DCollection::CubicSamplePreviousGridVector4AtIndexFunctionName ||
+		FuncName == UNiagaraDataInterfaceGrid3DCollection::GetOutputGridVector4ValueFunctionName)
 		return FNiagaraTypeDefinition::GetVec4Def();
 	else if (
 		FuncName == UNiagaraDataInterfaceGrid3DCollection::SetVector3ValueFunctionName ||
@@ -233,7 +244,8 @@ FNiagaraTypeDefinition UNiagaraDataInterfaceGrid3DCollection::GetValueTypeFromFu
 		FuncName == UNiagaraDataInterfaceGrid3DCollection::SetVector3ValueAtIndexFunctionName ||
 		FuncName == UNiagaraDataInterfaceGrid3DCollection::GetPreviousVector3ValueAtIndexFunctionName ||
 		FuncName == UNiagaraDataInterfaceGrid3DCollection::SamplePreviousGridVector3AtIndexFunctionName ||
-		FuncName == UNiagaraDataInterfaceGrid3DCollection::CubicSamplePreviousGridVector3AtIndexFunctionName)
+		FuncName == UNiagaraDataInterfaceGrid3DCollection::CubicSamplePreviousGridVector3AtIndexFunctionName ||
+		FuncName == UNiagaraDataInterfaceGrid3DCollection::GetOutputGridVectorValueFunctionName)
 		return FNiagaraTypeDefinition::GetVec3Def();
 	else if (
 		FuncName == UNiagaraDataInterfaceGrid3DCollection::SetVector2ValueFunctionName ||
@@ -244,7 +256,8 @@ FNiagaraTypeDefinition UNiagaraDataInterfaceGrid3DCollection::GetValueTypeFromFu
 		FuncName == UNiagaraDataInterfaceGrid3DCollection::SetVector2ValueAtIndexFunctionName ||
 		FuncName == UNiagaraDataInterfaceGrid3DCollection::GetPreviousVector2ValueAtIndexFunctionName ||
 		FuncName == UNiagaraDataInterfaceGrid3DCollection::SamplePreviousGridVector2AtIndexFunctionName ||
-		FuncName == UNiagaraDataInterfaceGrid3DCollection::CubicSamplePreviousGridVector2AtIndexFunctionName)
+		FuncName == UNiagaraDataInterfaceGrid3DCollection::CubicSamplePreviousGridVector2AtIndexFunctionName ||
+		FuncName == UNiagaraDataInterfaceGrid3DCollection::GetOutputGridVector2ValueFunctionName)
 		return FNiagaraTypeDefinition::GetVec2Def();
 	else if (
 		FuncName == UNiagaraDataInterfaceGrid3DCollection::SetFloatValueFunctionName ||
@@ -254,7 +267,8 @@ FNiagaraTypeDefinition UNiagaraDataInterfaceGrid3DCollection::GetValueTypeFromFu
 		FuncName == UNiagaraDataInterfaceGrid3DCollection::SetFloatValueAtIndexFunctionName ||
 		FuncName == UNiagaraDataInterfaceGrid3DCollection::GetPreviousFloatValueAtIndexFunctionName ||
 		FuncName == UNiagaraDataInterfaceGrid3DCollection::SamplePreviousGridFloatAtIndexFunctionName ||
-		FuncName == UNiagaraDataInterfaceGrid3DCollection::CubicSamplePreviousGridFloatAtIndexFunctionName)
+		FuncName == UNiagaraDataInterfaceGrid3DCollection::CubicSamplePreviousGridFloatAtIndexFunctionName ||
+		FuncName == UNiagaraDataInterfaceGrid3DCollection::GetOutputGridFloatValueFunctionName)
 		return FNiagaraTypeDefinition::GetFloatDef();
 
 	return FNiagaraTypeDefinition();
@@ -270,8 +284,9 @@ int32 UNiagaraDataInterfaceGrid3DCollection::GetComponentCountFromFuncName(const
 		FuncName == UNiagaraDataInterfaceGrid3DCollection::SetVector4ValueAtIndexFunctionName ||
 		FuncName == UNiagaraDataInterfaceGrid3DCollection::GetPreviousVector4ValueAtIndexFunctionName ||
 		FuncName == UNiagaraDataInterfaceGrid3DCollection::SamplePreviousGridVector4AtIndexFunctionName ||
-		FuncName == UNiagaraDataInterfaceGrid3DCollection::CubicSamplePreviousGridVector4AtIndexFunctionName || 
-		FuncName == UNiagaraDataInterfaceGrid3DCollection::GetVector4AttributeIndexFunctionName)
+		FuncName == UNiagaraDataInterfaceGrid3DCollection::CubicSamplePreviousGridVector4AtIndexFunctionName ||
+		FuncName == UNiagaraDataInterfaceGrid3DCollection::GetVector4AttributeIndexFunctionName ||
+		FuncName == UNiagaraDataInterfaceGrid3DCollection::GetOutputGridVector4ValueFunctionName)
 		return 4;
 	else if (
 		FuncName == UNiagaraDataInterfaceGrid3DCollection::SetVector3ValueFunctionName ||
@@ -283,7 +298,8 @@ int32 UNiagaraDataInterfaceGrid3DCollection::GetComponentCountFromFuncName(const
 		FuncName == UNiagaraDataInterfaceGrid3DCollection::GetPreviousVector3ValueAtIndexFunctionName ||
 		FuncName == UNiagaraDataInterfaceGrid3DCollection::SamplePreviousGridVector3AtIndexFunctionName ||
 		FuncName == UNiagaraDataInterfaceGrid3DCollection::CubicSamplePreviousGridVector3AtIndexFunctionName ||
-		FuncName == UNiagaraDataInterfaceGrid3DCollection::GetVectorAttributeIndexFunctionName)
+		FuncName == UNiagaraDataInterfaceGrid3DCollection::GetVectorAttributeIndexFunctionName ||
+		FuncName == UNiagaraDataInterfaceGrid3DCollection::GetOutputGridVectorValueFunctionName)
 		return 3;
 	else if (
 		FuncName == UNiagaraDataInterfaceGrid3DCollection::SetVector2ValueFunctionName ||
@@ -295,7 +311,8 @@ int32 UNiagaraDataInterfaceGrid3DCollection::GetComponentCountFromFuncName(const
 		FuncName == UNiagaraDataInterfaceGrid3DCollection::GetPreviousVector2ValueAtIndexFunctionName ||
 		FuncName == UNiagaraDataInterfaceGrid3DCollection::SamplePreviousGridVector2AtIndexFunctionName ||
 		FuncName == UNiagaraDataInterfaceGrid3DCollection::CubicSamplePreviousGridVector2AtIndexFunctionName ||
-		FuncName == UNiagaraDataInterfaceGrid3DCollection::GetVector2DAttributeIndexFunctionName)
+		FuncName == UNiagaraDataInterfaceGrid3DCollection::GetVector2DAttributeIndexFunctionName ||
+		FuncName == UNiagaraDataInterfaceGrid3DCollection::GetOutputGridVector2ValueFunctionName)
 		return 2;
 	else if (
 		FuncName == UNiagaraDataInterfaceGrid3DCollection::SetFloatValueFunctionName ||
@@ -306,7 +323,8 @@ int32 UNiagaraDataInterfaceGrid3DCollection::GetComponentCountFromFuncName(const
 		FuncName == UNiagaraDataInterfaceGrid3DCollection::GetPreviousFloatValueAtIndexFunctionName ||
 		FuncName == UNiagaraDataInterfaceGrid3DCollection::SamplePreviousGridFloatAtIndexFunctionName ||
 		FuncName == UNiagaraDataInterfaceGrid3DCollection::CubicSamplePreviousGridFloatAtIndexFunctionName ||
-		FuncName == UNiagaraDataInterfaceGrid3DCollection::GetFloatAttributeIndexFunctionName)
+		FuncName == UNiagaraDataInterfaceGrid3DCollection::GetFloatAttributeIndexFunctionName ||
+		FuncName == UNiagaraDataInterfaceGrid3DCollection::GetOutputGridFloatValueFunctionName)
 		return 1;
 
 	return INDEX_NONE;
@@ -918,6 +936,88 @@ void UNiagaraDataInterfaceGrid3DCollection::GetFunctionsInternal(TArray<FNiagara
 
 	}
 
+	// GetOutputGrid* functions - read from OutputGrid UAV (destination buffer) for in-place iterative solvers (e.g. RBGS-SOR).
+	// NOTE: bWriteFunction is set to true even though these functions only READ from the OutputGrid. This is intentional:
+	// bWriteFunction adds the DI to OutputDestinations, which makes IsOutputStage()=true so PreStage/PostStage fire.
+	// Crucially, bReadFunction is NOT set, so the DI is NOT added to InputDataInterfaces, meaning IsInputStage()=false.
+	// This causes BeginSimulate to receive RequiresBuffering=false, which sets DestinationData=CurrentData (single buffer).
+	// If we used bReadFunction instead, it would force double-buffering, defeating the purpose of these functions.
+	{
+		FNiagaraFunctionSignature Sig;
+		Sig.Name = GetOutputGridFloatValueFunctionName;
+		Sig.Inputs.Add(FNiagaraVariable(FNiagaraTypeDefinition(GetClass()), TEXT("Grid")));
+		Sig.Inputs.Add(FNiagaraVariable(FNiagaraTypeDefinition::GetIntDef(), TEXT("IndexX")));
+		Sig.Inputs.Add(FNiagaraVariable(FNiagaraTypeDefinition::GetIntDef(), TEXT("IndexY")));
+		Sig.Inputs.Add(FNiagaraVariable(FNiagaraTypeDefinition::GetIntDef(), TEXT("IndexZ")));
+		Sig.Outputs.Add(FNiagaraVariable(FNiagaraTypeDefinition::GetFloatDef(), TEXT("Value")));
+		Sig.FunctionSpecifiers.Add(FName(UNiagaraDataInterfaceRWBase::NAME_Attribute));
+		Sig.bMemberFunction = true;
+		Sig.bRequiresContext = false;
+		Sig.ModuleUsageBitmask = ENiagaraScriptUsageMask::Particle;
+		Sig.bSupportsCPU = false;
+		Sig.bSupportsGPU = true;
+		Sig.bWriteFunction = true;
+		Sig.Description = NSLOCTEXT("Niagara", "NiagaraDataInterfaceGridColl3D_GetOutputGridFloat", "Gets a float value from the current OutputGrid (destination buffer) by Attribute name. Use for in-place iterative solvers like Red-Black Gauss-Seidel.");
+		OutFunctions.Add(Sig);
+	}
+
+	{
+		FNiagaraFunctionSignature Sig;
+		Sig.Name = GetOutputGridVector2ValueFunctionName;
+		Sig.Inputs.Add(FNiagaraVariable(FNiagaraTypeDefinition(GetClass()), TEXT("Grid")));
+		Sig.Inputs.Add(FNiagaraVariable(FNiagaraTypeDefinition::GetIntDef(), TEXT("IndexX")));
+		Sig.Inputs.Add(FNiagaraVariable(FNiagaraTypeDefinition::GetIntDef(), TEXT("IndexY")));
+		Sig.Inputs.Add(FNiagaraVariable(FNiagaraTypeDefinition::GetIntDef(), TEXT("IndexZ")));
+		Sig.Outputs.Add(FNiagaraVariable(FNiagaraTypeDefinition::GetVec2Def(), TEXT("Value")));
+		Sig.FunctionSpecifiers.Add(FName(UNiagaraDataInterfaceRWBase::NAME_Attribute));
+		Sig.bMemberFunction = true;
+		Sig.bRequiresContext = false;
+		Sig.ModuleUsageBitmask = ENiagaraScriptUsageMask::Particle;
+		Sig.bSupportsCPU = false;
+		Sig.bSupportsGPU = true;
+		Sig.bWriteFunction = true;
+		Sig.Description = NSLOCTEXT("Niagara", "NiagaraDataInterfaceGridColl3D_GetOutputGridVector2", "Gets a Vector2 value from the current OutputGrid (destination buffer) by Attribute name. Use for in-place iterative solvers like Red-Black Gauss-Seidel.");
+		OutFunctions.Add(Sig);
+	}
+
+	{
+		FNiagaraFunctionSignature Sig;
+		Sig.Name = GetOutputGridVectorValueFunctionName;
+		Sig.Inputs.Add(FNiagaraVariable(FNiagaraTypeDefinition(GetClass()), TEXT("Grid")));
+		Sig.Inputs.Add(FNiagaraVariable(FNiagaraTypeDefinition::GetIntDef(), TEXT("IndexX")));
+		Sig.Inputs.Add(FNiagaraVariable(FNiagaraTypeDefinition::GetIntDef(), TEXT("IndexY")));
+		Sig.Inputs.Add(FNiagaraVariable(FNiagaraTypeDefinition::GetIntDef(), TEXT("IndexZ")));
+		Sig.Outputs.Add(FNiagaraVariable(FNiagaraTypeDefinition::GetVec3Def(), TEXT("Value")));
+		Sig.FunctionSpecifiers.Add(FName(UNiagaraDataInterfaceRWBase::NAME_Attribute));
+		Sig.bMemberFunction = true;
+		Sig.bRequiresContext = false;
+		Sig.ModuleUsageBitmask = ENiagaraScriptUsageMask::Particle;
+		Sig.bSupportsCPU = false;
+		Sig.bSupportsGPU = true;
+		Sig.bWriteFunction = true;
+		Sig.Description = NSLOCTEXT("Niagara", "NiagaraDataInterfaceGridColl3D_GetOutputGridVector", "Gets a Vector3 value from the current OutputGrid (destination buffer) by Attribute name. Use for in-place iterative solvers like Red-Black Gauss-Seidel.");
+		OutFunctions.Add(Sig);
+	}
+
+	{
+		FNiagaraFunctionSignature Sig;
+		Sig.Name = GetOutputGridVector4ValueFunctionName;
+		Sig.Inputs.Add(FNiagaraVariable(FNiagaraTypeDefinition(GetClass()), TEXT("Grid")));
+		Sig.Inputs.Add(FNiagaraVariable(FNiagaraTypeDefinition::GetIntDef(), TEXT("IndexX")));
+		Sig.Inputs.Add(FNiagaraVariable(FNiagaraTypeDefinition::GetIntDef(), TEXT("IndexY")));
+		Sig.Inputs.Add(FNiagaraVariable(FNiagaraTypeDefinition::GetIntDef(), TEXT("IndexZ")));
+		Sig.Outputs.Add(FNiagaraVariable(FNiagaraTypeDefinition::GetVec4Def(), TEXT("Value")));
+		Sig.FunctionSpecifiers.Add(FName(UNiagaraDataInterfaceRWBase::NAME_Attribute));
+		Sig.bMemberFunction = true;
+		Sig.bRequiresContext = false;
+		Sig.ModuleUsageBitmask = ENiagaraScriptUsageMask::Particle;
+		Sig.bSupportsCPU = false;
+		Sig.bSupportsGPU = true;
+		Sig.bWriteFunction = true;
+		Sig.Description = NSLOCTEXT("Niagara", "NiagaraDataInterfaceGridColl3D_GetOutputGridVector4", "Gets a Vector4 value from the current OutputGrid (destination buffer) by Attribute name. Use for in-place iterative solvers like Red-Black Gauss-Seidel.");
+		OutFunctions.Add(Sig);
+	}
+
 	{
 		FNiagaraFunctionSignature Sig;
 		Sig.Name = SampleGridFunctionName;
@@ -1763,6 +1863,51 @@ void UNiagaraDataInterfaceGrid3DCollection::WriteGetHLSL(const FNiagaraDataInter
 	OutHLSL += FString::Format(*FormatBounds, ArgsBounds);
 }
 
+void UNiagaraDataInterfaceGrid3DCollection::WriteGetOutputGridHLSL(const FNiagaraDataInterfaceGPUParamInfo& ParamInfo, const FNiagaraDataInterfaceGeneratedFunction& FunctionInfo, int FunctionInstanceIndex, int32 InNumChannels, FString& OutHLSL)
+{
+	// Like WriteGetHLSL but reads from OutputGrid UAV instead of Grid SRV
+	FString FormatBounds = TEXT(R"(
+			void {FunctionName}(int In_IndexX, int In_IndexY, int In_IndexZ, out float{NumChannelsVariableSuffix} Out_Val)
+			{
+				int In_AttributeIndex = {AttributeIndicesName}[{AttributeIndexGroup}]{AttributeIndexComponent};
+
+			    for (int i = 0; i < {NumChannels}; i++)
+				{
+					int CurAttributeIndex = In_AttributeIndex + i;
+					int3 TileOffset = {PerAttributeDataName}[CurAttributeIndex].xyz;
+					float Val = {OutputGrid}[int3(In_IndexX + TileOffset.x, In_IndexY + TileOffset.y, In_IndexZ + TileOffset.z)];
+					)");
+	if (InNumChannels == 1)
+	{
+		FormatBounds += TEXT("					Out_Val = Val;\n");
+	}
+	else
+	{
+		FormatBounds += TEXT("					Out_Val[i] = Val;\n");
+	}
+	FormatBounds += TEXT(R"(
+				}
+			}
+		)");
+	TMap<FString, FStringFormatArg> ArgsBounds = {
+		{TEXT("ParameterName"),				ParamInfo.DataInterfaceHLSLSymbol},
+		{TEXT("FunctionName"),				FunctionInfo.InstanceName},
+		{TEXT("OutputGrid"),				ParamInfo.DataInterfaceHLSLSymbol + OutputGridName},
+		{TEXT("Grid"),						ParamInfo.DataInterfaceHLSLSymbol + GridName},
+		{TEXT("NumCellsName"),				ParamInfo.DataInterfaceHLSLSymbol + UNiagaraDataInterfaceRWBase::NumCellsName},
+		{TEXT("UnitToUVName"),				ParamInfo.DataInterfaceHLSLSymbol + UNiagaraDataInterfaceRWBase::UnitToUVName},
+		{TEXT("PerAttributeDataName"),		ParamInfo.DataInterfaceHLSLSymbol + TEXT("_") + PerAttributeDataName},
+		{TEXT("AttributeIndicesName"),		ParamInfo.DataInterfaceHLSLSymbol + TEXT("_") + AttributeIndicesBaseName},
+		{TEXT("AttributeIndexGroup"),		FunctionInstanceIndex / 4},
+		{TEXT("AttributeIndexComponent"),	VectorComponentNames[FunctionInstanceIndex % 4]},
+		{TEXT("NumChannels"),				FString::FromInt(InNumChannels)},
+		{TEXT("NumChannelsVariableSuffix"),	InNumChannels > 1 ? FString::FromInt(InNumChannels) : TEXT("")},
+		{TEXT("NumTiles"),					ParamInfo.DataInterfaceHLSLSymbol + NumTilesName},
+		{TEXT("NumAttributesName"),			ParamInfo.DataInterfaceHLSLSymbol + UNiagaraDataInterfaceRWBase::NumAttributesName},
+	};
+	OutHLSL += FString::Format(*FormatBounds, ArgsBounds);
+}
+
 void UNiagaraDataInterfaceGrid3DCollection::WriteSampleHLSL(const FNiagaraDataInterfaceGPUParamInfo& ParamInfo, const FNiagaraDataInterfaceGeneratedFunction& FunctionInfo, int FunctionInstanceIndex, int32 InNumChannels, FString SampleFunction, FString& OutHLSL)
 {
 	SampleFunction.ReplaceInline(TEXT("{Grid}"), *(ParamInfo.DataInterfaceHLSLSymbol + GridName));
@@ -2371,6 +2516,54 @@ bool UNiagaraDataInterfaceGrid3DCollection::GetFunctionHLSL(const FNiagaraDataIn
 
 		return true;
 	}
+	else if (FunctionInfo.DefinitionName == GetOutputGridVector4ValueFunctionName)
+	{
+		if (AttributeWriter.UseRGBAGrid())
+		{
+			AttributeWriter.WriteGetOutputGridHLSL(EGridAttributeRetrievalMode::RGBAGrid, FunctionInfo, OutHLSL);
+		}
+		else
+		{
+			WriteGetOutputGridHLSL(ParamInfo, FunctionInfo, FunctionInstanceIndex, 4, OutHLSL);
+		}
+		return true;
+	}
+	else if (FunctionInfo.DefinitionName == GetOutputGridVectorValueFunctionName)
+	{
+		if (AttributeWriter.UseRGBAGrid())
+		{
+			AttributeWriter.WriteGetOutputGridHLSL(EGridAttributeRetrievalMode::RGBAGrid, FunctionInfo, OutHLSL);
+		}
+		else
+		{
+			WriteGetOutputGridHLSL(ParamInfo, FunctionInfo, FunctionInstanceIndex, 3, OutHLSL);
+		}
+		return true;
+	}
+	else if (FunctionInfo.DefinitionName == GetOutputGridVector2ValueFunctionName)
+	{
+		if (AttributeWriter.UseRGBAGrid())
+		{
+			AttributeWriter.WriteGetOutputGridHLSL(EGridAttributeRetrievalMode::RGBAGrid, FunctionInfo, OutHLSL);
+		}
+		else
+		{
+			WriteGetOutputGridHLSL(ParamInfo, FunctionInfo, FunctionInstanceIndex, 2, OutHLSL);
+		}
+		return true;
+	}
+	else if (FunctionInfo.DefinitionName == GetOutputGridFloatValueFunctionName)
+	{
+		if (AttributeWriter.UseRGBAGrid())
+		{
+			AttributeWriter.WriteGetOutputGridHLSL(EGridAttributeRetrievalMode::RGBAGrid, FunctionInfo, OutHLSL);
+		}
+		else
+		{
+			WriteGetOutputGridHLSL(ParamInfo, FunctionInfo, FunctionInstanceIndex, 1, OutHLSL);
+		}
+		return true;
+	}
 	else if (FunctionInfo.DefinitionName == SamplePreviousGridFloatFunctionName)
 	{
 		if (AttributeWriter.UseRGBAGrid())
@@ -3976,6 +4169,7 @@ void FNiagaraDataInterfaceProxyGrid3DCollectionProxy::PreStage(const FNDIGpuComp
 	{
 		FGrid3DCollectionRWInstanceData_RenderThread* ProxyData = SystemInstancesToProxyData_RT.Find(Context.GetSystemInstanceID());
 		FRDGBuilder& GraphBuilder = Context.GetGraphBuilder();
+
 		ProxyData->BeginSimulate(GraphBuilder, Context.IsInputStage());
 
 		// If there is an output DI and the stage dispatch count is different from the grid resolution, then we want to consider it for a pre-stage UAV clear.
@@ -3983,11 +4177,26 @@ void FNiagaraDataInterfaceProxyGrid3DCollectionProxy::PreStage(const FNDIGpuComp
 		const FNiagaraSimStageData& SimStageData = Context.GetSimStageData();
 		const FIntVector3& ElementCount = SimStageData.DispatchArgs.ElementCount;
 
-		if (ProxyData->ClearBeforeNonIterationStage && 
+		if (ProxyData->ClearBeforeNonIterationStage &&
 			ProxyData->DestinationData &&
 			(ElementCount.X != ProxyData->NumCells.X || ElementCount.Y != ProxyData->NumCells.Y || ElementCount.Z != ProxyData->NumCells.Z))
 		{
-			AddClearUAVPass(GraphBuilder, ProxyData->DestinationData->GetOrCreateUAV(GraphBuilder), FVector4f(ForceInitToZero));
+			// In single-buffer mode (DestinationData == CurrentData, i.e. when only GetOutputGrid* + Set* are used),
+			// clearing the destination would destroy the only copy of the data. Skip the clear and warn instead.
+			if (ProxyData->DestinationData == ProxyData->CurrentData)
+			{
+				static bool bWarnedSingleBufferClear = false;
+				if (!bWarnedSingleBufferClear)
+				{
+					bWarnedSingleBufferClear = true;
+					UE_LOG(LogNiagara, Warning, TEXT("Grid3DCollection: ClearBeforeNonIterationStage skipped in single-buffer mode. "
+						"Clearing the destination buffer would destroy in-place data. Disable ClearBeforeNonIterationStage on the DI to suppress this warning."));
+				}
+			}
+			else
+			{
+				AddClearUAVPass(GraphBuilder, ProxyData->DestinationData->GetOrCreateUAV(GraphBuilder), FVector4f(ForceInitToZero));
+			}
 		}
 	}
 }
diff --git a/Engine/Plugins/FX/Niagara/Source/Niagara/Private/NiagaraDataInterfaceGrid3DCollectionUtils.cpp b/Engine/Plugins/FX/Niagara/Source/Niagara/Private/NiagaraDataInterfaceGrid3DCollectionUtils.cpp
--- a/Engine/Plugins/FX/Niagara/Source/Niagara/Private/NiagaraDataInterfaceGrid3DCollectionUtils.cpp
+++ b/Engine/Plugins/FX/Niagara/Source/Niagara/Private/NiagaraDataInterfaceGrid3DCollectionUtils.cpp
@@ -118,6 +118,65 @@ bool FGrid3DCollectionAttributeHlslWriter::WriteGetHLSL(EGridAttributeRetrievalM
 	return true;
 }
 
+bool FGrid3DCollectionAttributeHlslWriter::WriteGetOutputGridHLSL(EGridAttributeRetrievalMode AttributeStorage, const FNiagaraDataInterfaceGeneratedFunction& FunctionInfo, FString& OutHLSL)
+{
+	const FName* AttributeName = FunctionInfo.FindSpecifierValue(UNiagaraDataInterfaceRWBase::NAME_Attribute);
+	if (AttributeName == nullptr)
+	{
+		return false;
+	}
+
+	const FNiagaraDataInterfaceRWAttributeHelper::FAttributeInfo* AttributeInfo = AttributeHelper.FindAttributeInfo(*AttributeName);
+	if (AttributeInfo == nullptr)
+	{
+		return false;
+	}
+
+	// Read from OutputGrid UAV instead of Grid SRV
+	const FString OutputGridNameHLSL = ParamInfo.DataInterfaceHLSLSymbol + UNiagaraDataInterfaceGrid3DCollection::OutputGridName;
+
+	if (AttributeStorage == EGridAttributeRetrievalMode::RGBAGrid)
+	{
+		FString NumChannelsString;
+		FString AttrGridChannels;
+
+		AttributeHelper.GetChannelStrings(AttributeInfo, NumChannelsString, AttrGridChannels);
+
+		OutHLSL.Appendf(TEXT("void %s(int IndexX, int IndexY, int IndexZ, out float%s Value)\n"), *FunctionInfo.InstanceName, *NumChannelsString);
+		OutHLSL.Appendf(TEXT("{\n"));
+		OutHLSL.Appendf(TEXT("	Value = %s[int3(IndexX, IndexY, IndexZ)].%s;\n"), *OutputGridNameHLSL, *AttrGridChannels);
+		OutHLSL.Appendf(TEXT("}\n"));
+	}
+	else
+	{
+		if (AttributeInfo->NumChannels == 1)
+		{
+			OutHLSL.Appendf(TEXT("void %s(int IndexX, int IndexY, int IndexZ, out float Value)\n"), *FunctionInfo.InstanceName);
+			OutHLSL.Appendf(TEXT("{\n"));
+			OutHLSL.Appendf(TEXT("	int AttributeIndex = %s;\n"), *AttributeHelper.GetAttributeIndex(AttributeInfo));
+			OutHLSL.Appendf(TEXT("	int3 PixelOffset = int3(IndexX, IndexY, IndexZ) + %s;\n"), *GetPerAttributePixelOffset());
+			OutHLSL.Appendf(TEXT("	Value = %s[PixelOffset];\n"), *OutputGridNameHLSL);
+			OutHLSL.Appendf(TEXT("}\n"));
+		}
+		else
+		{
+			OutHLSL.Appendf(TEXT("void %s(int IndexX, int IndexY, int IndexZ, out float%d Value)\n"), *FunctionInfo.InstanceName, AttributeInfo->NumChannels);
+			OutHLSL.Appendf(TEXT("{\n"));
+			for (int32 i = 0; i < AttributeInfo->NumChannels; ++i)
+			{
+				OutHLSL.Appendf(TEXT("	{\n"));
+				OutHLSL.Appendf(TEXT("		int AttributeIndex = %s;\n"), *AttributeHelper.GetAttributeIndex(AttributeInfo, i));
+				OutHLSL.Appendf(TEXT("		int3 PixelOffset = int3(IndexX, IndexY, IndexZ) + %s;\n"), *GetPerAttributePixelOffset());
+				OutHLSL.Appendf(TEXT("		Value[%d] = %s[PixelOffset];\n"), i, *OutputGridNameHLSL);
+				OutHLSL.Appendf(TEXT("	}\n"));
+			}
+			OutHLSL.Appendf(TEXT("}\n"));
+		}
+	}
+
+	return true;
+}
+
 bool FGrid3DCollectionAttributeHlslWriter::WriteGetAtIndexHLSL(EGridAttributeRetrievalMode AttributeStorage, const FNiagaraDataInterfaceGeneratedFunction& FunctionInfo, int NumChannels, FString& OutHLSL)
 {
 	FString NumChannelsString;
diff --git a/Engine/Plugins/FX/Niagara/Source/Niagara/Private/NiagaraDataInterfaceGrid3DCollectionUtils.h b/Engine/Plugins/FX/Niagara/Source/Niagara/Private/NiagaraDataInterfaceGrid3DCollectionUtils.h
--- a/Engine/Plugins/FX/Niagara/Source/Niagara/Private/NiagaraDataInterfaceGrid3DCollectionUtils.h
+++ b/Engine/Plugins/FX/Niagara/Source/Niagara/Private/NiagaraDataInterfaceGrid3DCollectionUtils.h
@@ -39,6 +39,7 @@ struct FGrid3DCollectionAttributeHlslWriter
 	FString GetPerAttributeUVWOffset() const;
 
 	bool WriteGetHLSL(EGridAttributeRetrievalMode AttributeStorage, const FNiagaraDataInterfaceGeneratedFunction& FunctionInfo, FString& OutHLSL);
+	bool WriteGetOutputGridHLSL(EGridAttributeRetrievalMode AttributeStorage, const FNiagaraDataInterfaceGeneratedFunction& FunctionInfo, FString& OutHLSL);
 	bool WriteGetAtIndexHLSL(EGridAttributeRetrievalMode AttributeStorage, const FNiagaraDataInterfaceGeneratedFunction& FunctionInfo, int NumChannels, FString& OutHLSL);
 	bool WriteSetAtIndexHLSL(EGridAttributeRetrievalMode AttributeStorage, const FNiagaraDataInterfaceGeneratedFunction& FunctionInfo, int NumChannels, FString& OutHLSL);
 	bool WriteSampleAtIndexHLSL(EGridAttributeRetrievalMode AttributeStorage, const FNiagaraDataInterfaceGeneratedFunction& FunctionInfo, int NumChannels, bool IsCubic, FString& OutHLSL);
