Add type hashes to assertion to show possible hash collision

#Synced-CL 52678434

diff --git a/Engine/Plugins/FX/Niagara/Source/Niagara/Private/NiagaraTypeRegistry.cpp b/Engine/Plugins/FX/Niagara/Source/Niagara/Private/NiagaraTypeRegistry.cpp
--- a/Engine/Plugins/FX/Niagara/Source/Niagara/Private/NiagaraTypeRegistry.cpp
+++ b/Engine/Plugins/FX/Niagara/Source/Niagara/Private/NiagaraTypeRegistry.cpp
@@ -329,8 +329,12 @@ int32 FNiagaraTypeRegistry::RegisterTypeInternal(const FNiagaraTypeDefinition& N
 					const FNiagaraTypeDefinition& OldType = *RegisteredTypes[NewType.RegisteredTypeDefIndex]; 
 					bNeedsRegistration = !OldType.IsValid();
 					checkf(bNeedsRegistration || (NewType == OldType),
-						TEXT("bNeedsRegistration: %d, NewType: %s, OldType: %s"), bNeedsRegistration, *NewType.GetNameText().ToString(),
-						*OldType.GetNameText().ToString() 
+						TEXT("bNeedsRegistration: %d, NewType: %s, OldType: %s NewTypeHash: %x OldTypeHash: %x"), 
+						bNeedsRegistration, 
+						*NewType.GetNameText().ToString(),
+						*OldType.GetNameText().ToString(),
+						TypeHash,
+						GetTypeHash(OldType)
 					);
 				}
 			}
