- Fix potential null pointer component access

#Synced-CL 52684591

diff --git a/Engine/Plugins/FX/Niagara/Source/Niagara/Private/NiagaraDataChannel_GameplayBurst.cpp b/Engine/Plugins/FX/Niagara/Source/Niagara/Private/NiagaraDataChannel_GameplayBurst.cpp
--- a/Engine/Plugins/FX/Niagara/Source/Niagara/Private/NiagaraDataChannel_GameplayBurst.cpp
+++ b/Engine/Plugins/FX/Niagara/Source/Niagara/Private/NiagaraDataChannel_GameplayBurst.cpp
@@ -235,8 +235,10 @@ void UNiagaraDataChannelHandler_GameplayBurst::GenerateKey(FNDCAccessContextInst
 		if (Comp->ComponentHasTag(AttachedComponentTag))
 		{
 			check(Comp->IsA<UNiagaraComponent>());
-			bUseAttachedPath = true;
 			Comp = Comp->GetAttachParent();
+
+			// Comp can potentially be nullptr now as we moved to the attach parent, update bUseAttachedPath accordingly to avoid a null deref
+			bUseAttachedPath = Comp != nullptr;
 		}
 		else
 		{
