fix nullptr crash

#Synced-CL 52589973

diff --git a/Engine/Plugins/FX/Niagara/Source/NiagaraEditor/Private/ViewModels/Stack/NiagaraStackFunctionInput.cpp b/Engine/Plugins/FX/Niagara/Source/NiagaraEditor/Private/ViewModels/Stack/NiagaraStackFunctionInput.cpp
--- a/Engine/Plugins/FX/Niagara/Source/NiagaraEditor/Private/ViewModels/Stack/NiagaraStackFunctionInput.cpp
+++ b/Engine/Plugins/FX/Niagara/Source/NiagaraEditor/Private/ViewModels/Stack/NiagaraStackFunctionInput.cpp
@@ -244,6 +244,11 @@ const UNiagaraNodeFunctionCall& UNiagaraStackFunctionInput::GetInputFunctionCall
 	return *OwningFunctionCallNode.Get();
 }
 
+UNiagaraNodeFunctionCall* UNiagaraStackFunctionInput::GetInputFunctionCallNodePtr() const
+{
+	return OwningFunctionCallNode.Get();
+}
+
 UNiagaraScript* UNiagaraStackFunctionInput::GetInputFunctionCallInitialScript() const
 {
 	return OwningFunctionCallInitialScript.Get();
diff --git a/Engine/Plugins/FX/Niagara/Source/NiagaraEditor/Private/ViewModels/Stack/NiagaraStackInputCategory.cpp b/Engine/Plugins/FX/Niagara/Source/NiagaraEditor/Private/ViewModels/Stack/NiagaraStackInputCategory.cpp
--- a/Engine/Plugins/FX/Niagara/Source/NiagaraEditor/Private/ViewModels/Stack/NiagaraStackInputCategory.cpp
+++ b/Engine/Plugins/FX/Niagara/Source/NiagaraEditor/Private/ViewModels/Stack/NiagaraStackInputCategory.cpp
@@ -383,12 +383,14 @@ void UNiagaraStackSummaryCategory::RefreshChildrenInternal(const TArray<UNiagara
 				
 				if(ModuleBaseInputData.IsSet())
 				{
-					UNiagaraStackFunctionInput* TopLevelInput = FindCurrentChildOfTypeByPredicate<UNiagaraStackFunctionInput>(CurrentChildren, [&](UNiagaraStackFunctionInput* CurrentInput) 
+					UNiagaraStackFunctionInput* TopLevelInput = FindCurrentChildOfTypeByPredicate<UNiagaraStackFunctionInput>(CurrentChildren, [&](UNiagaraStackFunctionInput* CurrentInput)
 					{
-						return CurrentInput->GetInputParameterHandle() == ModuleBaseInputData->InputName && CurrentInput->GetInputType() == ModuleBaseInputData->Type
+						UNiagaraNodeFunctionCall* FunctionCallNode = CurrentInput->GetInputFunctionCallNodePtr();
+						return FunctionCallNode != nullptr
+							&& CurrentInput->GetInputParameterHandle() == ModuleBaseInputData->InputName && CurrentInput->GetInputType() == ModuleBaseInputData->Type
 							&& CurrentInput->GetInputFunctionCallInitialScript() == OwningFunctionCallNode->FunctionScript
 							// we are checking for node guid as we could have 2x the same input from 2 identical modules
-							&& CurrentInput->GetInputFunctionCallNode().NodeGuid == OwningFunctionCallNode->NodeGuid;
+							&& FunctionCallNode->NodeGuid == OwningFunctionCallNode->NodeGuid;
 					});
 	
 					if(TopLevelInput == nullptr)
@@ -420,12 +422,14 @@ void UNiagaraStackSummaryCategory::RefreshChildrenInternal(const TArray<UNiagara
 						if(ChildInputData.IsSet())
 						{
 							UNiagaraNodeFunctionCall* ChildFunctionNode = ChildInputData->FunctionCallNode;
-							UNiagaraStackFunctionInput* StackChildInput = FindCurrentChildOfTypeByPredicate<UNiagaraStackFunctionInput>(CurrentChildren, [&](UNiagaraStackFunctionInput* CurrentInput) 
+							UNiagaraStackFunctionInput* StackChildInput = FindCurrentChildOfTypeByPredicate<UNiagaraStackFunctionInput>(CurrentChildren, [&](UNiagaraStackFunctionInput* CurrentInput)
 							{
-								return CurrentInput->GetInputParameterHandle() == ChildInputData->InputName && CurrentInput->GetInputType() == ChildInputData->Type
+								UNiagaraNodeFunctionCall* FunctionCallNode = CurrentInput->GetInputFunctionCallNodePtr();
+								return FunctionCallNode != nullptr
+									&& CurrentInput->GetInputParameterHandle() == ChildInputData->InputName && CurrentInput->GetInputType() == ChildInputData->Type
 									&& CurrentInput->GetInputFunctionCallInitialScript() == ChildFunctionNode->FunctionScript
 									// we are checking for node guid as we could have 2x the same input from 2 identical modules
-									&& CurrentInput->GetInputFunctionCallNode().NodeGuid == ChildFunctionNode->NodeGuid;
+									&& FunctionCallNode->NodeGuid == ChildFunctionNode->NodeGuid;
 							});
 							
 							if(StackChildInput == nullptr)
diff --git a/Engine/Plugins/FX/Niagara/Source/NiagaraEditor/Public/ViewModels/Stack/NiagaraStackFunctionInput.h b/Engine/Plugins/FX/Niagara/Source/NiagaraEditor/Public/ViewModels/Stack/NiagaraStackFunctionInput.h
--- a/Engine/Plugins/FX/Niagara/Source/NiagaraEditor/Public/ViewModels/Stack/NiagaraStackFunctionInput.h
+++ b/Engine/Plugins/FX/Niagara/Source/NiagaraEditor/Public/ViewModels/Stack/NiagaraStackFunctionInput.h
@@ -117,6 +117,9 @@ public:
 	/** Gets the function call node which owns this input. */
 	NIAGARAEDITOR_API const UNiagaraNodeFunctionCall& GetInputFunctionCallNode() const;
 
+	/** Gets the function call node which owns this input, or nullptr if the node has been garbage collected. */
+	NIAGARAEDITOR_API UNiagaraNodeFunctionCall* GetInputFunctionCallNodePtr() const;
+
 	/** Gets the script that the function call node was referencing when this input was initialized. */
 	NIAGARAEDITOR_API UNiagaraScript* GetInputFunctionCallInitialScript() const;
 
